[MLDeformer] A complete redesign, turning this into a whole framework that can handle different models.

diff --git a/Engine/Plugins/Experimental/MLDeformer/MLDeformerFramework/MLDeformerFramework.uplugin b/Engine/Plugins/Experimental/MLDeformer/MLDeformerFramework/MLDeformerFramework.uplugin
new file mode 100644
--- /dev/null
+++ b/Engine/Plugins/Experimental/MLDeformer/MLDeformerFramework/MLDeformerFramework.uplugin
@@ -0,0 +1,43 @@
+{
+	"FileVersion": 3,
+	"Version": 1,
+	"VersionName": "0.1",
+	"FriendlyName": "MLDeformerFramework",
+	"Description": "Machine Learning Mesh Deformer Framework",
+	"Category": "Animation",
+	"CreatedBy": "Epic Games, Inc.",
+	"CreatedByURL": "http://epicgames.com",
+	"DocsURL": "",
+	"MarketplaceURL": "",
+	"SupportURL": "",
+	"CanContainContent": true,
+	"IsBetaVersion": false,
+	"IsExperimentalVersion": true,
+	"Installed": false,
+	"Modules": [
+		{
+			"Name": "MLDeformerFramework",
+			"Type": "Runtime",
+			"LoadingPhase": "Default"
+		},
+		{
+			"Name": "MLDeformerFrameworkEditor",
+			"Type": "Editor",
+			"LoadingPhase": "Default"
+		}
+	],
+	"Plugins": [
+		{
+			"Name": "Optimus",
+			"Enabled": true
+		},
+		{
+			"Name": "PythonFoundationPackages",
+			"Enabled": true
+		},
+		{
+			"Name": "NeuralNetworkInference",
+			"Enabled": true
+		}
+	]
+}
\ No newline at end of file
diff --git a/Engine/Plugins/Experimental/MLDeformer/MLDeformerFramework/Shaders/Private/MLDeformerGraphDataInterface.ush b/Engine/Plugins/Experimental/MLDeformer/MLDeformerFramework/Shaders/Private/MLDeformerGraphDataInterface.ush
new file mode 100644
--- /dev/null
+++ b/Engine/Plugins/Experimental/MLDeformer/MLDeformerFramework/Shaders/Private/MLDeformerGraphDataInterface.ush
@@ -0,0 +1,22 @@
+// Copyright Epic Games, Inc. All Rights Reserved.
+
+#include "/Plugin/ComputeFramework/Private/ComputeKernelCommon.ush"
+
+uint DI_LOCAL(NumVertices);
+uint DI_LOCAL(InputStreamStart);
+float DI_LOCAL(Weight);
+
+Buffer<float> DI_LOCAL(PositionDeltaBuffer);
+Buffer<uint> DI_LOCAL(VertexMapBuffer);
+
+DI_IMPL_READ(ReadNumVertices, uint, void)
+{
+	return DI_LOCAL(NumVertices);
+}
+
+DI_IMPL_READ(ReadPositionDelta, float3, uint VertexIndex)
+{
+	uint BufferIndex = DI_LOCAL(VertexMapBuffer)[VertexIndex + DI_LOCAL(InputStreamStart)] * 3;
+	float3 Delta = float3(DI_LOCAL(PositionDeltaBuffer)[BufferIndex], DI_LOCAL(PositionDeltaBuffer)[BufferIndex + 1], DI_LOCAL(PositionDeltaBuffer)[BufferIndex + 2]);
+	return Delta * DI_LOCAL(Weight);
+}
diff --git a/Engine/Plugins/Experimental/MLDeformer/MLDeformerFramework/Shaders/Private/MLDeformerGraphHeatMapDataInterface.ush b/Engine/Plugins/Experimental/MLDeformer/MLDeformerFramework/Shaders/Private/MLDeformerGraphHeatMapDataInterface.ush
new file mode 100644
--- /dev/null
+++ b/Engine/Plugins/Experimental/MLDeformer/MLDeformerFramework/Shaders/Private/MLDeformerGraphHeatMapDataInterface.ush
@@ -0,0 +1,43 @@
+// Copyright Epic Games, Inc. All Rights Reserved.
+
+#include "/Plugin/ComputeFramework/Private/ComputeKernelCommon.ush"
+
+uint DI_LOCAL(NumVertices);
+uint DI_LOCAL(InputStreamStart);
+uint DI_LOCAL(HeatMapMode);
+float DI_LOCAL(HeatMapMax);
+float DI_LOCAL(GroundTruthLerp);
+uint DI_LOCAL(GroundTruthBufferSize);
+
+StructuredBuffer<float3> DI_LOCAL(PositionGroundTruthBuffer);
+Buffer<uint> DI_LOCAL(VertexMapBuffer);
+
+DI_IMPL_READ(ReadNumVertices, uint, void)
+{
+	return DI_LOCAL(NumVertices);
+}
+
+DI_IMPL_READ(ReadPositionGroundTruth, float3, uint VertexIndex)
+{
+	uint BufferIndex = DI_LOCAL(VertexMapBuffer)[VertexIndex + DI_LOCAL(InputStreamStart)];
+	if (BufferIndex >= DI_LOCAL(GroundTruthBufferSize))
+	{
+		return float3(0, 0, 0);
+	}
+	return DI_LOCAL(PositionGroundTruthBuffer)[BufferIndex];
+}
+
+DI_IMPL_READ(ReadHeatMapMode, uint, void)
+{
+	return DI_LOCAL(HeatMapMode);
+}
+
+DI_IMPL_READ(ReadHeatMapMax, float, void)
+{
+	return DI_LOCAL(HeatMapMax);
+}
+
+DI_IMPL_READ(ReadGroundTruthLerp, float, void)
+{
+	return DI_LOCAL(GroundTruthLerp);
+}
diff --git a/Engine/Plugins/Experimental/MLDeformer/MLDeformerFramework/Source/MLDeformerFramework/Private/MLDeformerComponent.cpp b/Engine/Plugins/Experimental/MLDeformer/MLDeformerFramework/Source/MLDeformerFramework/Private/MLDeformerComponent.cpp
new file mode 100644
--- /dev/null
+++ b/Engine/Plugins/Experimental/MLDeformer/MLDeformerFramework/Source/MLDeformerFramework/Private/MLDeformerComponent.cpp
@@ -0,0 +1,127 @@
+// Copyright Epic Games, Inc. All Rights Reserved.
+
+#include "MLDeformerComponent.h"
+#include "MLDeformerAsset.h"
+#include "MLDeformerModelInstance.h"
+#include "MLDeformerModel.h"
+#include "Components/SkeletalMeshComponent.h"
+
+UMLDeformerComponent::UMLDeformerComponent(const FObjectInitializer& ObjectInitializer)
+	: Super(ObjectInitializer)
+{
+	bTickInEditor = true;
+	bAutoActivate = true;
+	PrimaryComponentTick.TickGroup = TG_PrePhysics;
+	PrimaryComponentTick.bCanEverTick = true;
+}
+
+void UMLDeformerComponent::Init()
+{
+	using namespace UE::MLDeformer;
+
+	// If there is no deformer asset linked, release what we currently have.
+	if (DeformerAsset == nullptr)
+	{
+		if (ModelInstance)
+		{
+			ModelInstance->Release();
+		}
+		return;
+	}
+
+	// Try to initialize the deformer model.
+	UMLDeformerModel* Model = DeformerAsset->GetModel();
+	if (Model)
+	{
+		// Create the model instance if it isn't there yet.
+		if (ModelInstance == nullptr)
+		{
+			ModelInstance = TUniquePtr<FMLDeformerModelInstance>(Model->CreateModelInstance());
+		}
+
+		ModelInstance->Release();
+		ModelInstance->Init(SkelMeshComponent);
+	}
+	else
+	{
+		UE_LOG(LogMLDeformer, Warning, TEXT("ML Deformer component on '%s' has a deformer asset that has no ML model setup."), *GetOuter()->GetName());
+	}
+}
+
+void UMLDeformerComponent::SetupComponent(UMLDeformerAsset* InDeformerAsset, USkeletalMeshComponent* InSkelMeshComponent)
+{
+	AddTickPrerequisiteComponent(SkelMeshComponent);
+
+	DeformerAsset = InDeformerAsset;
+	SkelMeshComponent = InSkelMeshComponent;
+
+	// Initialize and make sure we have a model instance.
+	RemoveNeuralNetworkModifyDelegate();
+	Init();
+	AddNeuralNetworkModifyDelegate();
+}
+
+void UMLDeformerComponent::AddNeuralNetworkModifyDelegate()
+{
+	if (DeformerAsset == nullptr)
+	{
+		return;
+	}
+
+	UMLDeformerModel* Model = DeformerAsset->GetModel();		
+	NeuralNetworkModifyDelegateHandle = DeformerAsset->NeuralNetworkModifyDelegate.AddLambda
+	(
+		([this]()
+		{
+			ModelInstance->Release();
+			ModelInstance->Init(SkelMeshComponent);
+		})
+	);
+}
+
+void UMLDeformerComponent::RemoveNeuralNetworkModifyDelegate()
+{
+	if (DeformerAsset != nullptr && NeuralNetworkModifyDelegateHandle != FDelegateHandle())
+	{
+		DeformerAsset->NeuralNetworkModifyDelegate.Remove(NeuralNetworkModifyDelegateHandle);
+	}
+	
+	NeuralNetworkModifyDelegateHandle = FDelegateHandle();
+}
+
+void UMLDeformerComponent::Activate(bool bReset)
+{
+	// If we haven't pointed to some skeletal mesh component to use, then try to find one on the actor.
+	USkeletalMeshComponent* SkelMeshComponentToUse = SkelMeshComponent;
+	if (SkelMeshComponentToUse == nullptr)
+	{
+		AActor* Actor = Cast<AActor>(GetOuter());
+		SkelMeshComponent = Actor->FindComponentByClass<USkeletalMeshComponent>();
+	}
+
+	SetupComponent(DeformerAsset, SkelMeshComponent);
+}
+
+void UMLDeformerComponent::Deactivate()
+{
+	RemoveNeuralNetworkModifyDelegate();
+	if (ModelInstance)
+	{
+		ModelInstance->Release();
+	}
+}
+
+void UMLDeformerComponent::TickComponent(float DeltaTime, enum ELevelTick TickType, FActorComponentTickFunction* ThisTickFunction)
+{
+	if (TickType != ELevelTick::LEVELTICK_PauseTick)
+	{
+		if (ModelInstance &&
+			SkelMeshComponent && 
+			SkelMeshComponent->GetPredictedLODLevel() == 0)
+		{
+			ModelInstance->Tick(DeltaTime);
+		}
+	}
+
+	Super::TickComponent(DeltaTime, TickType, ThisTickFunction);
+}
diff --git a/Engine/Plugins/Experimental/MLDeformer/MLDeformerFramework/Source/MLDeformerFramework/Private/MLDeformerGeomCacheHelpers.cpp b/Engine/Plugins/Experimental/MLDeformer/MLDeformerFramework/Source/MLDeformerFramework/Private/MLDeformerGeomCacheHelpers.cpp
new file mode 100644
--- /dev/null
+++ b/Engine/Plugins/Experimental/MLDeformer/MLDeformerFramework/Source/MLDeformerFramework/Private/MLDeformerGeomCacheHelpers.cpp
@@ -0,0 +1,348 @@
+// Copyright Epic Games, Inc. All Rights Reserved.
+
+#include "MLDeformerGeomCacheHelpers.h"
+#include "MLDeformerModel.h"
+
+#include "GeometryCache.h"
+#include "GeometryCacheMeshData.h"
+#include "GeometryCacheTrack.h"
+
+#include "Animation/AnimSequence.h"
+#include "Rendering/SkeletalMeshModel.h"
+#include "Rendering/SkeletalMeshLODModel.h"
+
+#include "NeuralNetwork.h"
+
+#define LOCTEXT_NAMESPACE "MLDeformerGeomCacheHelpers"
+
+namespace UE::MLDeformer
+{
+#if WITH_EDITORONLY_DATA
+	FText GetGeomCacheErrorText(USkeletalMesh* InSkeletalMesh, UGeometryCache* InGeomCache)
+	{
+		FText Result;
+		if (InGeomCache)
+		{
+			FString ErrorString;
+
+			// Verify that we have imported vertex numbers enabled.
+			TArray<FGeometryCacheMeshData> MeshData;
+			InGeomCache->GetMeshDataAtTime(0.0f, MeshData);
+			if (MeshData.Num() == 0)
+			{
+				ErrorString = FText(LOCTEXT("TargetMeshNoMeshData", "No geometry data is present.")).ToString();
+			}
+			else
+			{
+				if (MeshData[0].ImportedVertexNumbers.Num() == 0)
+				{
+					ErrorString = FText(LOCTEXT("TargetMeshNoImportedVertexNumbers", "Please import Geometry Cache with option 'Store Imported Vertex Numbers' enabled!")).ToString();
+				}
+			}
+
+			// Check if we flattened the tracks.
+			if (InGeomCache->Tracks.Num() == 1 && InGeomCache->Tracks[0]->GetName() == TEXT("Flattened_Track"))
+			{
+				int32 NumSkelMeshes = 0;
+				if (InSkeletalMesh)
+				{
+					FSkeletalMeshModel* ImportedModel = InSkeletalMesh->GetImportedModel();
+					if (ImportedModel)
+					{
+						NumSkelMeshes = ImportedModel->LODModels[0].ImportedMeshInfos.Num();		
+					}
+				}
+
+				if (NumSkelMeshes > 1)
+				{
+					if (!ErrorString.IsEmpty())
+					{
+						ErrorString += TEXT("\n\n");
+					}
+					ErrorString += FText(LOCTEXT("TargetMeshFlattened", "Please import Geometry Cache with option 'Flatten Tracks' disabled!")).ToString();
+				}
+			}
+
+			Result = FText::FromString(ErrorString);
+		}
+
+		return Result;
+	}
+
+	FText GetGeomCacheVertexErrorText(USkeletalMesh* InSkelMesh, UGeometryCache* InGeomCache, const FText& SkelName, const FText& GeomCacheName)
+	{
+		FText Result;
+
+		if (InSkelMesh && InGeomCache)
+		{
+			const int32 SkelVertCount = UMLDeformerModel::ExtractNumImportedSkinnedVertices(InSkelMesh);
+			const int32 GeomCacheVertCount = ExtractNumImportedGeomCacheVertices(InGeomCache);
+			const bool bHasGeomCacheError = !GetGeomCacheErrorText(InSkelMesh, InGeomCache).IsEmpty();
+			if (SkelVertCount != GeomCacheVertCount && !bHasGeomCacheError)
+			{
+				Result = FText::Format(
+					LOCTEXT("MeshVertexNumVertsMismatch", "Vertex count of {0} doesn't match with {1}!\n\n{2} has {3} verts, while {4} has {5} verts."),
+					SkelName,
+					GeomCacheName,
+					SkelName,
+					SkelVertCount,
+					GeomCacheName,
+					GeomCacheVertCount);
+			}
+		}
+
+		return Result;
+	}
+
+	FText GetGeomCacheMeshMappingErrorText(USkeletalMesh* InSkelMesh, UGeometryCache* InGeomCache)
+	{
+		FText Result;
+
+		if (InGeomCache && InSkelMesh)
+		{
+			// Check for failed mesh mappings.
+			TArray<FMLDeformerGeomCacheMeshMapping> MeshMappings;
+			TArray<FString> FailedNames;
+			GenerateGeomCacheMeshMappings(InSkelMesh, InGeomCache, MeshMappings, FailedNames);
+
+			// List all mesh names that have issues.
+			FString ErrorString;
+			for (int32 Index = 0; Index < FailedNames.Num(); ++Index)
+			{
+				ErrorString += FailedNames[Index];
+				if (Index < FailedNames.Num() - 1)
+				{
+					ErrorString += TEXT("\n");
+				}
+			}
+
+			Result = FText::FromString(ErrorString);
+		}
+
+		return Result;
+	}
+
+	// A fuzzy name match.
+	// There is a match when the track name starts with the mesh name.
+	bool IsPotentialMatch(const FString& TrackName, const FString& MeshName)
+	{
+		return (TrackName.Find(MeshName) == 0);
+	}
+
+	void GenerateGeomCacheMeshMappings(USkeletalMesh* SkelMesh, UGeometryCache* GeomCache, TArray<FMLDeformerGeomCacheMeshMapping>& OutMeshMappings, TArray<FString>& OutFailedImportedMeshNames)
+	{
+		OutMeshMappings.Empty();
+		OutFailedImportedMeshNames.Empty();
+		if (SkelMesh == nullptr || GeomCache == nullptr)
+		{
+			return;
+		}
+
+		// If we haven't got any imported mesh infos then the asset needs to be reimported first.
+		// We show an error for this in the editor UI already.
+		FSkeletalMeshModel* ImportedModel = SkelMesh->GetImportedModel();
+		check(ImportedModel);
+		const TArray<FSkelMeshImportedMeshInfo>& SkelMeshInfos = ImportedModel->LODModels[0].ImportedMeshInfos;
+		if (SkelMeshInfos.IsEmpty())
+		{
+			return;
+		}
+
+		// For all meshes in the skeletal mesh.
+		const float SampleTime = 0.0f;
+		FString SkelMeshName;
+		for (int32 SkelMeshIndex = 0; SkelMeshIndex < SkelMeshInfos.Num(); ++SkelMeshIndex)
+		{
+			const FSkelMeshImportedMeshInfo& MeshInfo = SkelMeshInfos[SkelMeshIndex];
+			SkelMeshName = MeshInfo.Name.ToString();
+
+			// Find the matching one in the geom cache.
+			bool bFoundMatch = false;
+			for (int32 TrackIndex = 0; TrackIndex < GeomCache->Tracks.Num(); ++TrackIndex)
+			{
+				// Check if this is a candidate based on the mesh and track name.
+				UGeometryCacheTrack* Track = GeomCache->Tracks[TrackIndex];
+				const bool bIsSoloMesh = (GeomCache->Tracks.Num() == 1 && SkelMeshInfos.Num() == 1);	// Do we just have one mesh and one track?
+				if (Track && 
+					(IsPotentialMatch(Track->GetName(), SkelMeshName) || bIsSoloMesh))
+				{	
+					// Extract the geom cache mesh data.
+					FGeometryCacheMeshData GeomCacheMeshData;
+					if (!Track->GetMeshDataAtTime(SampleTime, GeomCacheMeshData))
+					{
+						continue;
+					}
+
+					// Verify that we have imported vertex numbers.
+					if (GeomCacheMeshData.ImportedVertexNumbers.IsEmpty())
+					{
+						continue;
+					}
+
+					// Get the number of geometry cache mesh imported verts.
+					int32 NumGeomMeshVerts = 0;
+					for (int32 GeomVertIndex = 0; GeomVertIndex < GeomCacheMeshData.ImportedVertexNumbers.Num(); ++GeomVertIndex)
+					{
+						NumGeomMeshVerts = FMath::Max(NumGeomMeshVerts, (int32)GeomCacheMeshData.ImportedVertexNumbers[GeomVertIndex]);
+					}
+					NumGeomMeshVerts += 1;	// +1 Because we use indices, so a cube's max index is 7, while there are 8 vertices.
+
+					// Make sure the vertex counts match.
+					const int32 NumSkelMeshVerts = MeshInfo.NumVertices;
+					if (NumSkelMeshVerts != NumGeomMeshVerts)
+					{
+						continue;
+					}
+
+					// Create a new mesh mapping entry.
+					OutMeshMappings.AddDefaulted();
+					UE::MLDeformer::FMLDeformerGeomCacheMeshMapping& Mapping = OutMeshMappings.Last();
+					Mapping.MeshIndex = SkelMeshIndex;
+					Mapping.TrackIndex = TrackIndex;
+					Mapping.SkelMeshToTrackVertexMap.AddUninitialized(NumSkelMeshVerts);
+					Mapping.ImportedVertexToRenderVertexMap.AddUninitialized(NumSkelMeshVerts);
+
+					// For all vertices (both skel mesh and geom cache mesh have the same number of verts here).
+					for (int32 VertexIndex = 0; VertexIndex < NumSkelMeshVerts; ++VertexIndex)
+					{
+						// Find the first vertex with the same dcc vertex in the geom cache mesh.
+						// When there are multiple vertices with the same vertex number here, they are duplicates with different normals or uvs etc.
+						// However they all share the same vertex position, so we can just find the first hit, as we only need the position later on.
+						const int32 GeomCacheVertexIndex = GeomCacheMeshData.ImportedVertexNumbers.Find(VertexIndex);
+						Mapping.SkelMeshToTrackVertexMap[VertexIndex] = GeomCacheVertexIndex;
+					
+						// Map the source asset vertex number to a render vertex. This is the first duplicate of that vertex.
+						const int32 RenderVertexIndex = ImportedModel->LODModels[0].MeshToImportVertexMap.Find(MeshInfo.StartImportedVertex + VertexIndex);
+						Mapping.ImportedVertexToRenderVertexMap[VertexIndex] = RenderVertexIndex;
+					}
+
+					// We found a match, no need to iterate over more Tracks.
+					bFoundMatch = true;
+					break;
+				} // If the track name matches the skeletal meshes internal mesh name.
+			} // For all tracks.
+
+			if (!bFoundMatch)
+			{
+				OutFailedImportedMeshNames.Add(SkelMeshName);
+				UE_LOG(LogMLDeformer, Warning, TEXT("Imported mesh '%s' cannot be matched with a geometry cache track."), *SkelMeshName);
+			}
+		} // For all meshes inside the skeletal mesh.
+	}
+
+	void SampleGeomCachePositions(
+		int32 InLODIndex,
+		float InSampleTime,
+		const TArray<UE::MLDeformer::FMLDeformerGeomCacheMeshMapping>& InMeshMappings,
+		const USkeletalMesh* SkelMesh,
+		const UGeometryCache* InGeometryCache,
+		const FTransform& AlignmentTransform,
+		TArray<FVector3f>& OutPositions)
+	{
+		if (InGeometryCache == nullptr)
+		{
+			return;
+		}
+
+		if (!ensure(SkelMesh != nullptr))
+		{
+			return;
+		}
+
+		const FSkeletalMeshModel* ImportedModel = SkelMesh->GetImportedModel();
+		if (!ensure(ImportedModel != nullptr))
+		{
+			return;
+		}
+	
+		const FSkeletalMeshLODModel& LODModel = ImportedModel->LODModels[InLODIndex];
+		const TArray<FSkelMeshImportedMeshInfo>& SkelMeshInfos = LODModel.ImportedMeshInfos;
+
+		const uint32 NumVertices = LODModel.MaxImportVertex + 1;
+		OutPositions.Reset(NumVertices);
+		OutPositions.AddZeroed(NumVertices);
+
+		// For all mesh mappings we found.
+		for (int32 MeshMappingIndex = 0; MeshMappingIndex < InMeshMappings.Num(); ++MeshMappingIndex)
+		{
+			const UE::MLDeformer::FMLDeformerGeomCacheMeshMapping& MeshMapping = InMeshMappings[MeshMappingIndex];
+			const FSkelMeshImportedMeshInfo& MeshInfo = SkelMeshInfos[MeshMapping.MeshIndex];
+			UGeometryCacheTrack* Track = InGeometryCache->Tracks[MeshMapping.TrackIndex];
+
+			FGeometryCacheMeshData GeomCacheMeshData;
+			if (!Track->GetMeshDataAtTime(InSampleTime, GeomCacheMeshData))
+			{
+				continue;
+			}
+
+			for (int32 VertexIndex = 0; VertexIndex < MeshInfo.NumVertices; ++VertexIndex)
+			{
+				const int32 SkinnedVertexIndex = MeshInfo.StartImportedVertex + VertexIndex;
+				const int32 GeomCacheVertexIndex = MeshMapping.SkelMeshToTrackVertexMap[VertexIndex];
+				if (GeomCacheVertexIndex != INDEX_NONE && GeomCacheMeshData.Positions.IsValidIndex(GeomCacheVertexIndex))
+				{
+					const FVector3f GeomCacheVertexPos = (FVector3f)AlignmentTransform.TransformPosition((FVector)GeomCacheMeshData.Positions[GeomCacheVertexIndex]);
+					OutPositions[SkinnedVertexIndex] = GeomCacheVertexPos;
+				}
+			}
+		}
+	}
+
+	int32 ExtractNumImportedGeomCacheVertices(UGeometryCache* GeometryCache)
+	{
+		if (GeometryCache == nullptr)
+		{
+			return 0;
+		}
+
+		int32 NumGeomCacheImportedVerts = 0;
+
+		// Extract the geom cache number of imported vertices.
+		TArray<FGeometryCacheMeshData> MeshDatas;
+		GeometryCache->GetMeshDataAtTime(0.0f, MeshDatas);
+		for (const FGeometryCacheMeshData& MeshData : MeshDatas)
+		{
+			const TArray<uint32>& ImportedVertexNumbers = MeshData.ImportedVertexNumbers;
+			if (ImportedVertexNumbers.Num() > 0)
+			{
+				// Find the maximum value.
+				int32 MaxIndex = -1;
+				for (int32 Index = 0; Index < ImportedVertexNumbers.Num(); ++Index)
+				{
+					MaxIndex = FMath::Max(static_cast<int32>(ImportedVertexNumbers[Index]), MaxIndex);
+				}
+				check(MaxIndex > -1);
+
+				NumGeomCacheImportedVerts += MaxIndex + 1;
+			}
+		}
+
+		return NumGeomCacheImportedVerts;
+	}
+
+	FText GetGeomCacheAnimSequenceErrorText(UGeometryCache* InGeomCache, UAnimSequence* InAnimSequence)
+	{
+		FText Result;
+
+		if (InAnimSequence && InGeomCache)
+		{
+			const float AnimSeqDuration = InAnimSequence->GetPlayLength();
+			const float GeomCacheDuration = InGeomCache->CalculateDuration();
+			if (FMath::Abs(AnimSeqDuration - GeomCacheDuration) > 0.001f)
+			{
+				FNumberFormattingOptions Options;
+				Options.SetUseGrouping(false);
+				Options.SetMaximumFractionalDigits(4);
+				Result = FText::Format(
+					LOCTEXT("AnimSeqNumFramesMismatch", "Anim sequence and Geometry Cache durations don't match!\n\nAnimSeq has a duration of {0} seconds, while GeomCache has a duration of {1} seconds.\n\nThis can produce incorrect results."),
+					FText::AsNumber(AnimSeqDuration, &Options),
+					FText::AsNumber(GeomCacheDuration, &Options));
+			}
+		}
+
+		return Result;
+	}
+#endif	// #if WITH_EDITORONLY_DATA
+}	// namespace UE::MLDeformer
+
+#undef LOCTEXT_NAMESPACE
diff --git a/Engine/Plugins/Experimental/MLDeformer/MLDeformerFramework/Source/MLDeformerFramework/Private/MLDeformerGraphDataInterface.cpp b/Engine/Plugins/Experimental/MLDeformer/MLDeformerFramework/Source/MLDeformerFramework/Private/MLDeformerGraphDataInterface.cpp
new file mode 100644
--- /dev/null
+++ b/Engine/Plugins/Experimental/MLDeformer/MLDeformerFramework/Source/MLDeformerFramework/Private/MLDeformerGraphDataInterface.cpp
@@ -0,0 +1,125 @@
+// Copyright Epic Games, Inc. All Rights Reserved.
+
+#include "MLDeformerGraphDataInterface.h"
+#include "MLDeformerModel.h"
+#include "MLDeformerModelInstance.h"
+#include "MLDeformerAsset.h"
+#include "MLDeformerComponent.h"
+#include "Components/SkeletalMeshComponent.h"
+#include "ComputeFramework/ComputeKernelPermutationSet.h"
+#include "ComputeFramework/ShaderParamTypeDefinition.h"
+#include "ComputeFramework/ShaderParameterMetadataBuilder.h"
+#include "NeuralNetwork.h"
+#include "OptimusDataDomain.h"
+#include "RenderGraphBuilder.h"
+#include "RenderGraphUtils.h"
+#include "RenderGraphResources.h"
+#include "SkeletalRenderPublic.h"
+
+TArray<FOptimusCDIPinDefinition> UMLDeformerGraphDataInterface::GetPinDefinitions() const
+{
+	TArray<FOptimusCDIPinDefinition> Defs;
+	Defs.Add({ "PositionDelta", "ReadPositionDelta", Optimus::DomainName::Vertex, "ReadNumVertices" });
+	return Defs;
+}
+
+void UMLDeformerGraphDataInterface::GetSupportedInputs(TArray<FShaderFunctionDefinition>& OutFunctions) const
+{
+	{
+		FShaderFunctionDefinition Fn;
+		Fn.Name = TEXT("ReadNumVertices");
+		Fn.bHasReturnType = true;
+		FShaderParamTypeDefinition ReturnParam = {};
+		ReturnParam.ValueType = FShaderValueType::Get(EShaderFundamentalType::Uint);
+		Fn.ParamTypes.Add(ReturnParam);
+		OutFunctions.Add(Fn);
+	}
+	{
+		FShaderFunctionDefinition Fn;
+		Fn.Name = TEXT("ReadPositionDelta");
+		Fn.bHasReturnType = true;
+		FShaderParamTypeDefinition ReturnParam = {};
+		ReturnParam.ValueType = FShaderValueType::Get(EShaderFundamentalType::Float, 3);
+		Fn.ParamTypes.Add(ReturnParam);
+		FShaderParamTypeDefinition Param0 = {};
+		Param0.ValueType = FShaderValueType::Get(EShaderFundamentalType::Uint);
+		Fn.ParamTypes.Add(Param0);
+		OutFunctions.Add(Fn);
+	}
+}
+
+BEGIN_SHADER_PARAMETER_STRUCT(FMLDeformerGraphDataInterfaceParameters, )
+	MLDEFORMER_SHADER_PARAMETERS()
+END_SHADER_PARAMETER_STRUCT()
+
+MLDEFORMER_GRAPH_IMPLEMENT_BASICS(
+	UMLDeformerGraphDataInterface,
+	UMLDeformerGraphDataProvider,
+	UE::MLDeformer::FMLDeformerGraphDataProviderProxy,
+	FMLDeformerGraphDataInterfaceParameters,
+	TEXT("#include \"/Plugin/MLDeformerFramework/Private/MLDeformerGraphDataInterface.ush\"\n"),
+	TEXT("ML Deformer"))
+
+void UMLDeformerGraphDataInterface::GetSourceTypes(TArray<UClass*>& OutSourceTypes) const
+{
+	OutSourceTypes.Add(USkeletalMeshComponent::StaticClass());
+}
+
+bool UMLDeformerGraphDataProvider::IsValid() const
+{
+	if (SkeletalMeshComponent == nullptr || SkeletalMeshComponent->MeshObject == nullptr)
+	{
+		return false;
+	}
+	
+	UMLDeformerComponent* DeformerComponent = SkeletalMeshComponent->GetOwner()->FindComponentByClass<UMLDeformerComponent>();
+	if (DeformerComponent == nullptr || DeformerComponent->GetDeformerAsset() == nullptr)
+	{
+		return false;
+	}
+
+	return DeformerComponent->GetModelInstance()->IsValidForDataProvider();
+}
+
+namespace UE::MLDeformer
+{
+	FMLDeformerGraphDataProviderProxy::FMLDeformerGraphDataProviderProxy(USkeletalMeshComponent* SkeletalMeshComponent, UMLDeformerComponent* DeformerComponent)
+	{
+		using namespace UE::MLDeformer;
+
+		SkeletalMeshObject = SkeletalMeshComponent->MeshObject;
+
+		const UMLDeformerAsset* DeformerAsset = DeformerComponent->GetDeformerAsset();
+		const UMLDeformerModel* Model = DeformerAsset->GetModel();
+		const FMLDeformerModelInstance* ModelInstance = DeformerComponent->GetModelInstance();
+
+		NeuralNetwork = Model->GetNeuralNetwork();
+		NeuralNetworkInferenceHandle = ModelInstance->GetNeuralNetworkInferenceHandle();
+		bCanRunNeuralNet = ModelInstance->IsCompatible();
+		Weight = DeformerComponent->GetWeight();
+		VertexMapBufferSRV = Model->GetVertexMapBuffer().ShaderResourceViewRHI;
+	}
+
+	void FMLDeformerGraphDataProviderProxy::AllocateResources(FRDGBuilder& GraphBuilder)
+	{
+		if (bCanRunNeuralNet)
+		{
+			check(NeuralNetwork != nullptr);
+			Buffer = GraphBuilder.RegisterExternalBuffer(NeuralNetwork->GetOutputTensorForContext(NeuralNetworkInferenceHandle).GetPooledBuffer());
+		}
+		else
+		{
+			// TODO: use an actual buffer that is of the right size, and filled with zero's.
+			Buffer = GraphBuilder.RegisterExternalBuffer(GWhiteVertexBufferWithRDG->Buffer);
+		}
+
+		BufferSRV = GraphBuilder.CreateSRV(Buffer, PF_R32_FLOAT);
+	}
+
+	void FMLDeformerGraphDataProviderProxy::GatherDispatchData(FDispatchSetup const& InDispatchSetup, FCollectedDispatchData& InOutDispatchData)
+	{
+		MLDEFORMER_GRAPH_DISPATCH_START(FMLDeformerGraphDataInterfaceParameters, InDispatchSetup, InOutDispatchData)
+		MLDEFORMER_GRAPH_DISPATCH_DEFAULT_PARAMETERS()
+		MLDEFORMER_GRAPH_DISPATCH_END()
+	}
+}	// namespace UE::MLDeformer
diff --git a/Engine/Plugins/Experimental/MLDeformer/MLDeformerFramework/Source/MLDeformerFramework/Private/MLDeformerGraphDebugDataInterface.cpp b/Engine/Plugins/Experimental/MLDeformer/MLDeformerFramework/Source/MLDeformerFramework/Private/MLDeformerGraphDebugDataInterface.cpp
new file mode 100644
--- /dev/null
+++ b/Engine/Plugins/Experimental/MLDeformer/MLDeformerFramework/Source/MLDeformerFramework/Private/MLDeformerGraphDebugDataInterface.cpp
@@ -0,0 +1,176 @@
+// Copyright Epic Games, Inc. All Rights Reserved.
+
+#include "MLDeformerGraphDebugDataInterface.h"
+#include "MLDeformerGraphDataInterface.h"
+#include "MLDeformerModelInstance.h"
+#include "MLDeformerAsset.h"
+#include "MLDeformerComponent.h"
+#include "MLDeformerVizSettings.h"
+#include "Animation/AnimSequence.h"
+#include "Components/SkeletalMeshComponent.h"
+#include "ComputeFramework/ComputeKernelPermutationSet.h"
+#include "ComputeFramework/ShaderParamTypeDefinition.h"
+#include "ComputeFramework/ShaderParameterMetadataBuilder.h"
+#include "GeometryCache.h"
+#include "GeometryCacheMeshData.h"
+#include "GeometryCacheTrack.h"
+#include "OptimusDataDomain.h"
+#include "RenderGraphBuilder.h"
+#include "RenderGraphUtils.h"
+#include "RenderGraphResources.h"
+#include "Rendering/SkeletalMeshModel.h"
+#include "SkeletalRenderPublic.h"
+#include "NeuralNetwork.h"
+
+TArray<FOptimusCDIPinDefinition> UMLDeformerGraphDebugDataInterface::GetPinDefinitions() const
+{
+	TArray<FOptimusCDIPinDefinition> Defs;
+	Defs.Add({ "HeatMapMode", "ReadHeatMapMode" });
+	Defs.Add({ "HeatMapMax", "ReadHeatMapMax" });
+	Defs.Add({ "GroundTruthLerp", "ReadGroundTruthLerp" });
+	Defs.Add({ "PositionGroundTruth", "ReadPositionGroundTruth", Optimus::DomainName::Vertex, "ReadNumVertices" });
+	return Defs;
+}
+
+void UMLDeformerGraphDebugDataInterface::GetSupportedInputs(TArray<FShaderFunctionDefinition>& OutFunctions) const
+{
+	{
+		FShaderFunctionDefinition Fn;
+		Fn.Name = TEXT("ReadNumVertices");
+		Fn.bHasReturnType = true;
+		FShaderParamTypeDefinition ReturnParam = {};
+		ReturnParam.ValueType = FShaderValueType::Get(EShaderFundamentalType::Uint);
+		Fn.ParamTypes.Add(ReturnParam);
+		OutFunctions.Add(Fn);
+	}
+	{
+		FShaderFunctionDefinition Fn;
+		Fn.Name = TEXT("ReadHeatMapMode");
+		Fn.bHasReturnType = true;
+		FShaderParamTypeDefinition ReturnParam = {};
+		ReturnParam.ValueType = FShaderValueType::Get(EShaderFundamentalType::Int);
+		Fn.ParamTypes.Add(ReturnParam);
+		OutFunctions.Add(Fn);
+	}
+	{
+		FShaderFunctionDefinition Fn;
+		Fn.Name = TEXT("ReadHeatMapMax");
+		Fn.bHasReturnType = true;
+		FShaderParamTypeDefinition ReturnParam = {};
+		ReturnParam.ValueType = FShaderValueType::Get(EShaderFundamentalType::Float);
+		Fn.ParamTypes.Add(ReturnParam);
+		OutFunctions.Add(Fn);
+	}
+	{
+		FShaderFunctionDefinition Fn;
+		Fn.Name = TEXT("ReadGroundTruthLerp");
+		Fn.bHasReturnType = true;
+		FShaderParamTypeDefinition ReturnParam = {};
+		ReturnParam.ValueType = FShaderValueType::Get(EShaderFundamentalType::Float);
+		Fn.ParamTypes.Add(ReturnParam);
+		OutFunctions.Add(Fn);
+	}
+	{
+		FShaderFunctionDefinition Fn;
+		Fn.Name = TEXT("ReadPositionGroundTruth");
+		Fn.bHasReturnType = true;
+		FShaderParamTypeDefinition ReturnParam = {};
+		ReturnParam.ValueType = FShaderValueType::Get(EShaderFundamentalType::Float, 3);
+		Fn.ParamTypes.Add(ReturnParam);
+		FShaderParamTypeDefinition Param0 = {};
+		Param0.ValueType = FShaderValueType::Get(EShaderFundamentalType::Uint);
+		Fn.ParamTypes.Add(Param0);
+		OutFunctions.Add(Fn);
+	}
+}
+
+BEGIN_SHADER_PARAMETER_STRUCT(FMLDeformerGraphDebugDataInterfaceParameters, )
+	MLDEFORMER_DEBUG_SHADER_PARAMETERS()
+END_SHADER_PARAMETER_STRUCT()
+
+MLDEFORMER_GRAPH_IMPLEMENT_DEBUG_BASICS_WITH_PROXY(
+	UMLDeformerGraphDebugDataInterface,
+	UMLDeformerGraphDebugDataProvider,
+	UE::MLDeformer::FMLDeformerGraphDebugDataProviderProxy,
+	FMLDeformerGraphDebugDataInterfaceParameters,
+	TEXT("#include \"/Plugin/MLDeformerFramework/Private/MLDeformerGraphHeatMapDataInterface.ush\"\n"),
+	TEXT("ML Deformer Debug"))
+
+void UMLDeformerGraphDebugDataInterface::GetSourceTypes(TArray<UClass*>& OutSourceTypes) const
+{
+	OutSourceTypes.Add(USkeletalMeshComponent::StaticClass());
+	OutSourceTypes.Add(UMLDeformerComponent::StaticClass());
+}
+
+bool UMLDeformerGraphDebugDataProvider::IsValid() const
+{
+#if WITH_EDITORONLY_DATA
+	if (SkeletalMeshComponent == nullptr || SkeletalMeshComponent->MeshObject == nullptr)
+	{
+		return false;
+	}
+	
+	UMLDeformerComponent* DeformerComponent = SkeletalMeshComponent->GetOwner()->FindComponentByClass<UMLDeformerComponent>();
+	if (DeformerComponent == nullptr || DeformerComponent->GetDeformerAsset() == nullptr)
+	{
+		return false;
+	}
+
+	return DeformerComponent->GetModelInstance()->IsValidForDataProvider();
+#else
+	return false; // This data interface is only valid in editor.
+#endif
+}
+
+#if WITH_EDITORONLY_DATA
+namespace UE::MLDeformer
+{
+	FMLDeformerGraphDebugDataProviderProxy::FMLDeformerGraphDebugDataProviderProxy(USkeletalMeshComponent* SkeletalMeshComponent, UMLDeformerAsset* DeformerAsset, UMLDeformerGraphDebugDataProvider* InProvider)
+		: FComputeDataProviderRenderProxy()
+	{
+		SkeletalMeshObject = SkeletalMeshComponent->MeshObject;
+		Provider = InProvider;
+
+		UMLDeformerModel* Model = DeformerAsset->GetModel();	
+		UMLDeformerVizSettings* VizSettings = Model->GetVizSettings();
+
+		VertexMapBufferSRV = Model->GetVertexMapBuffer().ShaderResourceViewRHI;
+		HeatMapMode = (int32)VizSettings->GetHeatMapMode();
+		HeatMapMax = 1.0f / FMath::Max(VizSettings->GetHeatMapMax(), 0.00001f);
+		GroundTruthLerp = VizSettings->GetGroundTruthLerp();
+	}
+
+	void FMLDeformerGraphDebugDataProviderProxy::HandleZeroGroundTruthPositions()
+	{
+		if (GroundTruthPositions.Num() == 0)
+		{	
+			// We didn't get valid ground truth vertices.
+			// Make non empty array for later buffer generation.
+			GroundTruthPositions.Add(FVector3f::ZeroVector);
+
+			// Silently disable relevant debug things.
+			if (HeatMapMode == (int32)EMLDeformerHeatMapMode::GroundTruth)
+			{
+				HeatMapMode = -1;
+				HeatMapMax = 0.0f;
+				GroundTruthLerp = 0.0f;
+			}
+		}
+	}
+
+	void FMLDeformerGraphDebugDataProviderProxy::AllocateResources(FRDGBuilder& GraphBuilder)
+	{
+		GroundTruthBuffer = GraphBuilder.CreateBuffer(FRDGBufferDesc::CreateStructuredDesc(sizeof(FVector3f), GroundTruthPositions.Num()), TEXT("MLDeformer.GroundTruthPositions"));
+		GroundTruthBufferSRV = GraphBuilder.CreateSRV(GroundTruthBuffer);
+		GraphBuilder.QueueBufferUpload(GroundTruthBuffer, GroundTruthPositions.GetData(), sizeof(FVector3f) * GroundTruthPositions.Num(), ERDGInitialDataFlags::None);
+	}
+
+	void FMLDeformerGraphDebugDataProviderProxy::GatherDispatchData(FDispatchSetup const& InDispatchSetup, FCollectedDispatchData& InOutDispatchData)
+	{
+		MLDEFORMER_GRAPH_DISPATCH_START(FMLDeformerGraphDebugDataInterfaceParameters, InDispatchSetup, InOutDispatchData)
+		MLDEFORMER_GRAPH_DISPATCH_DEFAULT_DEBUG_PARAMETERS()
+		MLDEFORMER_GRAPH_DISPATCH_END()
+	}
+}	// namespace UE::MLDeformer
+#endif // WITH_EDITORONLY_DATA
+
diff --git a/Engine/Plugins/Experimental/MLDeformer/MLDeformerFramework/Source/MLDeformerFramework/Private/MLDeformerModel.cpp b/Engine/Plugins/Experimental/MLDeformer/MLDeformerFramework/Source/MLDeformerFramework/Private/MLDeformerModel.cpp
new file mode 100644
--- /dev/null
+++ b/Engine/Plugins/Experimental/MLDeformer/MLDeformerFramework/Source/MLDeformerFramework/Private/MLDeformerModel.cpp
@@ -0,0 +1,187 @@
+// Copyright Epic Games, Inc. All Rights Reserved.
+
+#include "MLDeformerModel.h"
+#include "MLDeformerInputInfo.h"
+#include "MLDeformerModelInstance.h"
+#include "MLDeformerVizSettings.h"
+#include "MLDeformerAsset.h"
+#include "Engine/SkeletalMesh.h"
+#include "Rendering/SkeletalMeshModel.h"
+#include "Rendering/SkeletalMeshLODModel.h"
+#include "GeometryCache.h"
+#include "GeometryCacheMeshData.h"
+#include "GeometryCacheTrack.h"
+#include "NeuralNetwork.h"
+
+namespace UE::MLDeformer
+{
+	void FVertexMapBuffer::InitRHI()
+	{
+		if (VertexMap.Num() > 0)
+		{
+			FRHIResourceCreateInfo CreateInfo(TEXT("FVertexMapBuffer"));
+
+			VertexBufferRHI = RHICreateVertexBuffer(VertexMap.Num() * sizeof(uint32), BUF_Static | BUF_ShaderResource, CreateInfo);
+			uint32* Data = reinterpret_cast<uint32*>(RHILockBuffer(VertexBufferRHI, 0, VertexMap.Num() * sizeof(uint32), RLM_WriteOnly));
+			for (int32 Index = 0; Index < VertexMap.Num(); ++Index)
+			{
+				Data[Index] = static_cast<uint32>(VertexMap[Index]);
+			}
+			RHIUnlockBuffer(VertexBufferRHI);
+			VertexMap.Empty();
+
+			ShaderResourceViewRHI = RHICreateShaderResourceView(VertexBufferRHI, 4, PF_R32_UINT);
+		}
+		else
+		{
+			VertexBufferRHI = nullptr;
+			ShaderResourceViewRHI = nullptr;
+		}
+	}
+}	// namespace UE::MLDeformer
+
+UMLDeformerInputInfo* UMLDeformerModel::CreateInputInfo()
+{ 
+	return NewObject<UMLDeformerInputInfo>(this);
+}
+
+UE::MLDeformer::FMLDeformerModelInstance* UMLDeformerModel::CreateModelInstance()
+{ 
+	return new UE::MLDeformer::FMLDeformerModelInstance(this);
+}
+
+void UMLDeformerModel::Init(UMLDeformerAsset* InDeformerAsset) 
+{ 
+	check(InDeformerAsset); 
+	DeformerAsset = InDeformerAsset; 
+	if (InputInfo == nullptr)
+	{
+		InputInfo = CreateInputInfo();
+	}
+}
+
+void UMLDeformerModel::Serialize(FArchive& Archive)
+{
+	#if WITH_EDITOR
+		if (Archive.IsSaving() && Archive.IsPersistent())
+		{
+			InitVertexMap();
+			UpdateCachedNumVertices();
+		}
+	#endif
+
+	Super::Serialize(Archive);
+}
+
+UMLDeformerAsset* UMLDeformerModel::GetDeformerAsset() const
+{ 
+	return DeformerAsset.Get(); 
+}
+
+void UMLDeformerModel::PostLoad()
+{
+	Super::PostLoad();
+
+	InitGPUData();
+
+	#if WITH_EDITOR
+		UpdateCachedNumVertices();
+	#endif
+
+	if (InputInfo)
+	{
+		InputInfo->OnPostLoad();
+	}
+
+	if (NeuralNetwork)
+	{
+		NeuralNetwork->SetDeviceType(ENeuralDeviceType::GPU, ENeuralDeviceType::CPU, ENeuralDeviceType::GPU);
+		if (NeuralNetwork->GetDeviceType() != ENeuralDeviceType::GPU || 
+			NeuralNetwork->GetOutputDeviceType() != ENeuralDeviceType::GPU)
+		{
+			UE_LOG(LogMLDeformer, Error, TEXT("Neural net in ML Deformer '%s' cannot run on the GPU, it will not be active."), *GetName());
+		}
+	}
+}
+
+void UMLDeformerModel::SetNeuralNetwork(UNeuralNetwork* InNeuralNetwork)
+{
+	DeformerAsset->NeuralNetworkModifyDelegate.Broadcast();
+	NeuralNetwork = InNeuralNetwork;
+}
+
+// Used for the FBoenReference, so it knows what skeleton to pick bones from.
+USkeleton* UMLDeformerModel::GetSkeleton(bool& bInvalidSkeletonIsError, const IPropertyHandle* PropertyHandle)
+{
+	#if WITH_EDITOR
+		bInvalidSkeletonIsError = false;
+		if (SkeletalMesh)
+		{
+			return SkeletalMesh->GetSkeleton();
+		}
+	#endif
+	return nullptr;
+}
+
+void UMLDeformerModel::BeginDestroy()
+{
+	BeginReleaseResource(&VertexMapBuffer);
+	RenderResourceDestroyFence.BeginFence();
+	Super::BeginDestroy();
+}
+
+bool UMLDeformerModel::IsReadyForFinishDestroy()
+{
+	return Super::IsReadyForFinishDestroy() && RenderResourceDestroyFence.IsFenceComplete();
+}
+
+void UMLDeformerModel::InitGPUData()
+{
+	BeginReleaseResource(&VertexMapBuffer);
+	VertexMapBuffer.Init(VertexMap);
+	BeginInitResource(&VertexMapBuffer);
+}
+
+#if WITH_EDITOR
+	void UMLDeformerModel::UpdateNumBaseMeshVertices()
+	{
+		NumBaseMeshVerts = UMLDeformerModel::ExtractNumImportedSkinnedVertices(SkeletalMesh);
+	}
+
+	void UMLDeformerModel::PostEditChangeProperty(FPropertyChangedEvent& PropertyChangedEvent)
+	{
+		const FProperty* Property = PropertyChangedEvent.Property;
+		if (Property == nullptr)
+		{
+			return;
+		}
+
+		OnPostEditChangeProperty().ExecuteIfBound(PropertyChangedEvent);
+	}
+
+	void UMLDeformerModel::UpdateCachedNumVertices()
+	{
+		UpdateNumBaseMeshVertices();
+		UpdateNumTargetMeshVertices();
+	}
+
+	int32 UMLDeformerModel::ExtractNumImportedSkinnedVertices(const USkeletalMesh* SkeletalMesh)
+	{
+		return SkeletalMesh ? SkeletalMesh->GetNumImportedVertices() : 0;
+	}
+#endif	// #if WITH_EDITOR
+
+#if WITH_EDITORONLY_DATA
+	void UMLDeformerModel::InitVertexMap()
+	{
+		VertexMap.Empty();
+		if (SkeletalMesh)
+		{
+			FSkeletalMeshModel* SkeletalMeshModel = SkeletalMesh->GetImportedModel();
+			if (SkeletalMeshModel)
+			{
+				VertexMap = SkeletalMeshModel->LODModels[0].MeshToImportVertexMap;
+			}
+		}
+	}
+#endif	// WITH_EDITORDATA_ONLY
diff --git a/Engine/Plugins/Experimental/MLDeformer/MLDeformerFramework/Source/MLDeformerFramework/Private/MLDeformerModelInstance.cpp b/Engine/Plugins/Experimental/MLDeformer/MLDeformerFramework/Source/MLDeformerFramework/Private/MLDeformerModelInstance.cpp
new file mode 100644
--- /dev/null
+++ b/Engine/Plugins/Experimental/MLDeformer/MLDeformerFramework/Source/MLDeformerFramework/Private/MLDeformerModelInstance.cpp
@@ -0,0 +1,269 @@
+// Copyright Epic Games, Inc. All Rights Reserved.
+
+#include "MLDeformerModelInstance.h"
+#include "MLDeformerModel.h"
+#include "MLDeformerModule.h"
+#include "MLDeformerAsset.h"
+#include "MLDeformerInputInfo.h"
+#include "Components/SkeletalMeshComponent.h"
+#include "Animation/AnimInstance.h"
+#include "NeuralNetwork.h"
+
+namespace UE::MLDeformer
+{
+	FMLDeformerModelInstance::FMLDeformerModelInstance(UMLDeformerModel* InModel)
+		: Model(InModel)
+	{
+		check(InModel);
+	}
+
+	FMLDeformerModelInstance::~FMLDeformerModelInstance()
+	{
+		Release();
+	}
+
+	void FMLDeformerModelInstance::Release()
+	{
+		UNeuralNetwork* NeuralNetwork = Model->GetNeuralNetwork();
+		if (NeuralNetwork != nullptr && NeuralNetworkInferenceHandle != -1)
+		{
+			NeuralNetwork->DestroyInferenceContext(NeuralNetworkInferenceHandle);
+			NeuralNetworkInferenceHandle = -1;
+		}
+	}
+
+	void FMLDeformerModelInstance::Init(USkeletalMeshComponent* SkelMeshComponent)
+	{
+		SkeletalMeshComponent = SkelMeshComponent;
+
+		if (SkelMeshComponent == nullptr)
+		{
+			AssetBonesToSkelMeshMappings.Empty();
+			return;
+		}
+
+		USkeletalMesh* SkelMesh = SkelMeshComponent->SkeletalMesh;
+		if (SkelMesh)
+		{
+			// Init the bone mapping table.
+			const UMLDeformerInputInfo* InputInfo = Model->GetInputInfo();
+			const int32 NumAssetBones = InputInfo->GetNumBones();
+			AssetBonesToSkelMeshMappings.Reset();
+			AssetBonesToSkelMeshMappings.AddUninitialized(NumAssetBones);
+
+			// For each bone in the deformer asset, find the matching bone index inside the skeletal mesh component.
+			for (int32 Index = 0; Index < NumAssetBones; ++Index)
+			{
+				const FName BoneName = InputInfo->GetBoneName(Index);
+				const int32 SkelMeshBoneIndex = SkeletalMeshComponent->GetBoneIndex(BoneName);
+				AssetBonesToSkelMeshMappings[Index] = SkelMeshBoneIndex;
+			}
+		}
+
+		// Perform a compatibility check.
+		UpdateCompatibilityStatus();
+	}
+
+	void FMLDeformerModelInstance::UpdateCompatibilityStatus()
+	{
+		bIsCompatible = SkeletalMeshComponent->SkeletalMesh && CheckCompatibility(SkeletalMeshComponent, true).IsEmpty();
+	}
+
+	FString FMLDeformerModelInstance::CheckCompatibility(USkeletalMeshComponent* InSkelMeshComponent, bool LogIssues)
+	{
+		ErrorText = FString();
+
+		// If we're not compatible, generate a compatibility string.
+		USkeletalMesh* SkelMesh = InSkelMeshComponent ? InSkelMeshComponent->SkeletalMesh.Get() : nullptr;
+		UMLDeformerInputInfo* InputInfo = Model->GetInputInfo();
+		if (SkelMesh && !InputInfo->IsCompatible(SkelMesh))
+		{
+			ErrorText += InputInfo->GenerateCompatibilityErrorString(SkelMesh);
+			ErrorText += "\n";
+			check(!ErrorText.IsEmpty());
+			if (LogIssues)
+			{
+				UE_LOG(LogMLDeformer, Error, TEXT("ML Deformer '%s' isn't compatible with Skeletal Mesh '%s'.\nReason(s):\n%s"), 
+					*Model->GetDeformerAsset()->GetName(), 
+					*SkelMesh->GetName(), 
+					*ErrorText);
+			}
+		}
+
+		UNeuralNetwork* NeuralNetwork = Model->GetNeuralNetwork();
+		if (NeuralNetwork != nullptr && NeuralNetwork->IsLoaded())
+		{
+			const int64 NumNeuralNetInputs = NeuralNetwork->GetInputTensor().Num();
+			const int64 NumDeformerAssetInputs = static_cast<int64>(Model->GetInputInfo()->CalcNumNeuralNetInputs());
+			if (NumNeuralNetInputs != NumDeformerAssetInputs)
+			{
+				const FString InputErrorString = "The number of network inputs doesn't match the asset. Please retrain the asset."; 
+				ErrorText += InputErrorString + "\n";
+				if (LogIssues)
+				{
+					UE_LOG(LogMLDeformer, Error, TEXT("Deformer '%s': %s"), *(Model->GetDeformerAsset()->GetName()), *InputErrorString);
+				}
+			}
+		}
+
+		return ErrorText;
+	}
+
+	int64 FMLDeformerModelInstance::SetBoneTransforms(float* OutputBuffer, int64 OutputBufferSize, int64 StartIndex)
+	{
+		// Extract the bone space bone transforms from the component.
+		// Write the output transforms into the BoneTransforms array.
+		BoneTransforms = SkeletalMeshComponent->GetBoneSpaceTransforms();
+
+		// Write the transforms into the output buffer.
+		const UMLDeformerInputInfo* InputInfo = Model->GetInputInfo();
+		const int32 AssetNumBones = InputInfo->GetNumBones();
+		int64 Index = StartIndex;
+		check((Index + AssetNumBones * 6) <= OutputBufferSize); // Make sure we don't write past the OutputBuffer. (6 because of two columns of the 3x3 rotation matrix)
+		for (int32 BoneIndex = 0; BoneIndex < AssetNumBones; ++BoneIndex)
+		{
+			const int32 SkelMeshBoneIndex = AssetBonesToSkelMeshMappings[BoneIndex];
+			const FMatrix RotationMatrix = (SkelMeshBoneIndex != INDEX_NONE) ? BoneTransforms[SkelMeshBoneIndex].GetRotation().ToMatrix() : FMatrix::Identity;
+			const FVector X = RotationMatrix.GetColumn(0);
+			const FVector Y = RotationMatrix.GetColumn(1);	
+			OutputBuffer[Index++] = X.X;
+			OutputBuffer[Index++] = X.Y;
+			OutputBuffer[Index++] = X.Z;
+			OutputBuffer[Index++] = Y.X;
+			OutputBuffer[Index++] = Y.Y;
+			OutputBuffer[Index++] = Y.Z;
+		}
+
+		return Index;
+	}
+
+	int64 FMLDeformerModelInstance::SetCurveValues(float* OutputBuffer, int64 OutputBufferSize, int64 StartIndex)
+	{
+		const UMLDeformerInputInfo* InputInfo = Model->GetInputInfo();
+
+		// Write the weights into the output buffer.
+		int64 Index = StartIndex;
+		const int32 AssetNumCurves = InputInfo->GetNumCurves();
+		check((Index + AssetNumCurves) <= OutputBufferSize); // Make sure we don't write past the OutputBuffer.
+
+		// Write the curve weights to the output buffer.
+		UAnimInstance* AnimInstance = SkeletalMeshComponent->GetAnimInstance();
+		if (AnimInstance)
+		{
+			for (int32 CurveIndex = 0; CurveIndex < AssetNumCurves; ++CurveIndex)
+			{
+				const FName CurveName = InputInfo->GetCurveName(CurveIndex);
+				const float CurveValue = AnimInstance->GetCurveValue(CurveName);	// Outputs 0.0 when not found.
+				OutputBuffer[Index++] = CurveValue;
+			}
+		}
+		else
+		{
+			for (int32 CurveIndex = 0; CurveIndex < AssetNumCurves; ++CurveIndex)
+			{
+				OutputBuffer[Index++] = 0.0f;
+			}
+		}
+
+		return Index;
+	}
+
+	void FMLDeformerModelInstance::SetNeuralNetworkInputValues(float* InputData, int64 NumInputFloats)
+	{
+		check(SkeletalMeshComponent);
+
+		// Feed data to the network inputs.
+		int64 BufferOffset = 0;
+		BufferOffset = SetBoneTransforms(InputData, NumInputFloats, BufferOffset);
+		BufferOffset = SetCurveValues(InputData, NumInputFloats, BufferOffset);
+		check(BufferOffset == NumInputFloats);
+	}
+
+	bool FMLDeformerModelInstance::IsValidForDataProvider() const
+	{
+		const UNeuralNetwork* NeuralNetwork = Model->GetNeuralNetwork();
+		if (NeuralNetwork == nullptr || 
+			!NeuralNetwork->IsLoaded() || 
+			NeuralNetwork->GetDeviceType() != ENeuralDeviceType::GPU || 
+			NeuralNetwork->GetOutputDeviceType() != ENeuralDeviceType::GPU)
+		{
+			return false;
+		}
+
+		return
+			Model->GetVertexMapBuffer().ShaderResourceViewRHI != nullptr &&
+			GetNeuralNetworkInferenceHandle() != -1;
+	}
+
+	void FMLDeformerModelInstance::Tick(float DeltaTime)
+	{
+		// Some safety checks.
+		if (Model == nullptr || 
+			SkeletalMeshComponent == nullptr || 
+			SkeletalMeshComponent->SkeletalMesh == nullptr || 
+			!bIsCompatible)
+		{
+			return;
+		}
+
+		// Get the network and make sure it's loaded.
+		UNeuralNetwork* NeuralNetwork = Model->GetNeuralNetwork();
+		if (NeuralNetwork == nullptr || !NeuralNetwork->IsLoaded())
+		{
+			return;
+		}
+
+		// If we're not on the GPU we can't continue really.
+		// This is needed as the deformer graph system needs it on the GPU.
+		// Some platforms might not support GPU yet.
+		// Only the inputs are on the CPU.
+		check(NeuralNetwork->GetInputDeviceType() == ENeuralDeviceType::CPU);
+		if (!bAllowCPU)
+		{
+			if (NeuralNetwork->GetDeviceType() != ENeuralDeviceType::GPU ||
+				NeuralNetwork->GetOutputDeviceType() != ENeuralDeviceType::GPU)
+			{
+				return;
+			}
+		}
+
+		// Allocate an inference context if none has already been allocated.
+		if (NeuralNetworkInferenceHandle == -1)
+		{
+			NeuralNetworkInferenceHandle = NeuralNetwork->CreateInferenceContext();
+			if (NeuralNetworkInferenceHandle == -1)
+			{
+				return;
+			}
+		}
+
+		// If the neural network expects a different number of inputs, do nothing.
+		const int64 NumNeuralNetInputs = NeuralNetwork->GetInputTensorForContext(NeuralNetworkInferenceHandle).Num();
+		const int64 NumDeformerAssetInputs = Model->GetInputInfo()->CalcNumNeuralNetInputs();
+		if (NumNeuralNetInputs != NumDeformerAssetInputs)
+		{
+			return;
+		}
+
+		// Update and write the input values directly into the input tensor.
+		float* InputDataPointer = static_cast<float*>(NeuralNetwork->GetInputDataPointerMutableForContext(NeuralNetworkInferenceHandle));
+		SetNeuralNetworkInputValues(InputDataPointer, NumNeuralNetInputs);
+
+		// Run the neural network.
+		if (NeuralNetwork->GetOutputDeviceType() == ENeuralDeviceType::GPU)
+		{
+			ENQUEUE_RENDER_COMMAND(RunNeuralNetwork)(
+				[NeuralNetwork, Handle = NeuralNetworkInferenceHandle](FRHICommandListImmediate& RHICmdList)
+			{
+					// Output deltas will be available on GPU for DeformerGraph via UMLDeformerDataProvider.
+					FRDGBuilder GraphBuilder(RHICmdList);
+					NeuralNetwork->Run(GraphBuilder, Handle);
+					GraphBuilder.Execute();
+			});
+		}
+		else
+		{
+			NeuralNetwork->Run();
+		}
+	}
+}	// namespace UE::MLDeformer
diff --git a/Engine/Plugins/Experimental/MLDeformer/MLDeformerFramework/Source/MLDeformerFramework/Private/Tests/MLDeformerTestModel.cpp b/Engine/Plugins/Experimental/MLDeformer/MLDeformerFramework/Source/MLDeformerFramework/Private/Tests/MLDeformerTestModel.cpp
new file mode 100644
--- /dev/null
+++ b/Engine/Plugins/Experimental/MLDeformer/MLDeformerFramework/Source/MLDeformerFramework/Private/Tests/MLDeformerTestModel.cpp
@@ -0,0 +1,26 @@
+// Copyright Epic Games, Inc. All Rights Reserved.
+
+#include "MLDeformerTestModel.h"
+
+namespace UE::MLDeformerTests
+{
+	FTestModelInstance::FTestModelInstance(UMLDeformerModel* InModel)
+		: FMLDeformerModelInstance(InModel)
+	{
+	}
+}	// namespace UE::MLDeformerTests
+
+UE::MLDeformer::FMLDeformerModelInstance* UMLDeformerTestModel::CreateModelInstance()
+{
+	return new UE::MLDeformerTests::FTestModelInstance(this);
+}
+
+FString UMLDeformerTestModel::GetDisplayName() const 
+{ 
+	return FString("Test Model");
+}
+
+void UMLDeformerTestModel::UpdateNumTargetMeshVertices()
+{
+	NumTargetMeshVerts = NumBaseMeshVerts;
+}
diff --git a/Engine/Plugins/Experimental/MLDeformer/MLDeformerFramework/Source/MLDeformerFramework/Private/Tests/MLDeformerTestModel.h b/Engine/Plugins/Experimental/MLDeformer/MLDeformerFramework/Source/MLDeformerFramework/Private/Tests/MLDeformerTestModel.h
new file mode 100644
--- /dev/null
+++ b/Engine/Plugins/Experimental/MLDeformer/MLDeformerFramework/Source/MLDeformerFramework/Private/Tests/MLDeformerTestModel.h
@@ -0,0 +1,33 @@
+// Copyright Epic Games, Inc. All Rights Reserved.
+
+#pragma once
+
+#include "CoreMinimal.h"
+#include "MLDeformerModel.h"
+#include "MLDeformerModelInstance.h"
+#include "MLDeformerTestModel.generated.h"
+
+UCLASS()
+class UMLDeformerTestModel 
+	:  public UMLDeformerModel
+{
+	GENERATED_BODY()
+
+public:
+	UE::MLDeformer::FMLDeformerModelInstance* CreateModelInstance() override;
+	FString GetDisplayName() const override;
+	void UpdateNumTargetMeshVertices() override;
+};
+
+namespace UE::MLDeformerTests
+{
+	using namespace UE::MLDeformer;
+
+	class FTestModelInstance
+		: public FMLDeformerModelInstance
+	{
+	public:
+		FTestModelInstance(UMLDeformerModel* InModel);
+		void Tick(float DeltaTime) {}
+	};
+}	// namespace UE::MLDeformerTests
diff --git a/Engine/Plugins/Experimental/MLDeformer/MLDeformerFramework/Source/MLDeformerFramework/Private/Tests/ModelTests.cpp b/Engine/Plugins/Experimental/MLDeformer/MLDeformerFramework/Source/MLDeformerFramework/Private/Tests/ModelTests.cpp
new file mode 100644
--- /dev/null
+++ b/Engine/Plugins/Experimental/MLDeformer/MLDeformerFramework/Source/MLDeformerFramework/Private/Tests/ModelTests.cpp
@@ -0,0 +1,96 @@
+// Copyright Epic Games, Inc. All Rights Reserved.
+
+#include "CoreMinimal.h"
+#include "Misc/AutomationTest.h"
+
+#if WITH_DEV_AUTOMATION_TESTS
+
+#include "MLDeformerTestModel.h"
+#include "MLDeformerAsset.h"
+#include "MLDeformerComponent.h"
+#include "UObject/GCObjectScopeGuard.h"
+#include "PreviewScene.h"
+#include "Components/SkeletalMeshComponent.h"
+
+namespace UE::MLDeformerTests
+{
+	IMPLEMENT_SIMPLE_AUTOMATION_TEST(FModelMainTest, "MLDeformer.Model.MainTest", EAutomationTestFlags::ApplicationContextMask | EAutomationTestFlags::EngineFilter);
+	bool FModelMainTest::RunTest(const FString& Parameters)
+	{
+		using namespace UE::MLDeformer;
+
+#if WITH_EDITORONLY_DATA
+		UMLDeformerAsset* DeformerAsset = NewObject<UMLDeformerAsset>();
+		FGCObjectScopeGuard AssetScopeGuard(DeformerAsset);
+		UMLDeformerTestModel* TestModel = NewObject<UMLDeformerTestModel>(DeformerAsset);
+		FGCObjectScopeGuard ModelScopeGuard(TestModel);
+		UTEST_NOT_NULL(TEXT("DeformerAsset check"), DeformerAsset);
+		UTEST_NOT_NULL(TEXT("Creating test model"), TestModel);
+		UTEST_EQUAL(TEXT("Display name check"), TestModel->GetDisplayName(), FString("Test Model"));
+
+		USkeletalMesh* SkeletalMesh = LoadObject<USkeletalMesh>(nullptr, TEXT("SkeletalMesh'/MLDeformerFramework/Tests/SKM_Cylinder.SKM_Cylinder'"));
+		UTEST_NOT_NULL(TEXT("Loading skeletal mesh"), SkeletalMesh);
+
+		TestModel->SkeletalMesh = SkeletalMesh;
+		UTEST_EQUAL(TEXT("GetSkeletalMesh check"), TestModel->GetSkeletalMesh(), SkeletalMesh);
+
+		DeformerAsset->SetModel(TestModel);
+		UTEST_TRUE(TEXT("GetModel check"), DeformerAsset->GetModel() == TestModel);
+
+		TestModel->Init(DeformerAsset);
+		UTEST_TRUE(TEXT("GetDeformerAsset check"), TestModel->GetDeformerAsset() == DeformerAsset);
+
+		const int32 ExpectedNumImportedVerts = 1178;
+		const int32 ImportedVerts = UMLDeformerModel::ExtractNumImportedSkinnedVertices(SkeletalMesh);
+		UTEST_EQUAL(TEXT("Imported verts check"), ImportedVerts, ExpectedNumImportedVerts);
+
+		TestModel->UpdateCachedNumVertices();
+		UTEST_EQUAL(TEXT("Imported verts check"), TestModel->NumBaseMeshVerts, ImportedVerts);
+
+		const USkeleton* Skeleton = SkeletalMesh->GetSkeleton();
+		bool bErrorOnInvalid = false;
+		const USkeleton* ModelSkeleton = TestModel->GetSkeleton(bErrorOnInvalid, nullptr);
+		UTEST_EQUAL(TEXT("Skeleton check"), ModelSkeleton, Skeleton);
+
+		// Create a preview scene and spawn an actor.
+		TUniquePtr<FPreviewScene> PreviewScene = MakeUnique<FPreviewScene>(FPreviewScene::ConstructionValues());
+		UWorld* World = PreviewScene->GetWorld();
+		UTEST_NOT_NULL(TEXT("World check"), World);
+		FActorSpawnParameters SpawnParams;
+		SpawnParams.Name = MakeUniqueObjectName(World, AActor::StaticClass(), FName(TEXT("ML Deformer Test Actor")));	
+		AActor* Actor = World->SpawnActor<AActor>(SpawnParams);
+		Actor->SetFlags(RF_Transient);
+		UTEST_NOT_NULL(TEXT("Actor check"), Actor);
+
+		// Add a skeletal mesh component.
+		USkeletalMeshComponent* SkelMeshComponent = NewObject<USkeletalMeshComponent>(Actor);
+		SkelMeshComponent->SetSkeletalMesh(SkeletalMesh);
+		Actor->SetRootComponent(SkelMeshComponent);
+		SkelMeshComponent->RegisterComponent();
+		
+		// Add an ML Deformer component and activate it.
+		UMLDeformerComponent* MLDeformerComponent = NewObject<UMLDeformerComponent>(Actor);
+		MLDeformerComponent->SetDeformerAsset(DeformerAsset);
+		MLDeformerComponent->RegisterComponent();
+		MLDeformerComponent->Activate();
+		UTEST_EQUAL(TEXT("MLDeformerComponent weight check"), MLDeformerComponent->GetWeight(), 1.0f);
+		UTEST_TRUE(TEXT("MLDeformerComponent SkelMeshComponent check"), MLDeformerComponent->GetSkeletalMeshComponent() == SkelMeshComponent);
+
+		// Create and init the model instance.
+		FMLDeformerModelInstance* ModelInstance = MLDeformerComponent->GetModelInstance();
+		UTEST_NOT_NULL(TEXT("CreateModelInstance check"), ModelInstance);
+		UTEST_TRUE(TEXT("ModelInstance GetModel check"), ModelInstance->GetModel() == TestModel);
+		UTEST_TRUE(TEXT("ModelInstance compatible check"), ModelInstance->IsCompatible());
+
+		// Do some input info checks.
+		UMLDeformerInputInfo* InputInfo = TestModel->GetInputInfo();
+		UTEST_NOT_NULL(TEXT("Get InputInfo"), InputInfo);
+		UTEST_EQUAL(TEXT("InputInfo NumBaseMeshVerts check"), InputInfo->GetNumBaseMeshVertices(), 0);
+		UTEST_EQUAL(TEXT("InputInfo NumBones check"), InputInfo->GetNumBones(), 0);
+		UTEST_EQUAL(TEXT("InputInfo NumCurves check"), InputInfo->GetNumCurves(), 0);
+#endif	// #if WITH_EDITORONLY_DATA
+		return true;
+	}
+} // namespace UE::MLDeformerTests
+
+#endif // WITH_DEV_AUTOMATION_TESTS
diff --git a/Engine/Plugins/Experimental/MLDeformer/MLDeformerFramework/Source/MLDeformerFramework/Public/MLDeformerAsset.h b/Engine/Plugins/Experimental/MLDeformer/MLDeformerFramework/Source/MLDeformerFramework/Public/MLDeformerAsset.h
new file mode 100644
--- /dev/null
+++ b/Engine/Plugins/Experimental/MLDeformer/MLDeformerFramework/Source/MLDeformerFramework/Public/MLDeformerAsset.h
@@ -0,0 +1,33 @@
+// Copyright Epic Games, Inc. All Rights Reserved.
+
+#pragma once
+
+#include "CoreMinimal.h"
+#include "UObject/Object.h"
+#include "MLDeformerAsset.generated.h"
+
+class UMLDeformerModel;
+
+/**
+ * The machine learning deformer asset class.
+ * This class contains a Model property, through which all the magic happens.
+ */
+UCLASS(BlueprintType, hidecategories=Object)
+class MLDEFORMERFRAMEWORK_API UMLDeformerAsset 
+	: public UObject
+{
+	GENERATED_BODY()
+
+public:
+	UMLDeformerModel* GetModel() const { return Model.Get(); }
+	void SetModel(UMLDeformerModel* InModel) { Model = InModel; }
+
+public:
+	/** The ML Deformer model, used to deform the mesh. */
+	UPROPERTY()
+	TObjectPtr<UMLDeformerModel> Model;
+
+	/** Delegate that will be called immediately before the NeuralNetwork is changed. */
+	DECLARE_MULTICAST_DELEGATE(FNeuralNetworkModifyDelegate);
+	FNeuralNetworkModifyDelegate NeuralNetworkModifyDelegate;
+};
diff --git a/Engine/Plugins/Experimental/MLDeformer/MLDeformerFramework/Source/MLDeformerFramework/Public/MLDeformerComponent.h b/Engine/Plugins/Experimental/MLDeformer/MLDeformerFramework/Source/MLDeformerFramework/Public/MLDeformerComponent.h
new file mode 100644
--- /dev/null
+++ b/Engine/Plugins/Experimental/MLDeformer/MLDeformerFramework/Source/MLDeformerFramework/Public/MLDeformerComponent.h
@@ -0,0 +1,79 @@
+// Copyright Epic Games, Inc. All Rights Reserved.
+
+#pragma once
+
+#include "CoreMinimal.h"
+#include "UObject/Object.h"
+#include "Components/ActorComponent.h"
+#include "MLDeformerModelInstance.h"
+#include "MLDeformerComponent.generated.h"
+
+class UMLDeformerAsset;
+class USkeletalMeshComponent;
+
+namespace UE::MLDeformer
+{
+	class FMLDeformerModelInstance;
+}
+
+/**
+ * The ML mesh deformer component.
+ * This works in combination with a MLDeformerAsset and SkeletalMeshComponent.
+ * The component will perform runtime inference of the deformer model setup inside the asset.
+ */
+UCLASS(Blueprintable, ClassGroup = Component, BlueprintType, meta = (BlueprintSpawnableComponent))
+class MLDEFORMERFRAMEWORK_API UMLDeformerComponent : public UActorComponent
+{
+	GENERATED_UCLASS_BODY()
+
+public:
+	// UActorComponent overrides.
+	void Activate(bool bReset=false) override;
+	void Deactivate() override;
+	// ~END UActorComponent overrides.
+
+	UFUNCTION(BlueprintCallable, Category = "MLDeformer")
+	void SetupComponent(UMLDeformerAsset* InDeformerAsset, USkeletalMeshComponent* InSkelMeshComponent);
+
+	UMLDeformerAsset* GetDeformerAsset() { return DeformerAsset; }
+	void SetDeformerAsset(UMLDeformerAsset* InDeformerAsset) { DeformerAsset = InDeformerAsset; }
+	float GetWeight() const { return Weight; }
+	void SetWeight(float NormalizedWeightValue) { Weight = NormalizedWeightValue;  }
+	UE::MLDeformer::FMLDeformerModelInstance* GetModelInstance() const { return ModelInstance.Get(); }
+	USkeletalMeshComponent* GetSkeletalMeshComponent() const { return SkelMeshComponent.Get(); }
+
+protected:
+	// AActorComponent overrides.
+	virtual void TickComponent(float DeltaTime, enum ELevelTick TickType, FActorComponentTickFunction* ThisTickFunction) override;
+	// ~END AActorComponent overrides.
+
+	void Init();
+
+	/** Bind to the DeformerAsset NeuralNetworkModifyDelegate. */
+	void AddNeuralNetworkModifyDelegate();
+
+	/** Unbind from the DeformerAsset NeuralNetworkModifyDelegate. */
+	void RemoveNeuralNetworkModifyDelegate();
+
+protected:
+	/** The deformer asset to use. */
+	UPROPERTY(EditAnywhere, Category = "ML Deformer")
+	TObjectPtr<UMLDeformerAsset> DeformerAsset = nullptr;
+
+	/** How active is this deformer? Can be used to blend it in and out. */
+	UPROPERTY(EditAnywhere, Category = "ML Deformer", meta = (ClampMin = "0.0", ClampMax = "1.0"))
+	float Weight = 1.0f;
+
+	/** 
+	 * The skeletal mesh component we want to grab the bone transforms etc from. 
+	 * This can be a nullptr. When it is a nullptr then it will internally try to find the first skeletal mesh component on the actor.
+	 * You can see this as an override. You can specify this override through the SetupComponent function.
+	 */
+	TObjectPtr<USkeletalMeshComponent> SkelMeshComponent = nullptr;
+
+	/** The deformation model instance. This is used to perform the runtime updates and run the inference. */
+	TUniquePtr<UE::MLDeformer::FMLDeformerModelInstance> ModelInstance = nullptr;
+
+	/** DelegateHandle for NeuralNetwork modification. */
+	FDelegateHandle NeuralNetworkModifyDelegateHandle;
+};
diff --git a/Engine/Plugins/Experimental/MLDeformer/MLDeformerFramework/Source/MLDeformerFramework/Public/MLDeformerGeomCacheHelpers.h b/Engine/Plugins/Experimental/MLDeformer/MLDeformerFramework/Source/MLDeformerFramework/Public/MLDeformerGeomCacheHelpers.h
new file mode 100644
--- /dev/null
+++ b/Engine/Plugins/Experimental/MLDeformer/MLDeformerFramework/Source/MLDeformerFramework/Public/MLDeformerGeomCacheHelpers.h
@@ -0,0 +1,42 @@
+// Copyright Epic Games, Inc. All Rights Reserved.
+
+#pragma once
+
+#include "CoreMinimal.h"
+#include "MLDeformerModule.h"
+
+class UGeometryCache;
+class USkeletalMesh;
+class UAnimSequence;
+
+namespace UE::MLDeformer
+{
+#if WITH_EDITORONLY_DATA
+	/** Maps SkeletalMesh imported meshes to geometry cache tracks. */
+	struct FMLDeformerGeomCacheMeshMapping
+	{
+		int32 MeshIndex = INDEX_NONE;	// The imported model's mesh info index.
+		int32 TrackIndex = INDEX_NONE;	// The geom cache track that this mesh is mapped to.
+		TArray<int32> SkelMeshToTrackVertexMap;	// This maps imported model individual meshes to the geomcache track's mesh data.
+		TArray<int32> ImportedVertexToRenderVertexMap; // Map the imported dcc vertex number to a render vertex. This is just one of the duplicates, which shares the same position.
+	};
+
+	// Error checks.
+	MLDEFORMERFRAMEWORK_API FText GetGeomCacheErrorText(USkeletalMesh* InSkeletalMesh, UGeometryCache* InGeomCache);
+	MLDEFORMERFRAMEWORK_API FText GetGeomCacheAnimSequenceErrorText(UGeometryCache* InGeomCache, UAnimSequence* InAnimSequence);
+	MLDEFORMERFRAMEWORK_API FText GetGeomCacheVertexErrorText(USkeletalMesh* InSkelMesh, UGeometryCache* InGeomCache, const FText& SkelName, const FText& GeomCacheName);
+	MLDEFORMERFRAMEWORK_API FText GetGeomCacheMeshMappingErrorText(USkeletalMesh* InSkelMesh, UGeometryCache* InGeomCache);
+
+	// Geom cache operations.
+	MLDEFORMERFRAMEWORK_API int32 ExtractNumImportedGeomCacheVertices(UGeometryCache* GeometryCache);
+	MLDEFORMERFRAMEWORK_API void GenerateGeomCacheMeshMappings(USkeletalMesh* SkelMesh, UGeometryCache* GeomCache, TArray<FMLDeformerGeomCacheMeshMapping>& OutMeshMappings, TArray<FString>& OutFailedImportedMeshNames);
+	MLDEFORMERFRAMEWORK_API void SampleGeomCachePositions(
+		int32 InLODIndex,
+		float InSampleTime,
+		const TArray<FMLDeformerGeomCacheMeshMapping>& InMeshMappings,
+		const USkeletalMesh* SkelMesh,
+		const UGeometryCache* InGeometryCache,
+		const FTransform& AlignmentTransform,
+		TArray<FVector3f>& OutPositions);
+#endif	// #if WITH_EDITORONLY_DATA
+}	// namespace UE::MLDeformer
diff --git a/Engine/Plugins/Experimental/MLDeformer/MLDeformerFramework/Source/MLDeformerFramework/Public/MLDeformerGraphDataInterface.h b/Engine/Plugins/Experimental/MLDeformer/MLDeformerFramework/Source/MLDeformerFramework/Public/MLDeformerGraphDataInterface.h
new file mode 100644
--- /dev/null
+++ b/Engine/Plugins/Experimental/MLDeformer/MLDeformerFramework/Source/MLDeformerFramework/Public/MLDeformerGraphDataInterface.h
@@ -0,0 +1,147 @@
+// Copyright Epic Games, Inc. All Rights Reserved.
+
+#pragma once
+
+#include "CoreMinimal.h"
+#include "OptimusComputeDataInterface.h"
+#include "ComputeFramework/ComputeDataProvider.h"
+#include "Rendering/SkeletalMeshRenderData.h"
+#include "Rendering/SkeletalMeshLODRenderData.h"
+#include "SkeletalRenderPublic.h"
+#include "MLDeformerGraphDataInterface.generated.h"
+
+class FRDGBuffer;
+class FRDGBufferSRV;
+class FRHIShaderResourceView;
+class FSkeletalMeshObject;
+class UMLDeformerComponent;
+class UNeuralNetwork;
+class USkeletalMeshComponent;
+class UMLDeformerModel;
+
+#define MLDEFORMER_SHADER_PARAMETERS() \
+	SHADER_PARAMETER(uint32, NumVertices) \
+	SHADER_PARAMETER(uint32, InputStreamStart) \
+	SHADER_PARAMETER(float, Weight) \
+	SHADER_PARAMETER_RDG_BUFFER_SRV(Buffer<float>, PositionDeltaBuffer) \
+	SHADER_PARAMETER_SRV(Buffer<uint>, VertexMapBuffer)
+
+#define MLDEFORMER_GRAPH_DISPATCH_START(ParameterStructType, InDispatchSetup, InOutDispatchData) \
+	if (!ensure(InDispatchSetup.ParameterStructSizeForValidation == sizeof(ParameterStructType))) \
+	{ \
+		return; \
+	} \
+	const FSkeletalMeshRenderData& SkeletalMeshRenderData = SkeletalMeshObject->GetSkeletalMeshRenderData(); \
+	const FSkeletalMeshLODRenderData* LodRenderData = SkeletalMeshRenderData.GetPendingFirstLOD(0); \
+	for (int32 InvocationIndex = 0; InvocationIndex < InDispatchSetup.NumInvocations; ++InvocationIndex) \
+	{ \
+		const FSkelMeshRenderSection& RenderSection = LodRenderData->RenderSections[InvocationIndex]; \
+		ParameterStructType* Parameters = (ParameterStructType*)(InOutDispatchData.ParameterBuffer + InDispatchSetup.ParameterBufferOffset + InDispatchSetup.ParameterBufferStride * InvocationIndex);
+
+#define MLDEFORMER_GRAPH_DISPATCH_DEFAULT_PARAMETERS() \
+	Parameters->NumVertices = 0; \
+	Parameters->InputStreamStart = RenderSection.BaseVertexIndex; \
+	Parameters->Weight = Weight; \
+	Parameters->PositionDeltaBuffer = BufferSRV; \
+	Parameters->VertexMapBuffer = VertexMapBufferSRV;
+
+#define MLDEFORMER_GRAPH_DISPATCH_END() }
+
+#define MLDEFORMER_GRAPH_IMPLEMENT_BASICS(InterfaceClassName, DataProviderClassName, DataProviderProxyClassName, ParameterStructType, HLSLText, DisplayName) \
+	FString InterfaceClassName::GetDisplayName() const \
+	{ \
+		return DisplayName; \
+	} \
+	void InterfaceClassName::GetShaderParameters(TCHAR const* UID, FShaderParametersMetadataBuilder& OutBuilder) const \
+	{ \
+		OutBuilder.AddNestedStruct<ParameterStructType>(UID); \
+	} \
+	void InterfaceClassName::GetHLSL(FString& OutHLSL) const \
+	{ \
+		OutHLSL += HLSLText; \
+	} \
+	UComputeDataProvider* InterfaceClassName::CreateDataProvider(TArrayView<TObjectPtr<UObject>> InSourceObjects, uint64 InInputMask, uint64 InOutputMask) const \
+	{ \
+		DataProviderClassName* Provider = NewObject<DataProviderClassName>(); \
+		if (InSourceObjects.Num() == 1) \
+		{ \
+			Provider->SkeletalMeshComponent = Cast<USkeletalMeshComponent>(InSourceObjects[0]); \
+		} \
+		return Provider; \
+	} \
+	FComputeDataProviderRenderProxy* DataProviderClassName::GetRenderProxy() \
+	{ \
+		UMLDeformerComponent* DeformerComponent = SkeletalMeshComponent->GetOwner()->FindComponentByClass<UMLDeformerComponent>(); \
+		return new DataProviderProxyClassName(SkeletalMeshComponent, DeformerComponent); \
+	}
+
+#if WITH_EDITORONLY_DATA
+	#define MLDEFORMER_EDITORDATA_ONLY(Statement, ElseStatement) Statement
+#else
+	#define MLDEFORMER_EDITORDATA_ONLY(X, ElseStatement) ElseStatement
+#endif
+
+/** Compute Framework Data Interface for MLDefomer data. */
+UCLASS(Category = ComputeFramework)
+class MLDEFORMERFRAMEWORK_API UMLDeformerGraphDataInterface
+	: public UOptimusComputeDataInterface
+{
+	GENERATED_BODY()
+
+public:
+	// UOptimusComputeDataInterface overrides.
+	virtual FString GetDisplayName() const override;
+	virtual TArray<FOptimusCDIPinDefinition> GetPinDefinitions() const override;
+	// ~END UOptimusComputeDataInterface overrides.
+
+	// UComputeDataInterface overrides.
+	virtual void GetSupportedInputs(TArray<FShaderFunctionDefinition>& OutFunctions) const override;
+	virtual void GetShaderParameters(TCHAR const* UID, FShaderParametersMetadataBuilder& OutBuilder) const override;
+	virtual void GetHLSL(FString& OutHLSL) const override;
+	virtual void GetSourceTypes(TArray<UClass*>& OutSourceTypes) const override;
+	virtual UComputeDataProvider* CreateDataProvider(TArrayView<TObjectPtr<UObject>> InSourceObjects, uint64 InInputMask, uint64 InOutputMask) const override;
+	// ~END UComputeDataInterface overrides.
+};
+
+/** Compute Framework Data Provider for MLDeformer data. */
+UCLASS(BlueprintType, EditInlineNew, Category = ComputeFramework)
+class MLDEFORMERFRAMEWORK_API UMLDeformerGraphDataProvider
+	: public UComputeDataProvider
+{
+	GENERATED_BODY()
+
+public:
+	UPROPERTY(EditAnywhere, BlueprintReadWrite, Category = Binding)
+	TObjectPtr<USkeletalMeshComponent> SkeletalMeshComponent = nullptr;
+
+	// UComputeDataProvider overrides.
+	virtual bool IsValid() const override;
+	virtual FComputeDataProviderRenderProxy* GetRenderProxy() override;
+	// ~END UComputeDataProvider overrides.
+};
+
+namespace UE::MLDeformer
+{
+	/** Compute Framework Data Provider Proxy for MLDeformer data. */
+	class MLDEFORMERFRAMEWORK_API FMLDeformerGraphDataProviderProxy
+		: public FComputeDataProviderRenderProxy
+	{
+	public:
+		FMLDeformerGraphDataProviderProxy(USkeletalMeshComponent* SkeletalMeshComponent, UMLDeformerComponent* DeformerComponent);
+
+		// FComputeDataProviderRenderProxy overrides.
+		virtual void AllocateResources(FRDGBuilder& GraphBuilder) override;
+		virtual void GatherDispatchData(FDispatchSetup const& InDispatchSetup, FCollectedDispatchData& InOutDispatchData) override;
+		// ~END FComputeDataProviderRenderProxy overrides.
+
+	protected:
+		FSkeletalMeshObject* SkeletalMeshObject;
+		TObjectPtr<UNeuralNetwork> NeuralNetwork = nullptr;
+		FRHIShaderResourceView* VertexMapBufferSRV = nullptr;
+		FRDGBuffer* Buffer = nullptr;
+		FRDGBufferSRV* BufferSRV = nullptr;
+		float Weight = 1.0f;
+		int32 NeuralNetworkInferenceHandle = -1;
+		bool bCanRunNeuralNet = false;
+	};
+}	// namespace UE::MLDeformer
diff --git a/Engine/Plugins/Experimental/MLDeformer/MLDeformerFramework/Source/MLDeformerFramework/Public/MLDeformerGraphDebugDataInterface.h b/Engine/Plugins/Experimental/MLDeformer/MLDeformerFramework/Source/MLDeformerFramework/Public/MLDeformerGraphDebugDataInterface.h
new file mode 100644
--- /dev/null
+++ b/Engine/Plugins/Experimental/MLDeformer/MLDeformerFramework/Source/MLDeformerFramework/Public/MLDeformerGraphDebugDataInterface.h
@@ -0,0 +1,165 @@
+// Copyright Epic Games, Inc. All Rights Reserved.
+
+#pragma once
+
+#include "CoreMinimal.h"
+#include "MLDeformerModel.h"
+#include "MLDeformerGraphDataInterface.h"
+#include "OptimusComputeDataInterface.h"
+#include "ComputeFramework/ComputeDataProvider.h"
+#include "MLDeformerGraphDebugDataInterface.generated.h"
+
+class FRDGBuffer;
+class FRDGBufferSRV;
+class FRHIShaderResourceView;
+class FSkeletalMeshObject;
+class UMLDeformerAsset;
+class USkeletalMeshComponent;
+
+#define MLDEFORMER_DEBUG_SHADER_PARAMETERS() \
+	SHADER_PARAMETER(uint32, NumVertices) \
+	SHADER_PARAMETER(uint32, InputStreamStart) \
+	SHADER_PARAMETER(int32, HeatMapMode) \
+	SHADER_PARAMETER(float, HeatMapMax) \
+	SHADER_PARAMETER(float, GroundTruthLerp) \
+	SHADER_PARAMETER(uint32, GroundTruthBufferSize) \
+	SHADER_PARAMETER_RDG_BUFFER_SRV(StructuredBuffer<float3>, PositionGroundTruthBuffer) \
+	SHADER_PARAMETER_SRV(Buffer<uint>, VertexMapBuffer)
+
+#define MLDEFORMER_GRAPH_DISPATCH_DEFAULT_DEBUG_PARAMETERS() \
+	Parameters->NumVertices = 0; \
+	Parameters->InputStreamStart = RenderSection.BaseVertexIndex; \
+	Parameters->HeatMapMode = HeatMapMode; \
+	Parameters->HeatMapMax = HeatMapMax; \
+	Parameters->GroundTruthLerp = GroundTruthLerp; \
+	Parameters->GroundTruthBufferSize = GroundTruthPositions.Num(); \
+	Parameters->PositionGroundTruthBuffer = GroundTruthBufferSRV; \
+	Parameters->VertexMapBuffer = VertexMapBufferSRV;
+
+#define MLDEFORMER_GRAPH_IMPLEMENT_DEBUG_BASICS(InterfaceClassName, DataProviderClassName, ParameterStructType, HLSLText, DisplayName) \
+	FString InterfaceClassName::GetDisplayName() const \
+	{ \
+		return DisplayName; \
+	} \
+	void InterfaceClassName::GetShaderParameters(TCHAR const* UID, FShaderParametersMetadataBuilder& OutBuilder) const \
+	{ \
+		OutBuilder.AddNestedStruct<ParameterStructType>(UID); \
+	} \
+	void InterfaceClassName::GetHLSL(FString& OutHLSL) const \
+	{ \
+		OutHLSL += HLSLText; \
+	} \
+	UComputeDataProvider* InterfaceClassName::CreateDataProvider(TArrayView<TObjectPtr<UObject>> InSourceObjects, uint64 InInputMask, uint64 InOutputMask) const \
+	{ \
+		DataProviderClassName* Provider = NewObject<DataProviderClassName>(); \
+		if (InSourceObjects.Num() == 2) \
+		{ \
+			Provider->SkeletalMeshComponent = Cast<USkeletalMeshComponent>(InSourceObjects[0]); \
+			UMLDeformerComponent* DeformerComponent = Cast<UMLDeformerComponent>(InSourceObjects[1]); \
+			Provider->DeformerAsset = (DeformerComponent != nullptr) ? DeformerComponent->GetDeformerAsset() : nullptr; \
+		} \
+		MLDEFORMER_EDITORDATA_ONLY( \
+			if (Provider->DeformerAsset != nullptr) \
+			{ \
+				Provider->Init(); \
+			} \
+		,) \
+		return Provider; \
+	}
+
+
+#define MLDEFORMER_GRAPH_IMPLEMENT_DEBUG_BASICS_WITH_PROXY(InterfaceClassName, DataProviderClassName, DataProviderProxyClassName, ParameterStructType, HLSLText, DisplayName) \
+	MLDEFORMER_GRAPH_IMPLEMENT_DEBUG_BASICS(InterfaceClassName, DataProviderClassName, ParameterStructType, HLSLText, DisplayName) \
+	FComputeDataProviderRenderProxy* DataProviderClassName::GetRenderProxy() \
+	{ \
+		MLDEFORMER_EDITORDATA_ONLY( \
+			DataProviderProxyClassName* Proxy = new DataProviderProxyClassName(SkeletalMeshComponent, DeformerAsset, this); \
+			const float SampleTime = SkeletalMeshComponent->GetPosition(); \
+			UMLDeformerModel* Model = DeformerAsset->GetModel(); \
+			Model->SampleGroundTruthPositions(SampleTime, Proxy->GetGroundTruthPositions()); \
+			Proxy->HandleZeroGroundTruthPositions(); \
+			return Proxy;, \
+			return nullptr; \
+		) \
+	}
+
+/** 
+ * Compute Framework Data Interface for MLDeformer debugging data. 
+ * This interfaces to editor only data, and so will only give valid results in that context.
+ */
+UCLASS(Category = ComputeFramework)
+class MLDEFORMERFRAMEWORK_API UMLDeformerGraphDebugDataInterface
+	: public UOptimusComputeDataInterface
+{
+	GENERATED_BODY()
+
+public:
+	// UOptimusComputeDataInterface overrides.
+	virtual FString GetDisplayName() const override;
+	virtual TArray<FOptimusCDIPinDefinition> GetPinDefinitions() const override;
+	// ~END UOptimusComputeDataInterface overrides.
+
+	// UComputeDataInterface overrides.
+	virtual void GetSupportedInputs(TArray<FShaderFunctionDefinition>& OutFunctions) const override;
+	virtual void GetShaderParameters(TCHAR const* UID, FShaderParametersMetadataBuilder& OutBuilder) const override;
+	virtual void GetHLSL(FString& OutHLSL) const override;
+	virtual void GetSourceTypes(TArray<UClass*>& OutSourceTypes) const override;
+	virtual UComputeDataProvider* CreateDataProvider(TArrayView<TObjectPtr<UObject>> InSourceObjects, uint64 InInputMask, uint64 InOutputMask) const override;
+	// ~END UComputeDataInterface overrides.
+};
+
+/** Compute Framework Data Provider for MLDeformer debugging data. */
+UCLASS(BlueprintType, EditInlineNew, Category = ComputeFramework)
+class MLDEFORMERFRAMEWORK_API UMLDeformerGraphDebugDataProvider
+	: public UComputeDataProvider
+{
+	GENERATED_BODY()
+
+public:
+	virtual void Init() {};
+
+	// UComputeDataProvider overrides.
+	virtual bool IsValid() const override;
+	virtual FComputeDataProviderRenderProxy* GetRenderProxy() override;
+	// ~END UComputeDataProvider overrides.
+
+public:
+	UPROPERTY(EditAnywhere, BlueprintReadWrite, Category = Binding)
+	TObjectPtr<USkeletalMeshComponent> SkeletalMeshComponent = nullptr;
+
+	UPROPERTY(EditAnywhere, BlueprintReadWrite, Category = Binding)
+	TObjectPtr<UMLDeformerAsset> DeformerAsset = nullptr;
+};
+
+/** Compute Framework Data Provider Proxy for MLDeformer debugging data. */
+#if WITH_EDITORONLY_DATA
+namespace UE::MLDeformer
+{
+	class MLDEFORMERFRAMEWORK_API FMLDeformerGraphDebugDataProviderProxy
+		: public FComputeDataProviderRenderProxy
+	{
+	public:
+		FMLDeformerGraphDebugDataProviderProxy(USkeletalMeshComponent* SkeletalMeshComponent, UMLDeformerAsset* DeformerAsset, UMLDeformerGraphDebugDataProvider* InProvider);
+
+		virtual void HandleZeroGroundTruthPositions();
+
+		// FComputeDataProviderRenderProxy overrides.
+		virtual void AllocateResources(FRDGBuilder& GraphBuilder) override;
+		virtual void GatherDispatchData(FDispatchSetup const& InDispatchSetup, FCollectedDispatchData& InOutDispatchData) override;
+		// ~END FComputeDataProviderRenderProxy overrides.
+
+		TArray<FVector3f>& GetGroundTruthPositions() { return GroundTruthPositions; }
+
+	protected:
+		TObjectPtr<UMLDeformerGraphDebugDataProvider> Provider = nullptr;
+		FSkeletalMeshObject* SkeletalMeshObject;
+		TArray<FVector3f> GroundTruthPositions;
+		FRHIShaderResourceView* VertexMapBufferSRV = nullptr;
+		FRDGBuffer* GroundTruthBuffer = nullptr;
+		FRDGBufferSRV* GroundTruthBufferSRV = nullptr;
+		int32 HeatMapMode = 0;
+		float HeatMapMax = 0.0f;
+		float GroundTruthLerp = 0.0f;
+	};
+}	// namespace UE::MLDeformer
+#endif
diff --git a/Engine/Plugins/Experimental/MLDeformer/MLDeformerFramework/Source/MLDeformerFramework/Public/MLDeformerModel.h b/Engine/Plugins/Experimental/MLDeformer/MLDeformerFramework/Source/MLDeformerFramework/Public/MLDeformerModel.h
new file mode 100644
--- /dev/null
+++ b/Engine/Plugins/Experimental/MLDeformer/MLDeformerFramework/Source/MLDeformerFramework/Public/MLDeformerModel.h
@@ -0,0 +1,214 @@
+// Copyright Epic Games, Inc. All Rights Reserved.
+#pragma once
+
+#include "CoreMinimal.h"
+#include "MLDeformerInputInfo.h"
+#include "UObject/ObjectPtr.h"
+#include "UObject/Object.h"
+#include "Templates/UniquePtr.h"
+#include "Interfaces/Interface_BoneReferenceSkeletonProvider.h"
+#include "BoneContainer.h"
+#include "RenderCommandFence.h"
+#include "RenderResource.h"
+#include "MLDeformerCurveReference.h"
+#include "MLDeformerModel.generated.h"
+
+class UMLDeformerAsset;
+class UMLDeformerVizSettings;
+class UAnimSequence;
+class UGeometryCache;
+class UNeuralNetwork;
+
+namespace UE::MLDeformer
+{
+	class FMLDeformerModelInstance;
+
+	/** The vertex map on the GPU. */
+	class FVertexMapBuffer : public FVertexBufferWithSRV
+	{
+	public:
+		void Init(const TArray<int32>& InVertexMap)	{ VertexMap = InVertexMap; }
+
+	private:
+		void InitRHI() override;
+		TArray<int32> VertexMap;
+	};
+}	// namespace UE::MLDeformer
+
+/** The training inputs. Specifies what data to include in training. */
+UENUM()
+enum class EMLDeformerTrainingInputFilter : uint8
+{
+	BonesAndCurves = 0, /** Include both bone rotations and curve values. */
+	BonesOnly,			/** Include only bone rotations. */
+	CurvesOnly			/** Include only curve values. */
+};
+
+DECLARE_DELEGATE_OneParam(FMLDeformerModelOnPostEditProperty, FPropertyChangedEvent&)
+
+/**
+ * The ML Deformer runtime model base class.
+ * All models should be inherited from this class.
+ **/
+UCLASS(Abstract)
+class MLDEFORMERFRAMEWORK_API UMLDeformerModel 
+	: public UObject	
+	, public IBoneReferenceSkeletonProvider
+{
+	GENERATED_BODY()
+
+public:
+	virtual ~UMLDeformerModel() = default;
+
+	// Methods you might want to override.
+	virtual void Init(UMLDeformerAsset* InDeformerAsset);
+	virtual void InitGPUData();
+	virtual UMLDeformerInputInfo* CreateInputInfo();
+	virtual UE::MLDeformer::FMLDeformerModelInstance* CreateModelInstance();
+#if WITH_EDITORONLY_DATA
+	virtual bool HasTrainingGroundTruth() const { return false; }
+	virtual void SampleGroundTruthPositions(float SampleTime, TArray<FVector3f>& OutPositions) {}
+#endif
+#if WITH_EDITOR
+	virtual void UpdateNumBaseMeshVertices();
+	virtual void UpdateCachedNumVertices();
+
+	static int32 ExtractNumImportedSkinnedVertices(const USkeletalMesh* SkeletalMesh);
+#endif
+
+	// Methods you are required to override.
+	virtual FString GetDisplayName() const PURE_VIRTUAL(UMLDeformerModel::GetDisplayName, return FString(););
+	virtual void UpdateNumTargetMeshVertices() PURE_VIRTUAL(UMLDeformerModel::GetNumTargetMeshVertices; mTargetMeshVerts=0;);
+
+	// UObject overrides.
+	virtual void Serialize(FArchive& Archive) override;
+	virtual void PostLoad() override;
+	virtual void BeginDestroy() override;
+	virtual bool IsReadyForFinishDestroy() override;
+	// ~END UObject overrides.
+
+	// IBoneReferenceSkeletonProvider overrides.
+	virtual USkeleton* GetSkeleton(bool& bInvalidSkeletonIsError, const IPropertyHandle* PropertyHandle) override;
+	// ~END IBoneReferenceSkeletonProvider overrides.
+
+	UMLDeformerAsset* GetDeformerAsset() const;
+	EMLDeformerTrainingInputFilter GetTrainingInputs() const { return TrainingInputs; }
+	UMLDeformerInputInfo* GetInputInfo() const { return InputInfo.Get(); }
+
+	int32 GetNumBaseMeshVerts() const { return NumBaseMeshVerts; }
+	int32 GetNumTargetMeshVerts() const { return NumTargetMeshVerts; }
+
+	const TArray<int32>& GetVertexMap() const { return VertexMap; }
+	const UE::MLDeformer::FVertexMapBuffer& GetVertexMapBuffer() const { return VertexMapBuffer; }
+
+	UNeuralNetwork* GetNeuralNetwork() const { return NeuralNetwork.Get(); }
+	void SetNeuralNetwork(UNeuralNetwork* InNeuralNetwork);
+
+#if WITH_EDITORONLY_DATA
+	// UObject overrides.
+	virtual void PostEditChangeProperty(FPropertyChangedEvent& PropertyChangedEvent) override;
+	// ~END UObject overrides.
+
+	void InitVertexMap();
+
+	FMLDeformerModelOnPostEditProperty& OnPostEditChangeProperty() { return PostEditPropertyEvent; }
+
+	UMLDeformerVizSettings* GetVizSettings() const { return VizSettings; }
+	const USkeletalMesh* GetSkeletalMesh() const { return SkeletalMesh;  }
+	USkeletalMesh* GetSkeletalMesh() { return SkeletalMesh; }
+	const UAnimSequence* GetAnimSequence() const { return AnimSequence;  }
+	UAnimSequence* GetAnimSequence() { return AnimSequence; }
+	int32 GetTrainingFrameLimit() const { return MaxTrainingFrames; }
+	int32 GetMaxCacheSizeInGigaBytes() const { return MaxCacheSizeGB; }
+	const FTransform& GetAlignmentTransform() const { return AlignmentTransform; }
+	TArray<FBoneReference>& GetBoneIncludeList() { return BoneIncludeList; }
+	const TArray<FBoneReference>& GetBoneIncludeList() const { return BoneIncludeList; }
+	TArray<FMLDeformerCurveReference>& GetCurveIncludeList() { return CurveIncludeList; }
+	const TArray<FMLDeformerCurveReference>& GetCurveIncludeList() const { return CurveIncludeList; }
+	float GetDeltaCutoffLength() const { return DeltaCutoffLength; }
+#endif	// #if WITH_EDITORONLY_DATA
+
+protected:
+	TObjectPtr<UMLDeformerAsset> DeformerAsset = nullptr;
+	FMLDeformerModelOnPostEditProperty PostEditPropertyEvent;
+
+	void SetInputInfo(UMLDeformerInputInfo* Input) { InputInfo = Input; }
+
+public:
+	/** Cached number of skeltal mesh vertices. */
+	UPROPERTY()
+	int32 NumBaseMeshVerts = 0;
+
+	/** Cached number of target mesh vertices. */
+	UPROPERTY()
+	int32 NumTargetMeshVerts = 0;
+
+	/** Describes what inputs we should train the neural network on. */
+	UPROPERTY(EditAnywhere, Category = "Inputs and Output")
+	EMLDeformerTrainingInputFilter TrainingInputs = EMLDeformerTrainingInputFilter::BonesAndCurves;
+
+	/** 
+	 * The information about the neural network inputs. This contains things such as bone names and curve names.
+	 */
+	UPROPERTY()
+	TObjectPtr<UMLDeformerInputInfo> InputInfo = nullptr;
+
+	/** This is an index per vertex in the mesh, indicating the imported vertex number from the source asset. */
+	UPROPERTY()
+	TArray<int32> VertexMap;
+
+	/** The neural network that is used during inference. */
+	UPROPERTY()
+	TObjectPtr<UNeuralNetwork> NeuralNetwork = nullptr;
+
+	/** GPU buffers for Vertex Map. */
+	UE::MLDeformer::FVertexMapBuffer VertexMapBuffer;
+
+	/** Fence used in render thread cleanup on destruction. */
+	FRenderCommandFence RenderResourceDestroyFence;
+
+#if WITH_EDITORONLY_DATA
+	UPROPERTY()
+	TObjectPtr<UMLDeformerVizSettings> VizSettings = nullptr;
+
+	/** The skeletal mesh that represents the linear skinned mesh. */
+	UPROPERTY(EditAnywhere, BlueprintReadWrite, Category = "Base Mesh")
+	TObjectPtr<USkeletalMesh> SkeletalMesh = nullptr;
+
+	/**
+	 * The animation sequence to apply to the base mesh. This has to match the animation of the target mesh's geometry cache. 
+	 * Internally we force the Interpolation property for this motion to be "Step".
+	 */
+	UPROPERTY(EditAnywhere, BlueprintReadWrite, Category = "Base Mesh")
+	TObjectPtr<UAnimSequence> AnimSequence = nullptr;
+
+	/** The transform that aligns the Geometry Cache to the SkeletalMesh. This will mostly apply some scale and a rotation, but no translation. */
+	UPROPERTY(EditAnywhere, BlueprintReadWrite, Category = "Target Mesh")
+	FTransform AlignmentTransform = FTransform::Identity;
+
+	/** The bones to include during training. When none are provided, all bones of the Skeleton will be included. */
+	UPROPERTY(EditAnywhere, Category = "Inputs and Output", meta = (EditCondition = "TrainingInputs == EMLDeformerTrainingInputFilter::BonesAndCurves || TrainingInputs == EMLDeformerTrainingInputFilter::BonesOnly"))
+	TArray<FBoneReference> BoneIncludeList;
+
+	/** The curves to include during training. When none are provided, all curves of the Skeleton will be included. */
+	UPROPERTY(EditAnywhere, Category = "Inputs and Output", meta = (EditCondition = "TrainingInputs == EMLDeformerTrainingInputFilter::BonesAndCurves || TrainingInputs == EMLDeformerTrainingInputFilter::CurvesOnly"))
+	TArray<FMLDeformerCurveReference> CurveIncludeList;
+
+	/** The maximum numer of training frames (samples) to train on. Use this to train on a sub-section of your full training data. */
+	UPROPERTY(EditAnywhere, BlueprintReadWrite, Category = "Inputs and Output", meta = (ClampMin = "1"))
+	int32 MaxTrainingFrames = 1000000;
+
+	/** Sometimes there can be some vertices that cause some issues that cause deltas to be very long. We can ignore these deltas by setting a cutoff value. 
+	  * Deltas that are longer than the cutoff value (in units), will be ignored and set to zero length. */
+	UPROPERTY(EditAnywhere, BlueprintReadWrite, Category = "Inputs and Output", meta = (ClampMin = "0.01", ForceUnits="cm"))
+	float DeltaCutoffLength = 30.0f;
+
+	/** 
+	 * The maximum allowed size of the training cache in memory, in gigabytes.
+	 * So a value of 4 would use a maximum of four gigabyte of system memory. 
+	 * The larger the cache size the faster the training.
+	 */
+	UPROPERTY(EditAnywhere, BlueprintReadWrite, Category = "Training Settings", DisplayName = "Max Cache Size", meta = (ClampMin = "0", ForceUnits = "Gigabytes"))
+	int32 MaxCacheSizeGB = 4;	// 4 Gigabyte
+#endif
+};
diff --git a/Engine/Plugins/Experimental/MLDeformer/MLDeformerFramework/Source/MLDeformerFramework/Public/MLDeformerModule.h b/Engine/Plugins/Experimental/MLDeformer/MLDeformerFramework/Source/MLDeformerFramework/Public/MLDeformerModule.h
new file mode 100644
--- /dev/null
+++ b/Engine/Plugins/Experimental/MLDeformer/MLDeformerFramework/Source/MLDeformerFramework/Public/MLDeformerModule.h
@@ -0,0 +1,19 @@
+// Copyright Epic Games, Inc. All Rights Reserved.
+
+#pragma once
+
+#include "CoreMinimal.h"
+#include "Modules/ModuleManager.h"
+
+MLDEFORMERFRAMEWORK_API DECLARE_LOG_CATEGORY_EXTERN(LogMLDeformer, Log, All);
+
+namespace UE::MLDeformer
+{
+	class MLDEFORMERFRAMEWORK_API FMLDeformerModule : public IModuleInterface
+	{
+	public:
+		// IModuleInterface overrides.
+		virtual void StartupModule() override;
+		// ~END IModuleInterface overrides.
+	};
+}
diff --git a/Engine/Plugins/Experimental/MLDeformer/MLDeformerFramework/Source/MLDeformerFramework/Public/MLDeformerVizSettings.h b/Engine/Plugins/Experimental/MLDeformer/MLDeformerFramework/Source/MLDeformerFramework/Public/MLDeformerVizSettings.h
new file mode 100644
--- /dev/null
+++ b/Engine/Plugins/Experimental/MLDeformer/MLDeformerFramework/Source/MLDeformerFramework/Public/MLDeformerVizSettings.h
@@ -0,0 +1,156 @@
+// Copyright Epic Games, Inc. All Rights Reserved.
+
+#pragma once
+
+#include "CoreMinimal.h"
+#include "UObject/Object.h"
+#include "UObject/ObjectPtr.h"
+#include "MLDeformerVizSettings.generated.h"
+
+class UMeshDeformer;
+class UAnimSequence;
+
+/** The visualization mode, which selects whether you want to view the training data, or test your already trained model. */
+UENUM()
+enum class EMLDeformerVizMode : uint8
+{
+	/** Preview the training data. */
+	TrainingData = 0,
+
+	/** Preview testing data, used on trained models. */
+	TestData
+};
+
+/** The heat map mode which selects what the colors of the heatmap represent. */
+UENUM()
+enum class EMLDeformerHeatMapMode : uint8
+{
+	/** Visualize areas where the deformer is applying corrections. The color represents the size of the correction it applies. */
+	Activations = 0,
+
+	/** Visualize the error versus the ground truth model. Requires a ground truth model to be setup. */
+	GroundTruth
+};
+
+/**
+ * The vizualization settings.
+ */
+UCLASS()
+class MLDEFORMERFRAMEWORK_API UMLDeformerVizSettings : public UObject
+{
+	GENERATED_BODY()
+
+public:
+	// UObject overrides.
+	virtual bool IsEditorOnly() const override { return true; }
+#if WITH_EDITOR
+	virtual void PostEditChangeProperty(FPropertyChangedEvent& PropertyChangedEvent) { GetOuter()->PostEditChangeProperty(PropertyChangedEvent); } // Forward to the UMLDeformerModel.
+#endif
+	// ~END UObject overrides.
+
+#if WITH_EDITORONLY_DATA
+	virtual bool HasTestGroundTruth() const { return false; }
+
+	FVector GetMeshSpacingOffsetVector() const { return FVector(MeshSpacing, 0.0f, 0.0f); }
+	float GetMeshSpacing() const { return MeshSpacing; }
+	float GetLabelHeight() const { return LabelHeight; }
+	bool GetDrawLabels() const { return bDrawLabels; }
+	float GetLabelScale() const { return LabelScale; }
+	EMLDeformerVizMode GetVisualizationMode() const { return VisualizationMode; }
+	int32 GetFrameNumber() const { return FrameNumber; }
+	float GetAnimPlaySpeed() const { return AnimPlaySpeed; }
+	const UAnimSequence* GetTestAnimSequence() const { return TestAnimSequence; }
+	UAnimSequence* GetTestAnimSequence() { return TestAnimSequence; }
+	bool GetDrawLinearSkinnedActor() const { return bDrawLinearSkinnedActor; }
+	bool GetDrawMLDeformedActor() const { return bDrawMLDeformedActor; }
+	bool GetDrawGroundTruthActor() const { return bDrawGroundTruthActor; }
+	bool GetShowHeatMap() const { return bShowHeatMap; }
+	EMLDeformerHeatMapMode GetHeatMapMode() const { return HeatMapMode; }
+	float GetHeatMapMax() const { return HeatMapMax; }
+	float GetGroundTruthLerp() const { return HeatMapMode == EMLDeformerHeatMapMode::GroundTruth ? GroundTruthLerp : 0.0f; }
+	UMeshDeformer* GetDeformerGraph() const { return DeformerGraph; }
+	void SetDeformerGraph(UMeshDeformer* InDeformerGraph) { DeformerGraph = InDeformerGraph; }
+	float GetWeight() const { return Weight; }
+	bool GetXRayDeltas() const { return bXRayDeltas; }
+	bool GetDrawVertexDeltas() const { return bDrawDeltas; }
+#endif
+
+public:
+#if WITH_EDITORONLY_DATA
+	/** The data to visualize. */
+	UPROPERTY()
+	EMLDeformerVizMode VisualizationMode = EMLDeformerVizMode::TrainingData;
+
+	/** The animation sequence to play on the skeletal mesh. */
+	UPROPERTY(EditAnywhere, Category = "Test Assets")
+	TObjectPtr<UAnimSequence> TestAnimSequence = nullptr;
+
+	/** The deformer graph to use on the asset editor's deformed test actor. */
+	UPROPERTY(EditAnywhere, Category = "Test Assets")
+	TObjectPtr<UMeshDeformer> DeformerGraph = nullptr;
+
+	/** The play speed factor of the test anim sequence. */
+	UPROPERTY(EditAnywhere, Category = "Live Settings", meta = (ClampMin = "0.0", ClampMax = "2.0", ForceUnits="Multiplier"))
+	float AnimPlaySpeed = 1.0f;
+
+	/** The frame number of the training data to visualize. */
+	UPROPERTY(EditAnywhere, Category = "Training Meshes", meta = (ClampMin = "0"))
+	uint32 FrameNumber = 0;
+
+	/** Draw the text labels above each actor? */
+	UPROPERTY(EditAnywhere, Category = "Shared Settings")
+	bool bDrawLabels = true;
+
+	/** The height in units to draw the labels at. */
+	UPROPERTY(EditAnywhere, Category = "Shared Settings", meta = (EditCondition = "bDrawLabels", ForceUnits="cm"))
+	float LabelHeight = 300.0f;
+
+	/** The scale of the label text. */
+	UPROPERTY(EditAnywhere, Category = "Shared Settings", meta = (ClampMin = "0.001", EditCondition = "bDrawLabels", ForceUnits="Multiplier"))
+	float LabelScale = 1.0f;
+
+	/** The spacing between meshes. */
+	UPROPERTY(EditAnywhere, Category = "Shared Settings", meta = (ClampMin = "0", ForceUnits="cm"))
+	float MeshSpacing = 300.0f;
+
+	/** Show the heat map? This will visualize the active areas of the deformer. */
+	UPROPERTY(EditAnywhere, Category = "Live Settings")
+	bool bShowHeatMap = false;
+
+	/** What should the heatmap visualize? */
+	UPROPERTY(EditAnywhere, Category = "Live Settings", meta = (EditCondition = "bShowHeatMap"))
+	EMLDeformerHeatMapMode HeatMapMode = EMLDeformerHeatMapMode::Activations;
+
+	/** How many centimeters does the most intense color of the heatmap represent? */
+	UPROPERTY(EditAnywhere, Category = "Live Settings", meta = (EditCondition = "bShowHeatMap", ClampMin = "0.001", ForceUnits="cm"))
+	float HeatMapMax = 1.0f;
+
+	/** Lerp from ML deformed model to ground truth model when in heat map mode. */
+	UPROPERTY(EditAnywhere, Category = "Live Settings", meta = (EditCondition = "HeatMapMode==EMLDeformerHeatMapMode::GroundTruth", UIMin = "0.0", UIMax = "1.0", ClampMin = "0.0", ClampMax = "1.0"))
+	float GroundTruthLerp = 0.0f;
+
+	/** Draw the linear skinned actor? */
+	UPROPERTY(EditAnywhere, Category = "Live Settings")
+	bool bDrawLinearSkinnedActor = true;
+
+	/** Draw the ML Deformed actor? */
+	UPROPERTY(EditAnywhere, Category = "Live Settings", DisplayName = "Draw ML Deformed Actor")
+	bool bDrawMLDeformedActor = true;
+
+	/** Draw the ground truth actor? */
+	UPROPERTY(EditAnywhere, Category = "Live Settings")
+	bool bDrawGroundTruthActor = true;
+
+	/** The scale factor of the ML deformer deltas being applied on top of the linear skinned results. */
+	UPROPERTY(EditAnywhere, Category = "Live Settings", meta = (ClampMin = "0.0", ClampMax = "1.0"))
+	float Weight = 1.0f;
+
+	/** Draw the vertex deltas? */
+	UPROPERTY(EditAnywhere, Category = "Training Meshes")
+	bool bDrawDeltas = true;
+
+	/** Enable this to draw the deltas in xray mode? */
+	UPROPERTY(EditAnywhere, Category = "Training Meshes", meta = (EditCondition = "bDrawDeltas"))
+	bool bXRayDeltas = true;
+#endif
+};
diff --git a/Engine/Plugins/Experimental/MLDeformer/MLDeformerFramework/Source/MLDeformerFrameworkEditor/Private/MLDeformerAssetActions.h b/Engine/Plugins/Experimental/MLDeformer/MLDeformerFramework/Source/MLDeformerFrameworkEditor/Private/MLDeformerAssetActions.h
new file mode 100644
--- /dev/null
+++ b/Engine/Plugins/Experimental/MLDeformer/MLDeformerFramework/Source/MLDeformerFrameworkEditor/Private/MLDeformerAssetActions.h
@@ -0,0 +1,23 @@
+// Copyright Epic Games, Inc. All Rights Reserved.
+
+#pragma once
+
+#include "CoreMinimal.h"
+#include "AssetTypeActions_Base.h"
+
+namespace UE::MLDeformer
+{
+	class FMLDeformerAssetActions
+		: public FAssetTypeActions_Base
+	{
+	public:
+		// IAssetTypeActions overrides.
+		virtual FText GetName() const override { return NSLOCTEXT("MLDeformerAssetActions", "Name", "ML Deformer"); }
+		virtual void GetActions(const TArray<UObject*>& InObjects, struct FToolMenuSection& Section) override;
+		virtual FColor GetTypeColor() const override { return FColor(255, 255, 0); }
+		virtual UClass* GetSupportedClass() const override;
+		virtual void OpenAssetEditor(const TArray<UObject*>& InObjects, TSharedPtr<class IToolkitHost> EditWithinLevelEditor = TSharedPtr<IToolkitHost>()) override;
+		virtual uint32 GetCategories() override { return EAssetTypeCategories::Animation; }
+		// ~END IAssetTypeActions overrides.
+	};
+}	// namespace UE::MLDeformer
