[MLDeformer] A complete redesign, turning this into a whole framework that can handle different models.

diff --git a/Engine/Plugins/Experimental/MLDeformer/MLDeformerFramework/Source/MLDeformerFrameworkEditor/Private/MLDeformerEditorModel.cpp b/Engine/Plugins/Experimental/MLDeformer/MLDeformerFramework/Source/MLDeformerFrameworkEditor/Private/MLDeformerEditorModel.cpp
new file mode 100644
--- /dev/null
+++ b/Engine/Plugins/Experimental/MLDeformer/MLDeformerFramework/Source/MLDeformerFrameworkEditor/Private/MLDeformerEditorModel.cpp
@@ -0,0 +1,1423 @@
+// Copyright Epic Games, Inc. All Rights Reserved.
+
+#include "MLDeformerEditorModel.h"
+#include "MLDeformerModule.h"
+#include "MLDeformerModel.h"
+#include "MLDeformerEditorStyle.h"
+#include "MLDeformerEditorModule.h"
+#include "MLDeformerModelRegistry.h"
+#include "MLDeformerEditorToolkit.h"
+#include "MLDeformerEditorActor.h"
+#include "MLDeformerComponent.h"
+#include "MLDeformerInputInfo.h"
+#include "MLDeformerSampler.h"
+#include "AnimationEditorPreviewActor.h"
+#include "AnimPreviewInstance.h"
+#include "Animation/MeshDeformer.h"
+#include "Animation/DebugSkelMeshComponent.h"
+#include "Animation/AnimSequence.h"
+#include "Rendering/SkeletalMeshModel.h"
+#include "Rendering/SkeletalMeshLODModel.h"
+#include "IPersonaPreviewScene.h"
+#include "IPersonaToolkit.h"
+#include "Modules/ModuleManager.h"
+#include "Engine/World.h"
+#include "GameFramework/Actor.h"
+#include "EditorViewportClient.h"
+#include "Components/TextRenderComponent.h"
+#include "Materials/Material.h"
+#include "GeometryCache.h"
+#include "GeometryCacheMeshData.h"
+#include "GeometryCacheTrack.h"
+#include "NeuralNetwork.h"
+
+#define LOCTEXT_NAMESPACE "MLDeformerEditorModel"
+
+namespace UE::MLDeformer
+{
+	FMLDeformerEditorModel::~FMLDeformerEditorModel()
+	{
+		DeleteEditorActors();
+
+		FMLDeformerEditorModule& EditorModule = FModuleManager::GetModuleChecked<FMLDeformerEditorModule>("MLDeformerFrameworkEditor");
+		EditorModule.GetModelRegistry().RemoveEditorModelInstance(this);
+	}
+
+	void FMLDeformerEditorModel::Init(const InitSettings& Settings)
+	{
+		check(Settings.Editor);
+		check(Settings.Model);
+
+		Editor = Settings.Editor;
+		Model = Settings.Model;
+
+		EditorInputInfo = Model->CreateInputInfo();
+		check(EditorInputInfo);
+
+		Sampler = CreateSampler();
+		check(Sampler);
+		Sampler->Init(this);
+	}
+
+	void FMLDeformerEditorModel::UpdateEditorInputInfo()
+	{
+		InitInputInfo(EditorInputInfo);
+	}
+
+	UWorld* FMLDeformerEditorModel::GetWorld() const
+	{
+		check(Editor);
+		return Editor->GetPersonaToolkit()->GetPreviewScene()->GetWorld();
+	}
+
+	FMLDeformerSampler* FMLDeformerEditorModel::CreateSampler() const
+	{
+		return new FMLDeformerSampler();
+	}
+
+	void FMLDeformerEditorModel::CreateTrainingLinearSkinnedActor(const TSharedRef<IPersonaPreviewScene>& InPersonaPreviewScene)
+	{
+		UWorld* World = InPersonaPreviewScene->GetWorld();
+
+		// Spawn the linear skinned actor.
+		FActorSpawnParameters BaseSpawnParams;
+		BaseSpawnParams.Name = MakeUniqueObjectName(World, AAnimationEditorPreviewActor::StaticClass(), "Train Base Actor");
+		AActor* Actor = World->SpawnActor<AAnimationEditorPreviewActor>(AAnimationEditorPreviewActor::StaticClass(), FTransform::Identity, BaseSpawnParams);
+		Actor->SetFlags(RF_Transient);
+
+		// Create the preview skeletal mesh component.
+		const FLinearColor BaseWireColor = FMLDeformerEditorStyle::Get().GetColor("MLDeformer.BaseMesh.WireframeColor");
+		UDebugSkelMeshComponent* SkelMeshComponent = NewObject<UDebugSkelMeshComponent>(Actor);
+		SkelMeshComponent->SetWireframeMeshOverlayColor(BaseWireColor);
+		SkelMeshComponent->SetVisibility(false);
+		SkelMeshComponent->MarkRenderStateDirty();
+
+		// Setup an apply an anim instance to the skeletal mesh component.
+		UAnimPreviewInstance* AnimPreviewInstance = NewObject<UAnimPreviewInstance>(SkelMeshComponent, TEXT("MLDeformerAnimInstance"));
+		SkelMeshComponent->PreviewInstance = AnimPreviewInstance;
+		AnimPreviewInstance->InitializeAnimation();
+
+		// Set the skeletal mesh on the component.
+		// NOTE: This must be done AFTER setting the AnimInstance so that the correct root anim node is loaded.
+		USkeletalMesh* Mesh = Model->GetSkeletalMesh();
+		SkelMeshComponent->SetSkeletalMesh(Mesh);
+
+		// Update the persona scene.
+		InPersonaPreviewScene->SetActor(Actor);
+		InPersonaPreviewScene->SetPreviewMeshComponent(SkelMeshComponent);
+		InPersonaPreviewScene->AddComponent(SkelMeshComponent, FTransform::Identity);
+		InPersonaPreviewScene->SetAdditionalMeshesSelectable(false);
+		InPersonaPreviewScene->SetPreviewMesh(Mesh);
+
+		// Register the editor actor.
+		const FLinearColor LabelColor = FMLDeformerEditorStyle::Get().GetColor("MLDeformer.BaseMesh.LabelColor");
+		FMLDeformerEditorActor::FConstructSettings Settings;
+		Settings.Actor = Actor;
+		Settings.TypeID = ActorID_Train_Base;
+		Settings.LabelColor = LabelColor;
+		Settings.LabelText = LOCTEXT("TrainBaseActorLabelText", "Training Base");
+		Settings.bIsTrainingActor = true;
+		FMLDeformerEditorActor* EditorActor = CreateEditorActor(Settings);
+		EditorActor->SetSkeletalMeshComponent(SkelMeshComponent);
+		EditorActor->SetCanDestroyActor(false);	// Crash will occur when destroying the Persona actor, so disable this.
+		EditorActors.Add(EditorActor);
+	}
+
+	void FMLDeformerEditorModel::CreateTestLinearSkinnedActor(UWorld* World)
+	{
+		FActorSpawnParameters SpawnParams;
+		SpawnParams.Name = MakeUniqueObjectName(World, AActor::StaticClass(), "Test Linear Skinned Actor");
+		AActor* Actor = World->SpawnActor<AActor>(AActor::StaticClass(), FTransform::Identity, SpawnParams);
+		Actor->SetFlags(RF_Transient);
+
+		const FLinearColor BaseWireColor = FMLDeformerEditorStyle::Get().GetColor("MLDeformer.BaseMesh.WireframeColor");
+		UDebugSkelMeshComponent* SkelMeshComponent = NewObject<UDebugSkelMeshComponent>(Actor);
+		SkelMeshComponent->SetWireframeMeshOverlayColor(BaseWireColor);
+		SkelMeshComponent->SetSkeletalMesh(Model->GetSkeletalMesh());
+		Actor->SetRootComponent(SkelMeshComponent);
+		SkelMeshComponent->RegisterComponent();
+		SkelMeshComponent->SetVisibility(false);
+		SkelMeshComponent->MarkRenderStateDirty();
+
+		// Register the editor actor.
+		const FLinearColor LabelColor = FMLDeformerEditorStyle::Get().GetColor("MLDeformer.BaseMesh.LabelColor");
+		FMLDeformerEditorActor::FConstructSettings Settings;
+		Settings.Actor = Actor;
+		Settings.TypeID = ActorID_Test_Base;
+		Settings.LabelColor = LabelColor;
+		Settings.LabelText = LOCTEXT("TestBaseActorLabelText", "Linear Skinned");
+		Settings.bIsTrainingActor = false;
+		FMLDeformerEditorActor* EditorActor = CreateEditorActor(Settings);
+		EditorActor->SetSkeletalMeshComponent(SkelMeshComponent);
+		EditorActors.Add(EditorActor);
+	}
+
+	void FMLDeformerEditorModel::CreateTestMLDeformedActor(UWorld* World)
+	{
+		// Create the ML deformed actor.
+		FActorSpawnParameters SpawnParams;
+		SpawnParams.Name = MakeUniqueObjectName(World, AActor::StaticClass(), "Test ML Deformed");
+		AActor* Actor = World->SpawnActor<AActor>(SpawnParams);
+		Actor->SetFlags(RF_Transient);
+
+		// Create the skeletal mesh component.
+		const FLinearColor MLDeformedWireColor = FMLDeformerEditorStyle::Get().GetColor("MLDeformer.MLDeformedMesh.WireframeColor");
+		UDebugSkelMeshComponent* SkelMeshComponent = NewObject<UDebugSkelMeshComponent>(Actor);
+		SkelMeshComponent->SetSkeletalMesh(Model->GetSkeletalMesh());
+		Actor->SetRootComponent(SkelMeshComponent);
+		SkelMeshComponent->RegisterComponent();
+		SkelMeshComponent->SetWireframeMeshOverlayColor(MLDeformedWireColor);
+		SkelMeshComponent->SetVisibility(false);
+		SkelMeshComponent->MarkRenderStateDirty();
+
+		// Create the ML Deformer component.
+		UMLDeformerAsset* DeformerAsset = Model->GetDeformerAsset();
+		UMLDeformerComponent* MLDeformerComponent = NewObject<UMLDeformerComponent>(Actor);
+		MLDeformerComponent->SetDeformerAsset(DeformerAsset);
+		MLDeformerComponent->RegisterComponent();
+		MLDeformerComponent->SetupComponent(DeformerAsset, SkelMeshComponent);
+
+		// Create the editor actor.
+		const FLinearColor LabelColor = FMLDeformerEditorStyle::Get().GetColor("MLDeformer.MLDeformedMesh.LabelColor");
+		FMLDeformerEditorActor::FConstructSettings Settings;
+		Settings.Actor = Actor;
+		Settings.TypeID = ActorID_Test_MLDeformed;
+		Settings.LabelColor = LabelColor;
+		Settings.LabelText = LOCTEXT("TestMLDeformedActorLabelText", "ML Deformed");
+		Settings.bIsTrainingActor = false;
+		FMLDeformerEditorActor* EditorActor = static_cast<FMLDeformerEditorActor*>(CreateEditorActor(Settings));
+		EditorActor->SetSkeletalMeshComponent(SkelMeshComponent);
+		EditorActor->SetMLDeformerComponent(MLDeformerComponent);
+		EditorActors.Add(EditorActor);
+	}
+
+	void FMLDeformerEditorModel::CreateActors(const TSharedRef<IPersonaPreviewScene>& InPersonaPreviewScene)
+	{
+		UWorld* World = InPersonaPreviewScene->GetWorld();
+		CreateTrainingLinearSkinnedActor(InPersonaPreviewScene);
+		CreateTestLinearSkinnedActor(World);
+		CreateTestMLDeformedActor(World);
+		CreateTrainingGroundTruthActor(World);
+		CreateTestGroundTruthActor(World);
+	}
+
+	void FMLDeformerEditorModel::ClearWorld()
+	{
+		TSharedRef<IPersonaPreviewScene> PreviewScene = Editor->GetPersonaToolkit()->GetPreviewScene();
+
+		UWorld* World = PreviewScene->GetWorld();
+		for (FMLDeformerEditorActor* EditorActor : EditorActors)
+		{
+			World->RemoveActor(EditorActor->GetActor(), true);
+			if (EditorActor->GetCanDestroyActor())
+			{
+				EditorActor->GetActor()->Destroy();
+			}
+		}
+
+		PreviewScene->SetPreviewAnimationAsset(nullptr);
+		PreviewScene->SetPreviewAnimationBlueprint(nullptr, nullptr);
+		PreviewScene->SetPreviewMesh(nullptr);
+		PreviewScene->SetPreviewMeshComponent(nullptr);
+		PreviewScene->SetActor(nullptr);
+		PreviewScene->ClearSelectedActor();
+
+		// Clear the editor actors.
+		DeleteEditorActors();
+	}
+
+	FMLDeformerEditorActor* FMLDeformerEditorModel::CreateEditorActor(const FMLDeformerEditorActor::FConstructSettings& Settings) const
+	{ 
+		return new FMLDeformerEditorActor(Settings);
+	}
+
+	void FMLDeformerEditorModel::DeleteEditorActors()
+	{
+		for (FMLDeformerEditorActor* Actor : EditorActors)
+		{
+			delete Actor;
+		}
+		EditorActors.Empty();
+	}
+
+	FMLDeformerEditorActor* FMLDeformerEditorModel::FindEditorActor(int32 ActorTypeID) const
+	{
+		for (FMLDeformerEditorActor* EditorActor : EditorActors)
+		{
+			if (EditorActor->GetTypeID() == ActorTypeID)
+			{
+				return EditorActor;
+			}
+		}
+
+		return nullptr;
+	}
+
+	void FMLDeformerEditorModel::Tick(FEditorViewportClient* ViewportClient, float DeltaTime)
+	{
+		UpdateActorTransforms();
+		UpdateLabels();
+		CheckTrainingDataFrameChanged();
+
+		// Update the ML Deformer component's weight.
+		FMLDeformerEditorActor* EditorActor = FindEditorActor(ActorID_Test_MLDeformed);
+		if (EditorActor)
+		{
+			UMLDeformerComponent* DeformerComponent = EditorActor->GetMLDeformerComponent();
+			if (DeformerComponent)
+			{		
+				DeformerComponent->SetWeight(Model->GetVizSettings()->GetWeight());
+			}
+		}
+	}
+
+	void FMLDeformerEditorModel::UpdateLabels()
+	{
+		const UMLDeformerVizSettings* VizSettings = Model->GetVizSettings();
+		const bool bDrawTrainingActors = (VizSettings->GetVisualizationMode() == EMLDeformerVizMode::TrainingData);
+		const bool bDrawTestActors = (VizSettings->GetVisualizationMode() == EMLDeformerVizMode::TestData);
+
+		for (FMLDeformerEditorActor* EditorActor : EditorActors)
+		{
+			UTextRenderComponent* LabelComponent = EditorActor->GetLabelComponent();
+			if (LabelComponent == nullptr)
+			{
+				continue;
+			}
+
+			if (VizSettings->GetDrawLabels())
+			{
+				const AActor* Actor = EditorActor->GetActor();
+				const FVector ActorLocation = Actor->GetActorLocation();
+				const FVector AlignmentOffset = Model->HasTrainingGroundTruth() ? Model->GetAlignmentTransform().GetTranslation() : FVector::ZeroVector;
+
+				LabelComponent->SetRelativeLocation(ActorLocation + FVector(0.0f, 0.0f, VizSettings->GetLabelHeight()) - AlignmentOffset);
+				LabelComponent->SetRelativeRotation(FQuat(FVector(0.0f, 0.0f, 1.0f), FMath::DegreesToRadians(90.0f)));
+				LabelComponent->SetRelativeScale3D(FVector(VizSettings->GetLabelScale()));
+
+				// Update visibility.
+				const bool bLabelIsVisible = (bDrawTrainingActors && EditorActor->IsTrainingActor()) || (bDrawTestActors && EditorActor->IsTestActor());
+				LabelComponent->SetVisibility(bLabelIsVisible);
+
+				// Handle test ground truth, disable its label when no ground truth asset was selected.
+				if (EditorActor->GetTypeID() == ActorID_Test_GroundTruth && !VizSettings->HasTestGroundTruth())
+				{
+					LabelComponent->SetVisibility(false);
+				}
+			}
+			else
+			{
+				LabelComponent->SetVisibility(false);
+			}
+		}
+	}
+
+	void FMLDeformerEditorModel::UpdateActorTransforms()
+	{
+		const FVector MeshSpacingVector = Model->GetVizSettings()->GetMeshSpacingOffsetVector();
+
+		// Update the linear skinned training model.
+		FMLDeformerEditorActor* EditorActor = FindEditorActor(ActorID_Train_Base);
+		if (EditorActor)
+		{
+			EditorActor->GetActor()->SetActorTransform(FTransform::Identity);
+		}
+
+		// Update training ground truth actor's transform.
+		EditorActor = FindEditorActor(ActorID_Train_GroundTruth);
+		if (EditorActor)
+		{
+			FTransform Transform = FTransform::Identity;
+			Transform.AddToTranslation(MeshSpacingVector);
+			EditorActor->GetActor()->SetActorTransform(Transform);
+		}
+
+		// Update the linear skinned test model.
+		EditorActor = FindEditorActor(ActorID_Test_Base);
+		if (EditorActor)
+		{
+			EditorActor->GetActor()->SetActorTransform(FTransform::Identity);
+		}
+
+		// Update the ML deformed test actor.
+		EditorActor = FindEditorActor(ActorID_Test_MLDeformed);
+		if (EditorActor)
+		{
+			FTransform Transform = FTransform::Identity;
+			Transform.AddToTranslation(MeshSpacingVector);
+			EditorActor->GetActor()->SetActorTransform(Transform);
+		}
+
+		// Update the ground truth test actor.
+		EditorActor = FindEditorActor(ActorID_Test_GroundTruth);
+		if (EditorActor)
+		{
+			FTransform Transform = FTransform::Identity;
+			Transform.AddToTranslation(MeshSpacingVector * 2.0f);
+			EditorActor->GetActor()->SetActorTransform(Transform);
+		}
+	}
+
+	void FMLDeformerEditorModel::UpdateActorVisibility()
+	{
+		UMLDeformerVizSettings* VizSettings = Model->GetVizSettings();
+		const bool bShowTrainingData = (VizSettings->GetVisualizationMode() == EMLDeformerVizMode::TrainingData);
+		const bool bShowTestData = (VizSettings->GetVisualizationMode() == EMLDeformerVizMode::TestData);
+
+		FMLDeformerEditorActor* EditorActor = FindEditorActor(ActorID_Train_Base);
+		if (EditorActor)
+		{
+			EditorActor->SetVisibility(bShowTrainingData);
+		}
+	
+		EditorActor = FindEditorActor(ActorID_Train_GroundTruth);
+		if (EditorActor)
+		{
+			EditorActor->SetVisibility(bShowTrainingData);
+		}
+
+		EditorActor = FindEditorActor(ActorID_Test_Base);
+		if (EditorActor)
+		{
+			EditorActor->SetVisibility(bShowTestData && VizSettings->GetDrawLinearSkinnedActor());
+		}
+
+		EditorActor = FindEditorActor(ActorID_Test_MLDeformed);
+		if (EditorActor)
+		{
+			EditorActor->SetVisibility(bShowTestData && VizSettings->GetDrawMLDeformedActor());
+		}
+
+		EditorActor = FindEditorActor(ActorID_Test_GroundTruth);
+		if (EditorActor)
+		{
+			EditorActor->SetVisibility(bShowTestData && VizSettings->GetDrawGroundTruthActor());
+		}
+	}
+
+	void FMLDeformerEditorModel::OnInputAssetsChanged()
+	{
+		// Force the training sequence to use Step interpolation and sample raw animation data.
+		UAnimSequence* TrainingAnimSequence = Model->GetAnimSequence();
+		if (TrainingAnimSequence)
+		{
+			TrainingAnimSequence->bUseRawDataOnly = true;
+			TrainingAnimSequence->Interpolation = EAnimInterpolationType::Step;
+		}
+
+		// Update the training base actor.
+		UDebugSkelMeshComponent* SkeletalMeshComponent = FindEditorActor(ActorID_Train_Base)->GetSkeletalMeshComponent();
+		check(SkeletalMeshComponent);
+		SkeletalMeshComponent->SetSkeletalMesh(Model->GetSkeletalMesh());
+		if (GetEditor()->GetPersonaToolkitPointer())
+		{
+			GetEditor()->GetPersonaToolkit()->GetPreviewScene()->SetPreviewMesh(Model->GetSkeletalMesh());
+		}
+		UMLDeformerVizSettings* VizSettings = Model->GetVizSettings();
+		UAnimSequence* TestAnimSequence = VizSettings->GetTestAnimSequence();
+		const float TestAnimSpeed = VizSettings->GetAnimPlaySpeed();
+		SkeletalMeshComponent->SetAnimationMode(EAnimationMode::AnimationSingleNode);
+		SkeletalMeshComponent->SetAnimation(TrainingAnimSequence);
+		SkeletalMeshComponent->SetPosition(0.0f);
+		SkeletalMeshComponent->SetPlayRate(TestAnimSpeed);
+		SkeletalMeshComponent->Play(false);
+
+		// Update the test base model.
+		SkeletalMeshComponent = FindEditorActor(ActorID_Test_Base)->GetSkeletalMeshComponent();
+		if (SkeletalMeshComponent)
+		{
+			SkeletalMeshComponent->SetSkeletalMesh(Model->GetSkeletalMesh());
+			SkeletalMeshComponent->SetAnimationMode(EAnimationMode::AnimationSingleNode);
+			SkeletalMeshComponent->SetAnimation(TestAnimSequence);
+			SkeletalMeshComponent->SetPosition(0.0f);
+			SkeletalMeshComponent->SetPlayRate(TestAnimSpeed);
+			SkeletalMeshComponent->Play(true);
+		}
+
+		// Update the test ML Deformed skeletal mesh component.
+		SkeletalMeshComponent = FindEditorActor(ActorID_Test_MLDeformed)->GetSkeletalMeshComponent();
+		if (SkeletalMeshComponent)
+		{
+			SkeletalMeshComponent->SetSkeletalMesh(Model->GetSkeletalMesh());
+			SkeletalMeshComponent->SetAnimationMode(EAnimationMode::AnimationSingleNode);
+			SkeletalMeshComponent->SetAnimation(TestAnimSequence);
+			SkeletalMeshComponent->SetPosition(0.0f);
+			SkeletalMeshComponent->SetPlayRate(TestAnimSpeed);
+			SkeletalMeshComponent->Play(true);
+		}
+
+		bIsDataNormalized = false;
+	}
+
+	void FMLDeformerEditorModel::OnPostInputAssetChanged()
+	{
+		CurrentTrainingFrame = -1;
+		Editor->UpdateTimeSliderRange();
+		Model->UpdateCachedNumVertices();
+		UpdateDeformerGraph();
+		RefreshMLDeformerComponents();
+		UpdateIsReadyForTrainingState();
+		SetTrainingFrame(0);
+		UpdateEditorInputInfo();
+		CheckTrainingDataFrameChanged();
+	}
+
+	void FMLDeformerEditorModel::OnTimeSliderScrubPositionChanged(double NewScrubTime, bool bIsScrubbing)
+	{
+		float PlayOffset = static_cast<float>(NewScrubTime);
+		const int32 TargetFrame = GetFrameAtTime(NewScrubTime);
+
+		UMLDeformerVizSettings* VizSettings = Model->GetVizSettings();
+		if (VizSettings->GetVisualizationMode() == EMLDeformerVizMode::TrainingData)
+		{
+			for (FMLDeformerEditorActor* EditorActor : EditorActors)
+			{
+				if (EditorActor->IsTrainingActor())
+				{
+					if (Model->HasTrainingGroundTruth())
+					{
+						PlayOffset = GetTimeAtFrame(TargetFrame);
+					}
+					EditorActor->SetPlayPosition(PlayOffset);
+				}
+			}
+			VizSettings->FrameNumber = TargetFrame;
+		}
+		else if (VizSettings->GetVisualizationMode() == EMLDeformerVizMode::TestData)
+		{
+			for (FMLDeformerEditorActor* EditorActor : EditorActors)
+			{
+				if (EditorActor->IsTestActor())
+				{
+					if (Model->GetVizSettings()->HasTestGroundTruth())
+					{
+						PlayOffset = GetTimeAtFrame(TargetFrame);
+					}
+					EditorActor->SetPlayPosition(PlayOffset);
+				}
+			}
+		}
+	}
+
+	void FMLDeformerEditorModel::SetTrainingFrame(int32 FrameNumber)
+	{
+		UMLDeformerVizSettings* VizSettings = Model->GetVizSettings();
+		VizSettings->FrameNumber = FrameNumber;
+		OnTimeSliderScrubPositionChanged(GetTimeAtFrame(FrameNumber), false);
+	}
+
+	void FMLDeformerEditorModel::HandleDefaultPropertyChanges(FPropertyChangedEvent& PropertyChangedEvent)
+	{
+		const FProperty* Property = PropertyChangedEvent.Property;
+		if (Property == nullptr)
+		{
+			return;
+		}
+
+		// When we change one of these properties below, restart animations etc.
+		if (Property->GetFName() == GET_MEMBER_NAME_CHECKED(UMLDeformerModel, SkeletalMesh))
+		{
+			TriggerInputAssetChanged();
+			Model->InitVertexMap();
+			Model->InitGPUData();
+			SetComputeGraphDataProviders();
+		}
+		else
+		if (Property->GetFName() == GET_MEMBER_NAME_CHECKED(UMLDeformerModel, AnimSequence) ||
+			Property->GetFName() == GET_MEMBER_NAME_CHECKED(UMLDeformerVizSettings, TestAnimSequence))
+		{
+			TriggerInputAssetChanged(true);
+		}
+		else if (Property->GetFName() == GET_MEMBER_NAME_CHECKED(UMLDeformerModel, AlignmentTransform))
+		{
+			if (PropertyChangedEvent.ChangeType == EPropertyChangeType::ValueSet)
+			{
+				SampleDeltas();
+			}
+		}
+		else if (Property->GetFName() == GET_MEMBER_NAME_CHECKED(UMLDeformerModel, MaxTrainingFrames))
+		{
+			TriggerInputAssetChanged();
+		}
+		else if (Property->GetFName() == GET_MEMBER_NAME_CHECKED(UMLDeformerModel, TrainingInputs))
+		{
+			if (PropertyChangedEvent.ChangeType == EPropertyChangeType::ValueSet)
+			{
+				UpdateEditorInputInfo();
+				UpdateIsReadyForTrainingState();
+				GetEditor()->GetModelDetailsView()->ForceRefresh();
+			}
+		}
+		else
+		if (Property->GetFName() == GET_MEMBER_NAME_CHECKED(UMLDeformerModel, BoneIncludeList) ||
+			Property->GetFName() == GET_MEMBER_NAME_CHECKED(UMLDeformerModel, CurveIncludeList))
+		{
+			UpdateEditorInputInfo();
+		}
+		else if (Property->GetFName() == GET_MEMBER_NAME_CHECKED(UMLDeformerVizSettings, AnimPlaySpeed))
+		{
+			UpdateTestAnimPlaySpeed();
+		}
+		else if (Property->GetFName() == GET_MEMBER_NAME_CHECKED(UMLDeformerVizSettings, FrameNumber))
+		{
+			ClampCurrentFrameIndex();
+			const int32 CurrentFrameNumber = Model->GetVizSettings()->GetFrameNumber();
+			OnTimeSliderScrubPositionChanged(GetTimeAtFrame(CurrentFrameNumber), false);
+		}
+		else if (Property->GetFName() == GET_MEMBER_NAME_CHECKED(UMLDeformerVizSettings, bShowHeatMap))
+		{
+			SetHeatMapMaterialEnabled(Model->GetVizSettings()->GetShowHeatMap());
+			UpdateDeformerGraph();
+		}
+		else
+		if (Property->GetFName() == GET_MEMBER_NAME_CHECKED(UMLDeformerVizSettings, bDrawLinearSkinnedActor) ||
+			Property->GetFName() == GET_MEMBER_NAME_CHECKED(UMLDeformerVizSettings, bDrawMLDeformedActor) ||
+			Property->GetFName() == GET_MEMBER_NAME_CHECKED(UMLDeformerVizSettings, bDrawGroundTruthActor))
+		{
+			UpdateActorVisibility();
+		}
+		else if (Property->GetFName() == GET_MEMBER_NAME_CHECKED(UMLDeformerVizSettings, bDrawDeltas))
+		{
+			SampleDeltas();
+		} 
+		else if (Property->GetFName() == GET_MEMBER_NAME_CHECKED(UMLDeformerVizSettings, DeformerGraph))
+		{
+			SetComputeGraphDataProviders();
+			GetEditor()->GetVizSettingsDetailsView()->ForceRefresh();
+		}
+	}
+
+	void FMLDeformerEditorModel::OnPlayButtonPressed()
+	{
+		UMLDeformerVizSettings* VizSettings = Model->GetVizSettings();
+		if (VizSettings->GetVisualizationMode() != EMLDeformerVizMode::TestData)
+		{
+			return;
+		}
+		
+		FMLDeformerEditorActor* BaseTestActor = FindEditorActor(ActorID_Test_Base);
+		const bool bMustPause = (BaseTestActor && BaseTestActor->GetSkeletalMeshComponent()) ? !BaseTestActor->GetSkeletalMeshComponent()->bPauseAnims : false;
+		for (FMLDeformerEditorActor* EditorActor : EditorActors)
+		{
+			if (EditorActor->IsTestActor())
+			{
+				EditorActor->Pause(bMustPause);
+			}
+		}
+	}
+
+	bool FMLDeformerEditorModel::IsPlayingAnim() const
+	{
+		UMLDeformerVizSettings* VizSettings = Model->GetVizSettings();
+		if (VizSettings->GetVisualizationMode() == EMLDeformerVizMode::TestData)
+		{
+			FMLDeformerEditorActor* EditorActor = FindEditorActor(ActorID_Test_Base);
+			if (EditorActor)
+			{
+				UDebugSkelMeshComponent* SkeletalMeshComponent = EditorActor->GetSkeletalMeshComponent();
+				if (SkeletalMeshComponent)
+				{
+					return !SkeletalMeshComponent->bPauseAnims;
+				}
+			}
+		}
+		return false;
+	}
+
+	double FMLDeformerEditorModel::CalcTimelinePosition() const
+	{
+		// Get the base editor actor, depending on the mode we're in.
+		FMLDeformerEditorActor* EditorActor = nullptr;
+		UMLDeformerVizSettings* VizSettings = Model->GetVizSettings();
+		if (VizSettings->GetVisualizationMode() == EMLDeformerVizMode::TestData)
+		{
+			EditorActor = FindEditorActor(ActorID_Test_Base);
+		}
+		else if (VizSettings->GetVisualizationMode() == EMLDeformerVizMode::TrainingData)
+		{
+			EditorActor = FindEditorActor(ActorID_Train_Base);
+		}
+
+		return EditorActor ? EditorActor->GetPlayPosition() : 0.0;
+	}
+
+	void FMLDeformerEditorModel::UpdateTestAnimPlaySpeed()
+	{
+		UMLDeformerVizSettings* VizSettings = Model->GetVizSettings();
+		const float Speed = VizSettings->GetAnimPlaySpeed();
+		for (FMLDeformerEditorActor* EditorActor : EditorActors)
+		{
+			if (EditorActor->IsTestActor()) // Only do test actors, no training actors.
+			{
+				EditorActor->SetPlaySpeed(Speed);
+			}
+		}
+	}
+
+	void FMLDeformerEditorModel::ClampCurrentFrameIndex()
+	{
+		UMLDeformerVizSettings* VizSettings = Model->GetVizSettings();
+		if (GetNumFrames() > 0)
+		{
+			VizSettings->FrameNumber = FMath::Min(VizSettings->FrameNumber, static_cast<uint32>(GetNumFrames() - 1));
+		}
+		else
+		{
+			VizSettings->FrameNumber = 0;
+		}
+	}
+
+	int32 FMLDeformerEditorModel::GetNumFramesForTraining() const 
+	{ 
+		return FMath::Min(GetNumFrames(), Model->GetTrainingFrameLimit());
+	}
+
+	FText FMLDeformerEditorModel::GetBaseAssetChangedErrorText() const
+	{
+		FText Result;
+		if (Model->SkeletalMesh && Model->GetInputInfo())
+		{
+			if (Model->NumBaseMeshVerts != Model->GetInputInfo()->GetNumBaseMeshVertices() &&
+				Model->NumBaseMeshVerts > 0 && Model->GetInputInfo()->GetNumBaseMeshVertices() > 0)
+			{
+				Result = FText::Format(LOCTEXT("BaseMeshMismatch", "Number of vertices in base mesh has changed from {0} to {1} vertices since this ML Deformer Asset was saved! {2}"),
+					Model->GetInputInfo()->GetNumBaseMeshVertices(),
+					Model->NumBaseMeshVerts,
+					IsTrained() ? LOCTEXT("BaseMeshMismatchNN", "Neural network needs to be retrained!") : FText());
+			}
+		}
+		return Result;
+	}
+
+	FText FMLDeformerEditorModel::GetVertexMapChangedErrorText() const
+	{
+		FText Result;
+		if (Model->SkeletalMesh)
+		{
+			bool bVertexMapMatch = true;
+			const FSkeletalMeshModel* ImportedModel = Model->SkeletalMesh->GetImportedModel();
+			if (ImportedModel)
+			{
+				const TArray<int32>& MeshVertexMap = ImportedModel->LODModels[0].MeshToImportVertexMap;
+				const TArray<int32>& ModelVertexMap = Model->GetVertexMap();
+				if (MeshVertexMap.Num() == ModelVertexMap.Num())
+				{
+					for (int32 Index = 0; Index < ModelVertexMap.Num(); ++Index)
+					{
+						if (MeshVertexMap[Index] != ModelVertexMap[Index])
+						{
+							bVertexMapMatch = false;
+							break;
+						}
+					}
+					
+					if (!bVertexMapMatch)
+					{
+						Result = FText(LOCTEXT("VertexMapMismatch", "The vertex order of your Skeletal Mesh changed."));
+					}
+				}
+			}
+		}
+		return Result;
+	}
+
+	FText FMLDeformerEditorModel::GetSkeletalMeshNeedsReimportErrorText() const
+	{
+		FText Result;
+
+		if (Model->SkeletalMesh)
+		{
+			FSkeletalMeshModel* ImportedModel = Model->SkeletalMesh->GetImportedModel();
+			check(ImportedModel);
+
+			const TArray<FSkelMeshImportedMeshInfo>& SkelMeshInfos = ImportedModel->LODModels[0].ImportedMeshInfos;
+			if (SkelMeshInfos.IsEmpty())
+			{
+				Result = LOCTEXT("SkelMeshNeedsReimport", "Skeletal Mesh asset needs to be reimported.");
+			}
+		}
+
+		return Result;
+	}
+
+	FText FMLDeformerEditorModel::GetInputsErrorText() const
+	{
+		if (Model->SkeletalMesh && GetEditorInputInfo()->IsEmpty())
+		{
+			switch (Model->TrainingInputs)
+			{
+				case EMLDeformerTrainingInputFilter::BonesOnly:			return FText(LOCTEXT("InputsEmptyBonesErrorText", "Your base mesh has no bones to train on."));
+				case EMLDeformerTrainingInputFilter::CurvesOnly:		return FText(LOCTEXT("InputsEmptyCurvesErrorText", "Your base mesh has no curves to train on."));
+				case EMLDeformerTrainingInputFilter::BonesAndCurves:	return FText(LOCTEXT("InputsEmptyBonesCurvesErrorText", "Your base mesh has no bones or curves to train on."));
+				default: return FText(LOCTEXT("InputsEmptyDefaultErrorText", "There are no inputs to train on. There are no bones, curves or other inputs we can use."));
+			}
+		}
+
+		return FText();
+	}
+
+	FText FMLDeformerEditorModel::GetIncompatibleSkeletonErrorText(USkeletalMesh* InSkelMesh, UAnimSequence* InAnimSeq) const
+	{
+		FText Result;
+		if (InSkelMesh && InAnimSeq)
+		{
+			if (!InSkelMesh->GetSkeleton()->IsCompatible(InAnimSeq->GetSkeleton()))
+			{
+				Result = LOCTEXT("SkeletonMismatch", "The base skeletal mesh and anim sequence use different skeletons. The animation might not play correctly.");
+			}
+		}
+		return Result;
+	}
+
+	FText FMLDeformerEditorModel::GetTargetAssetChangedErrorText() const
+	{
+		FText Result;
+
+		UMLDeformerInputInfo* InputInfo = Model->GetInputInfo();
+		if (Model->HasTrainingGroundTruth() && InputInfo)
+		{
+			if (Model->NumTargetMeshVerts != InputInfo->GetNumTargetMeshVertices() &&
+				Model->NumTargetMeshVerts > 0 && InputInfo->GetNumTargetMeshVertices() > 0)
+			{
+				Result = FText::Format(LOCTEXT("TargetMeshMismatch", "Number of vertices in target mesh has changed from {0} to {1} vertices since this ML Deformer Asset was saved! {2}"),
+					InputInfo->GetNumTargetMeshVertices(),
+					Model->NumTargetMeshVerts,
+					IsTrained() ? LOCTEXT("BaseMeshMismatchNN", "Model needs to be retrained!") : FText());
+			}
+		}
+
+		return Result;
+	}
+
+	void FMLDeformerEditorModel::InitInputInfo(UMLDeformerInputInfo* InputInfo)
+	{
+		InputInfo->Reset();
+
+		TArray<FString>& BoneNameStrings = InputInfo->GetBoneNameStrings();
+		TArray<FString>& CurveNameStrings = InputInfo->GetCurveNameStrings();
+		TArray<FName>& BoneNames = InputInfo->GetBoneNames();
+		TArray<FName>& CurveNames = InputInfo->GetCurveNames();
+		USkeletalMesh* SkeletalMesh = Model->GetSkeletalMesh();
+
+		BoneNames.Reset();
+		BoneNameStrings.Reset();
+		BoneNames.Reset();
+		CurveNames.Reset();
+
+		InputInfo->SetNumBaseVertices(Model->GetNumBaseMeshVerts());
+		InputInfo->SetNumTargetVertices(Model->GetNumTargetMeshVerts());
+
+		const bool bIncludeBones = (Model->GetTrainingInputs() == EMLDeformerTrainingInputFilter::BonesAndCurves || Model->GetTrainingInputs() == EMLDeformerTrainingInputFilter::BonesOnly);
+		const bool bIncludeCurves = (Model->GetTrainingInputs() == EMLDeformerTrainingInputFilter::BonesAndCurves || Model->GetTrainingInputs() == EMLDeformerTrainingInputFilter::CurvesOnly);
+		const USkeleton* Skeleton = Model->GetSkeletalMesh() ? Model->GetSkeletalMesh()->GetSkeleton() : nullptr;
+
+		// Handle bones.
+		if (bIncludeBones && SkeletalMesh)
+		{
+			// Include all the bones when no list was provided.
+			const FReferenceSkeleton& RefSkeleton = SkeletalMesh->GetRefSkeleton();
+			if (Model->GetBoneIncludeList().IsEmpty())
+			{
+				// Grab all bone names.
+				const int32 NumBones = RefSkeleton.GetNum();
+				BoneNameStrings.Reserve(NumBones);
+				for (int32 Index = 0; Index < NumBones; ++Index)
+				{
+					const FName BoneName = RefSkeleton.GetBoneName(Index);
+					BoneNameStrings.Add(BoneName.ToString());
+					BoneNames.Add(BoneName);
+				}
+			}
+			else // A list of bones to include was provided.
+			{
+				for (const FBoneReference& BoneReference : Model->GetBoneIncludeList())
+				{
+					if (BoneReference.BoneName.IsValid())
+					{
+						const FName BoneName = BoneReference.BoneName;
+						if (RefSkeleton.FindBoneIndex(BoneName) == INDEX_NONE)
+						{
+							UE_LOG(LogMLDeformer, Warning, TEXT("Bone '%s' in the bones include list doesn't exist, ignoring it."), *BoneName.ToString());
+							continue;
+						}
+
+						BoneNameStrings.Add(BoneName.ToString());
+						BoneNames.Add(BoneName);
+					}
+				}
+			}
+		}
+
+		// Handle curves.
+		if (bIncludeCurves && SkeletalMesh)
+		{
+			// Anim curves.
+			const FSmartNameMapping* SmartNameMapping = Skeleton ? Skeleton->GetSmartNameContainer(USkeleton::AnimCurveMappingName) : nullptr;
+			if (SmartNameMapping) // When there are curves.
+			{
+				// Include all curves when no list was provided.
+				if (Model->GetCurveIncludeList().IsEmpty())
+				{
+					SmartNameMapping->FillNameArray(CurveNames);
+					CurveNameStrings.Reserve(CurveNames.Num());
+					for (const FName Name : CurveNames)
+					{
+						CurveNameStrings.Add(Name.ToString());
+					}
+				}
+				else // A list of curve names was provided.
+				{
+					for (const FMLDeformerCurveReference& CurveReference : Model->GetCurveIncludeList())
+					{
+						if (CurveReference.CurveName.IsValid())
+						{
+							const FName CurveName = CurveReference.CurveName;
+							if (!SmartNameMapping->Exists(CurveName))
+							{
+								UE_LOG(LogMLDeformer, Warning, TEXT("Curve '%s' doesn't exist, ignoring it."), *CurveName.ToString());
+								continue;
+							}
+
+							CurveNameStrings.Add(CurveName.ToString());
+							CurveNames.Add(CurveName);
+						}
+					}
+				}
+			}
+		}
+	}
+
+	void FMLDeformerEditorModel::InitBoneIncludeListToAnimatedBonesOnly()
+	{
+		if (!Model->AnimSequence)
+		{
+			UE_LOG(LogMLDeformer, Warning, TEXT("Cannot initialize bone list as no Anim Sequence has been picked."));
+			return;
+		}
+
+		const UAnimDataModel* DataModel = Model->AnimSequence->GetDataModel();
+		if (!DataModel)
+		{
+			UE_LOG(LogMLDeformer, Warning, TEXT("Anim sequence has no data model."));
+			return;
+		}
+
+		if (!Model->SkeletalMesh)
+		{
+			UE_LOG(LogMLDeformer, Warning, TEXT("Skeletal Mesh has not been set."));
+			return;
+		}
+
+		USkeleton* Skeleton = Model->SkeletalMesh->GetSkeleton();
+		if (!Skeleton)
+		{
+			UE_LOG(LogMLDeformer, Warning, TEXT("Skeletal Mesh has no skeleton."));
+			return;
+		}
+
+		// Iterate over all bones that are both in the skeleton and the animation.
+		TArray<FName> AnimatedBoneList;
+		const FReferenceSkeleton& RefSkeleton = Skeleton->GetReferenceSkeleton();
+		const int32 NumBones = RefSkeleton.GetNum();
+		for (int32 Index = 0; Index < NumBones; ++Index)
+		{
+			const FName BoneName = RefSkeleton.GetBoneName(Index);
+			const int32 BoneTrackIndex = DataModel->GetBoneTrackIndexByName(BoneName);
+			if (BoneTrackIndex == INDEX_NONE)
+			{
+				continue;
+			}
+
+			// Check if there is actually animation data.
+			const FBoneAnimationTrack& BoneAnimTrack = DataModel->GetBoneTrackByIndex(BoneTrackIndex);
+			const TArray<FQuat4f>& Rotations = BoneAnimTrack.InternalTrackData.RotKeys;
+			bool bIsAnimated = false;
+			if (!Rotations.IsEmpty())
+			{
+				const FQuat4f FirstQuat = Rotations[0];
+				for (const FQuat4f KeyValue : Rotations)
+				{
+					if (!KeyValue.Equals(FirstQuat))
+					{
+						bIsAnimated = true;
+						break;
+					}
+				}
+
+				if (!bIsAnimated)
+				{
+					UE_LOG(LogMLDeformer, Display, TEXT("Bone '%s' has keyframes but isn't animated."), *BoneName.ToString());
+				}
+			}
+
+			if (bIsAnimated)
+			{
+				AnimatedBoneList.Add(BoneName);
+			}
+		}
+
+		// Init the bone include list using the animated bones.
+		if (!AnimatedBoneList.IsEmpty())
+		{
+			Model->BoneIncludeList.Empty();
+			Model->BoneIncludeList.Reserve(AnimatedBoneList.Num());
+			for (FName BoneName : AnimatedBoneList)
+			{
+				Model->BoneIncludeList.AddDefaulted();
+				FBoneReference& BoneRef = Model->BoneIncludeList.Last();
+				BoneRef.BoneName = BoneName;
+			}
+		}
+		else
+		{
+			Model->BoneIncludeList.Empty();
+			UE_LOG(LogMLDeformer, Warning, TEXT("There are no animated bone rotations in Anim Sequence '%s'."), *(Model->AnimSequence->GetName()));
+		}
+	}
+
+	void FMLDeformerEditorModel::InitCurveIncludeListToAnimatedCurvesOnly()
+	{
+		if (!Model->AnimSequence)
+		{
+			UE_LOG(LogMLDeformer, Warning, TEXT("Cannot initialize curve list as no Anim Sequence has been picked."));
+			return;
+		}
+
+		const UAnimDataModel* DataModel = Model->AnimSequence->GetDataModel();
+		if (!DataModel)
+		{
+			UE_LOG(LogMLDeformer, Warning, TEXT("Anim sequence has no data model."));
+			return;
+		}
+
+		if (!Model->SkeletalMesh)
+		{
+			UE_LOG(LogMLDeformer, Warning, TEXT("Skeletal Mesh has not been set."));
+			return;
+		}
+
+		USkeleton* Skeleton = Model->SkeletalMesh->GetSkeleton();
+		if (!Skeleton)
+		{
+			UE_LOG(LogMLDeformer, Warning, TEXT("Skeletal Mesh has no skeleton."));
+			return;
+		}
+
+		// Iterate over all curves that are both in the skeleton and the animation.
+		TArray<FName> AnimatedCurveList;
+		const FSmartNameMapping* Mapping = Skeleton->GetSmartNameContainer(USkeleton::AnimCurveMappingName);
+		if (Mapping)
+		{
+			TArray<FName> SkeletonCurveNames;
+			Mapping->FillNameArray(SkeletonCurveNames);
+			for (const FName& SkeletonCurveName : SkeletonCurveNames)
+			{
+				const TArray<FFloatCurve>& AnimCurves = DataModel->GetFloatCurves();
+				for (const FFloatCurve& AnimCurve : AnimCurves)
+				{
+					if (AnimCurve.Name.IsValid() && AnimCurve.Name.DisplayName == SkeletonCurveName)
+					{
+						TArray<float> TimeValues;
+						TArray<float> KeyValues;
+						AnimCurve.GetKeys(TimeValues, KeyValues);
+						if (KeyValues.Num() > 0)
+						{
+							const float FirstKeyValue = KeyValues[0];					
+							for (float CurKeyValue : KeyValues)
+							{
+								if (CurKeyValue != FirstKeyValue)
+								{
+									AnimatedCurveList.Add(SkeletonCurveName);
+									break;
+								}
+							}
+						}
+						break;
+					}
+				}
+			}
+		}
+
+		// Init the bone include list using the animated bones.
+		if (!AnimatedCurveList.IsEmpty())
+		{
+			Model->CurveIncludeList.Empty();
+			Model->CurveIncludeList.Reserve(AnimatedCurveList.Num());
+			for (FName CurveName : AnimatedCurveList)
+			{
+				Model->CurveIncludeList.AddDefaulted();
+				FMLDeformerCurveReference& CurveRef = Model->CurveIncludeList.Last();
+				CurveRef.CurveName = CurveName;
+			}
+		}
+		else
+		{
+			Model->CurveIncludeList.Empty();
+			UE_LOG(LogMLDeformer, Warning, TEXT("There are no animated curves in Anim Sequence '%s'."), *(Model->AnimSequence->GetName()));
+		}
+	}
+
+	void FMLDeformerEditorModel::Render(const FSceneView* View, FViewport* Viewport, FPrimitiveDrawInterface* PDI)
+	{
+		// Make sure that before we render anything, that our sampler is ready.
+		if (!Sampler->IsInitialized())
+		{
+			Sampler->Init(this); // This can still fail.
+			Sampler->SetVertexDeltaSpace(EVertexDeltaSpace::PostSkinning);
+			if (Sampler->IsInitialized()) // If we actually managed to initialize this frame.
+			{
+				SampleDeltas(); // Update the deltas.
+			}
+		}
+
+		const UMLDeformerVizSettings* VizSettings = Model->GetVizSettings();
+		if (VizSettings->GetVisualizationMode() == EMLDeformerVizMode::TrainingData)
+		{
+			// Draw the deltas for the current frame.
+			const TArray<float>& VertexDeltas = Sampler->GetVertexDeltas();
+			const TArray<FVector3f>& LinearSkinnedPositions = Sampler->GetSkinnedVertexPositions();
+			if (VizSettings->GetDrawVertexDeltas() && (VertexDeltas.Num() / 3) == LinearSkinnedPositions.Num())
+			{
+				const FLinearColor DeltasColor = FMLDeformerEditorStyle::Get().GetColor("MLDeformer.Deltas.Color");
+				const FLinearColor DebugVectorsColor = FMLDeformerEditorStyle::Get().GetColor("MLDeformer.DebugVectors.Color");
+				const FLinearColor DebugVectorsColor2 = FMLDeformerEditorStyle::Get().GetColor("MLDeformer.DebugVectors.Color2");
+				const uint8 DepthGroup = VizSettings->GetXRayDeltas() ? 100 : 0;
+				for (int32 Index = 0; Index < LinearSkinnedPositions.Num(); ++Index)
+				{
+					const int32 ArrayIndex = 3 * Index;
+					const FVector Delta(
+						VertexDeltas[ArrayIndex], 
+						VertexDeltas[ArrayIndex + 1], 
+						VertexDeltas[ArrayIndex + 2]);
+					const FVector VertexPos = (FVector)LinearSkinnedPositions[Index];
+					PDI->DrawLine(VertexPos, VertexPos + Delta, DeltasColor, DepthGroup);
+				}
+			}
+		}
+	}
+
+	void FMLDeformerEditorModel::SampleDeltas()
+	{
+		UMLDeformerVizSettings* VizSettings = Model->GetVizSettings();
+		ClampCurrentFrameIndex();
+
+		// If we have no Persona toolkit yet, then it is not yet safe to init the sampler.
+		if (Editor->GetPersonaToolkitPointer() != nullptr)
+		{
+			Sampler->Init(this);
+		}
+
+		if (Sampler->IsInitialized())
+		{
+			Sampler->SetVertexDeltaSpace(EVertexDeltaSpace::PostSkinning);
+			Sampler->Sample(VizSettings->FrameNumber);
+		}
+	}
+
+	void FMLDeformerEditorModel::CheckTrainingDataFrameChanged()
+	{
+		UMLDeformerVizSettings* VizSettings = Model->GetVizSettings();
+		ClampCurrentFrameIndex();
+		if (CurrentTrainingFrame != VizSettings->FrameNumber)
+		{
+			OnTrainingDataFrameChanged();
+		}
+	}
+
+	void FMLDeformerEditorModel::OnTrainingDataFrameChanged()
+	{
+		UMLDeformerVizSettings* VizSettings = Model->GetVizSettings();
+
+		// If the current frame number changed, re-sample the deltas if needed.
+		if (CurrentTrainingFrame != VizSettings->FrameNumber)
+		{
+			CurrentTrainingFrame = VizSettings->FrameNumber;
+			if (VizSettings->GetDrawVertexDeltas() && VizSettings->GetVisualizationMode() == EMLDeformerVizMode::TrainingData)
+			{
+				SampleDeltas();
+			}
+		}
+	}
+
+	void FMLDeformerEditorModel::RefreshMLDeformerComponents()
+	{
+		for (FMLDeformerEditorActor* EditorActor : EditorActors)
+		{
+			if (EditorActor && EditorActor->GetMLDeformerComponent())
+			{
+				USkeletalMeshComponent* SkelMeshComponent = EditorActor ? EditorActor->GetSkeletalMeshComponent() : nullptr;
+				UMLDeformerAsset* DeformerAsset = GetModel()->GetDeformerAsset();
+				EditorActor->GetMLDeformerComponent()->SetupComponent(DeformerAsset, SkelMeshComponent);
+				UE::MLDeformer::FMLDeformerModelInstance* ModelInstance = EditorActor->GetMLDeformerComponent()->GetModelInstance();
+				if (ModelInstance)
+				{
+					ModelInstance->UpdateCompatibilityStatus();
+				}
+			}
+		}
+	}
+
+	void FMLDeformerEditorModel::CreateHeatMapMaterial()
+	{
+		const FString HeatMapMaterialPath = GetHeatMapMaterialPath();
+		UObject* MaterialObject = StaticLoadObject(UMaterial::StaticClass(), nullptr, *HeatMapMaterialPath);
+		HeatMapMaterial = Cast<UMaterial>(MaterialObject);
+	}
+
+	void FMLDeformerEditorModel::CreateHeatMapDeformerGraph()
+	{
+		const FString HeatMapDeformerPath = GetHeatMapDeformerGraphPath();
+		UObject* DeformerObject = StaticLoadObject(UMeshDeformer::StaticClass(), nullptr, *HeatMapDeformerPath);
+		HeatMapDeformerGraph = Cast<UMeshDeformer>(DeformerObject);
+	}
+
+	void FMLDeformerEditorModel::CreateHeatMapAssets()
+	{
+		CreateHeatMapMaterial();
+		CreateHeatMapDeformerGraph();
+	}
+
+	void FMLDeformerEditorModel::SetHeatMapMaterialEnabled(bool bEnabled)
+	{
+		FMLDeformerEditorActor* EditorActor = FindEditorActor(ActorID_Test_MLDeformed);
+		if (EditorActor == nullptr)
+		{
+			return;
+		}
+
+		USkeletalMeshComponent* Component = EditorActor->GetSkeletalMeshComponent();
+		if (Component)
+		{
+			if (bEnabled)
+			{
+				for (int32 Index = 0; Index < Component->GetNumMaterials(); ++Index)
+				{
+					Component->SetMaterial(Index, HeatMapMaterial);
+				}
+			}
+			else
+			{
+				Component->EmptyOverrideMaterials();
+			}
+		}
+
+		UpdateDeformerGraph();
+	}
+
+	UMeshDeformer* FMLDeformerEditorModel::LoadDefaultDeformerGraph()
+	{
+		const FString GraphAssetPath = GetDefaultDeformerGraphAssetPath();
+		UObject* Object = StaticLoadObject(UMeshDeformer::StaticClass(), nullptr, *GraphAssetPath);
+		UMeshDeformer* DeformerGraph = Cast<UMeshDeformer>(Object);
+		if (DeformerGraph == nullptr)
+		{
+			UE_LOG(LogMLDeformer, Warning, TEXT("Failed to load default ML Deformer compute graph from: %s"), *GraphAssetPath);
+		}
+		else
+		{
+			UE_LOG(LogMLDeformer, Verbose, TEXT("Loaded default ML Deformer compute graph from: %s"), *GraphAssetPath);
+		}
+
+		return DeformerGraph;
+	}
+
+	void FMLDeformerEditorModel::SetDefaultDeformerGraphIfNeeded()
+	{
+		// Initialize the asset on the default plugin deformer graph.
+		UMLDeformerVizSettings* VizSettings = Model->GetVizSettings();
+		if (VizSettings && VizSettings->GetDeformerGraph() == nullptr)
+		{
+			UMeshDeformer* DefaultGraph = LoadDefaultDeformerGraph();
+			VizSettings->SetDeformerGraph(DefaultGraph);
+		}
+	}
+
+	FText FMLDeformerEditorModel::GetOverlayText() const
+	{
+		const FMLDeformerEditorActor* EditorActor = FindEditorActor(ActorID_Test_MLDeformed);
+		const UMLDeformerComponent* DeformerComponent = EditorActor ? EditorActor->GetMLDeformerComponent() : nullptr;
+		if (DeformerComponent)
+		{
+			const FMLDeformerModelInstance* ModelInstance = DeformerComponent->GetModelInstance();
+			if (ModelInstance &&
+				ModelInstance->GetSkeletalMeshComponent() && 
+				ModelInstance->GetSkeletalMeshComponent()->SkeletalMesh &&
+				!ModelInstance->IsCompatible() )
+			{
+				return FText::FromString( ModelInstance->GetCompatibilityErrorText() );
+			}
+		}
+	 
+		return FText::GetEmpty();
+	}
+
+	void FMLDeformerEditorModel::UpdateDeformerGraph()
+	{	
+		for (FMLDeformerEditorActor* EditorActor : EditorActors)
+		{
+			USkeletalMeshComponent* SkelMeshComponent = EditorActor->GetSkeletalMeshComponent();
+			if (EditorActor->GetMLDeformerComponent() == nullptr || SkelMeshComponent == nullptr)
+			{
+				continue;
+			}
+
+			if (SkelMeshComponent)
+			{
+				UMLDeformerVizSettings* VizSettings = Model->GetVizSettings();
+				UMeshDeformer* MeshDeformer = IsTrained() ? VizSettings->GetDeformerGraph() : nullptr;	
+				const bool bUseHeatMapDeformer = VizSettings->GetShowHeatMap();
+				SkelMeshComponent->SetMeshDeformer(bUseHeatMapDeformer ? HeatMapDeformerGraph.Get() : MeshDeformer);
+			}
+		}
+	}
+
+	void FMLDeformerEditorModel::SetComputeGraphDataProviders() const
+	{
+		const UMLDeformerVizSettings* VizSettings = Model->GetVizSettings();
+		UMeshDeformer* MeshDeformer = VizSettings->GetDeformerGraph();
+
+		FMLDeformerEditorActor* EditorActor = FindEditorActor(ActorID_Test_MLDeformed);
+		if (EditorActor)
+		{
+			UDebugSkelMeshComponent* SkelMeshComponent = EditorActor->GetSkeletalMeshComponent();
+			if (SkelMeshComponent)
+			{
+				SkelMeshComponent->SetMeshDeformer(MeshDeformer);
+			}
+		}
+	}
+
+	void FMLDeformerEditorModel::OnPostTraining(ETrainingResult TrainingResult)
+	{
+		for (FMLDeformerEditorActor* EditorActor : EditorActors)
+		{
+			if (EditorActor->GetMLDeformerComponent())
+			{
+				USkeletalMeshComponent* SkelMeshComponent = EditorActor ? EditorActor->GetSkeletalMeshComponent() : nullptr;
+				UMLDeformerAsset* DeformerAsset = Model->GetDeformerAsset();
+				EditorActor->GetMLDeformerComponent()->SetupComponent(DeformerAsset, SkelMeshComponent);
+			}
+		}
+
+		if (TrainingResult == ETrainingResult::Success || TrainingResult == ETrainingResult::Aborted)
+		{
+			// The InitAssets call resets the normalized flag, so set it back to true.
+			// This is safe as we finished training, which means we already normalized data.
+			// If we aborted we still have normalized the data. Only when we have AbortedCantUse then we canceled the normalization process.
+			bIsDataNormalized = true;
+		}
+	}
+
+	FMLDeformerEditorActor* FMLDeformerEditorModel::GetTimelineEditorActor() const
+	{
+		const UMLDeformerVizSettings* VizSettings = Model->GetVizSettings();
+		if (VizSettings->GetVisualizationMode() == EMLDeformerVizMode::TrainingData)
+		{
+			return FindEditorActor(ActorID_Train_GroundTruth);
+		}
+		else if (VizSettings->GetVisualizationMode() == EMLDeformerVizMode::TestData)
+		{
+			return FindEditorActor(ActorID_Test_GroundTruth);
+		}
+		return nullptr;
+	}
+
+	UNeuralNetwork* FMLDeformerEditorModel::LoadNeuralNetworkFromOnnx(const FString& Filename) const
+	{
+		FString OnnxFile = FPaths::ConvertRelativePathToFull(Filename);
+		if (FPaths::FileExists(OnnxFile))
+		{
+			UE_LOG(LogMLDeformer, Display, TEXT("Loading Onnx file '%s'..."), *OnnxFile);
+			UNeuralNetwork* Result = NewObject<UNeuralNetwork>(Model, UNeuralNetwork::StaticClass());		
+			if (Result->Load(OnnxFile))
+			{
+				Result->SetDeviceType(ENeuralDeviceType::GPU, ENeuralDeviceType::CPU, ENeuralDeviceType::GPU);	
+				SetComputeGraphDataProviders();
+				UE_LOG(LogMLDeformer, Display, TEXT("Successfully loaded Onnx file '%s'..."), *OnnxFile);
+				return Result;
+			}
+			else
+			{
+				UE_LOG(LogMLDeformer, Error, TEXT("Failed to load Onnx file '%s'"), *OnnxFile);
+			}
+		}
+		else
+		{
+			UE_LOG(LogMLDeformer, Error, TEXT("Onnx file '%s' does not exist!"), *OnnxFile);
+		}
+
+		return nullptr;
+	}
+
+	bool FMLDeformerEditorModel::IsEditorReadyForTrainingBasicChecks()
+	{
+		// Make sure we have picked required assets.
+		if (!Model->HasTrainingGroundTruth() ||
+			Model->GetAnimSequence() == nullptr ||
+			Model->GetSkeletalMesh() == nullptr ||
+			GetNumFrames() == 0)
+		{
+			return false;
+		}
+
+		// Make sure we have inputs.
+		UpdateEditorInputInfo();
+		if (GetEditorInputInfo()->IsEmpty())
+		{
+			return false;
+		}
+
+		return true;
+	}
+
+	bool FMLDeformerEditorModel::LoadTrainedNetwork() const
+	{
+		const FString OnnxFile = GetTrainedNetworkOnnxFile();
+		UNeuralNetwork* Network = LoadNeuralNetworkFromOnnx(OnnxFile);
+		if (Network)
+		{
+			Model->SetNeuralNetwork(Network);
+			return true;
+		}
+
+		return false;
+	}
+
+	bool FMLDeformerEditorModel::IsTrained() const
+	{
+		return (Model->GetNeuralNetwork() != nullptr);
+	}
+
+	void FMLDeformerEditorModel::TriggerInputAssetChanged(bool RefreshVizSettings)
+	{
+		OnInputAssetsChanged();
+		OnPostInputAssetChanged();
+		GetEditor()->GetModelDetailsView()->ForceRefresh();
+		if (RefreshVizSettings)
+		{
+			GetEditor()->GetVizSettingsDetailsView()->ForceRefresh();
+		}
+	}
+
+	FString FMLDeformerEditorModel::GetHeatMapMaterialPath() const
+	{
+		return FString(TEXT("/MLDeformerFramework/Materials/MLDeformerHeatMapMat.MLDeformerHeatMapMat"));
+	}
+
+	FString FMLDeformerEditorModel::GetHeatMapDeformerGraphPath() const
+	{
+		return FString(TEXT("/MLDeformerFramework/Deformers/MLDeformerHeatMapGraph.MLDeformerHeatMapGraph"));
+	}
+
+	FString FMLDeformerEditorModel::GetDefaultDeformerGraphAssetPath() const
+	{
+		return FString(TEXT("/MLDeformerFramework/Deformers/MLDeformerDefaultGraph.MLDeformerDefaultGraph"));
+	}
+
+	FString FMLDeformerEditorModel::GetTrainedNetworkOnnxFile() const
+	{
+		return FString(FPaths::ProjectIntermediateDir() + TEXT("MLDeformerNetwork.onnx"));
+	}
+}	// namespace UE::MLDeformer
+
+#undef LOCTEXT_NAMESPACE
diff --git a/Engine/Plugins/Experimental/MLDeformer/MLDeformerFramework/Source/MLDeformerFrameworkEditor/Private/MLDeformerEditorModule.cpp b/Engine/Plugins/Experimental/MLDeformer/MLDeformerFramework/Source/MLDeformerFrameworkEditor/Private/MLDeformerEditorModule.cpp
new file mode 100644
--- /dev/null
+++ b/Engine/Plugins/Experimental/MLDeformer/MLDeformerFramework/Source/MLDeformerFrameworkEditor/Private/MLDeformerEditorModule.cpp
@@ -0,0 +1,54 @@
+// Copyright Epic Games, Inc. All Rights Reserved.
+
+#include "MLDeformerEditorModule.h"
+#include "MLDeformerAssetActions.h"
+#include "MLDeformerEditorMode.h"
+#include "MLDeformerModule.h"
+#include "MLDeformerEditorModel.h"
+#include "MLDeformerModel.h"
+#include "MLDeformerCurveReferenceCustomization.h"
+#include "Modules/ModuleManager.h"
+#include "EditorModeRegistry.h"
+#include "PropertyEditorDelegates.h"
+
+#define LOCTEXT_NAMESPACE "MLDeformerEditorModule"
+
+IMPLEMENT_MODULE(UE::MLDeformer::FMLDeformerEditorModule, MLDeformerEditor)
+
+namespace UE::MLDeformer
+{
+	void FMLDeformerEditorModule::StartupModule()
+	{
+		MLDeformerAssetActions = MakeShareable(new FMLDeformerAssetActions);
+		FModuleManager::LoadModuleChecked<FAssetToolsModule>("AssetTools").Get().RegisterAssetTypeActions(MLDeformerAssetActions.ToSharedRef());
+		FEditorModeRegistry::Get().RegisterMode<FMLDeformerEditorMode>(FMLDeformerEditorMode::ModeName, LOCTEXT("MLDeformerEditorMode", "MLDeformer"), FSlateIcon(), false);
+
+		// Register object detail customizations.
+		FPropertyEditorModule& PropertyModule = FModuleManager::GetModuleChecked<FPropertyEditorModule>("PropertyEditor");
+		PropertyModule.RegisterCustomPropertyTypeLayout("CurveReference", FOnGetPropertyTypeCustomizationInstance::CreateStatic(&FMLDeformerCurveReferenceCustomization::MakeInstance) );
+		PropertyModule.NotifyCustomizationModuleChanged();
+	}
+
+	void FMLDeformerEditorModule::ShutdownModule()
+	{
+		FEditorModeRegistry::Get().UnregisterMode(FMLDeformerEditorMode::ModeName);
+		if (MLDeformerAssetActions.IsValid())
+		{
+			if (FModuleManager::Get().IsModuleLoaded("AssetTools"))
+			{
+				FModuleManager::GetModuleChecked<FAssetToolsModule>("AssetTools").Get().UnregisterAssetTypeActions(MLDeformerAssetActions.ToSharedRef());
+			}
+			MLDeformerAssetActions.Reset();
+		}
+
+		// Unregister object detail customizations.
+		if (FModuleManager::Get().IsModuleLoaded(TEXT("PropertyEditor")))
+		{
+			FPropertyEditorModule& PropertyModule = FModuleManager::GetModuleChecked<FPropertyEditorModule>(TEXT("PropertyEditor"));
+			PropertyModule.UnregisterCustomPropertyTypeLayout(TEXT("CurveReference"));
+			PropertyModule.NotifyCustomizationModuleChanged();
+		}
+	}
+}	// namespace UE::MLDeformer
+
+#undef LOCTEXT_NAMESPACE
diff --git a/Engine/Plugins/Experimental/MLDeformer/MLDeformerFramework/Source/MLDeformerFrameworkEditor/Private/MLDeformerModelDetails.cpp b/Engine/Plugins/Experimental/MLDeformer/MLDeformerFramework/Source/MLDeformerFrameworkEditor/Private/MLDeformerModelDetails.cpp
new file mode 100644
--- /dev/null
+++ b/Engine/Plugins/Experimental/MLDeformer/MLDeformerFramework/Source/MLDeformerFrameworkEditor/Private/MLDeformerModelDetails.cpp
@@ -0,0 +1,282 @@
+// Copyright Epic Games, Inc. All Rights Reserved.
+
+#include "MLDeformerModelDetails.h"
+#include "MLDeformerModule.h"
+#include "MLDeformerAsset.h"
+#include "MLDeformerModel.h"
+#include "MLDeformerEditorModel.h"
+#include "MLDeformerEditorModule.h"
+#include "NeuralNetwork.h"
+#include "GeometryCache.h"
+#include "GeometryCacheTrack.h"
+#include "Animation/AnimSequence.h"
+#include "Engine/SkeletalMesh.h"
+#include "PropertyCustomizationHelpers.h"
+#include "AssetRegistry/AssetData.h"
+#include "Modules/ModuleManager.h"
+#include "DetailLayoutBuilder.h"
+#include "DetailCategoryBuilder.h"
+#include "DetailWidgetRow.h"
+#include "IDetailsView.h"
+#include "IDetailGroup.h"
+#include "Widgets/DeclarativeSyntaxSupport.h"
+#include "Widgets/Layout/SBox.h"
+#include "Widgets/Input/SButton.h"
+#include "DetailLayoutBuilder.h"
+#include "SWarningOrErrorBox.h"
+
+#define LOCTEXT_NAMESPACE "MLDeformerModelDetails"
+
+namespace UE::MLDeformer
+{
+	bool FMLDeformerModelDetails::UpdateMemberPointers(const TArray<TWeakObjectPtr<UObject>>& Objects)
+	{
+		Model = nullptr;
+		EditorModel = nullptr;
+		if (Objects.Num() == 1 && Objects[0] != nullptr)
+		{
+			Model = static_cast<UMLDeformerModel*>(Objects[0].Get());
+
+			// Get the editor model for this runtime model.
+			FMLDeformerEditorModule& EditorModule = FModuleManager::GetModuleChecked<FMLDeformerEditorModule>("MLDeformerFrameworkEditor");
+			if (Model)
+			{
+				EditorModel = EditorModule.GetModelRegistry().GetEditorModel(Model);
+			}
+		}
+
+		return (Model != nullptr && EditorModel != nullptr);
+	}
+
+	void FMLDeformerModelDetails::CreateCategories()
+	{
+		BaseMeshCategoryBuilder = &DetailLayoutBuilder->EditCategory("Base Mesh", FText::GetEmpty(), ECategoryPriority::Important);
+		TargetMeshCategoryBuilder = &DetailLayoutBuilder->EditCategory("Target Mesh", FText::GetEmpty(), ECategoryPriority::Important);
+		InputOutputCategoryBuilder = &DetailLayoutBuilder->EditCategory("Inputs and Output", FText::GetEmpty(), ECategoryPriority::Important);
+		SettingsCategoryBuilder = &DetailLayoutBuilder->EditCategory("Training Settings", FText::GetEmpty(), ECategoryPriority::Important);
+	}
+
+	void FMLDeformerModelDetails::CustomizeDetails(IDetailLayoutBuilder& DetailBuilder)
+	{
+		DetailLayoutBuilder = &DetailBuilder;
+
+		// Update the pointers and check if they are valid.
+		TArray<TWeakObjectPtr<UObject>> Objects;
+		DetailBuilder.GetObjectsBeingCustomized(Objects);
+		if (!UpdateMemberPointers(Objects))
+		{
+			return;
+		}
+
+		CreateCategories();
+
+		// Base mesh details.
+		BaseMeshCategoryBuilder->AddProperty(GET_MEMBER_NAME_CHECKED(UMLDeformerModel, SkeletalMesh), UMLDeformerModel::StaticClass());
+
+		AddBaseMeshErrors();
+
+		// Check if the vertex counts of our asset have changed.
+		const FText ChangedErrorText = EditorModel->GetBaseAssetChangedErrorText();
+		FDetailWidgetRow& ChangedErrorRow = BaseMeshCategoryBuilder->AddCustomRow(FText::FromString("BaseMeshChangedError"))
+			.Visibility(!ChangedErrorText.IsEmpty() ? EVisibility::Visible : EVisibility::Collapsed)
+			.WholeRowContent()
+			[
+				SNew(SBox)
+				.Padding(FMargin(0.0f, 4.0f))
+				[
+					SNew(SWarningOrErrorBox)
+					.MessageStyle(EMessageStyle::Error)
+					.Message(ChangedErrorText)
+				]
+			];
+
+		// Check if our skeletal mesh's imported model contains a list of mesh infos. If not, we need to reimport it as it is an older asset.
+		const FText NeedsReimportErrorText = EditorModel->GetSkeletalMeshNeedsReimportErrorText();
+		FDetailWidgetRow& NeedsReimportErrorRow = BaseMeshCategoryBuilder->AddCustomRow(FText::FromString("BaseMeshNeedsReimportError"))
+			.Visibility(!NeedsReimportErrorText.IsEmpty() ? EVisibility::Visible : EVisibility::Collapsed)
+			.WholeRowContent()
+			[
+				SNew(SBox)
+				.Padding(FMargin(0.0f, 4.0f))
+				[
+					SNew(SWarningOrErrorBox)
+					.MessageStyle(EMessageStyle::Error)
+					.Message(NeedsReimportErrorText)
+				]
+			];
+
+		const FText VertexMapMisMatchErrorText = EditorModel->GetVertexMapChangedErrorText();
+		FDetailWidgetRow& VertexMapErrorRow = BaseMeshCategoryBuilder->AddCustomRow(FText::FromString("VertexMapError"))
+			.Visibility(!VertexMapMisMatchErrorText.IsEmpty() ? EVisibility::Visible : EVisibility::Collapsed)
+			.WholeRowContent()
+			[
+				SNew(SBox)
+				.Padding(FMargin(0.0f, 4.0f))
+				[
+					SNew(SWarningOrErrorBox)
+					.MessageStyle(EMessageStyle::Error)
+					.Message(VertexMapMisMatchErrorText)
+				]
+			];
+
+		// Animation sequence.
+		IDetailPropertyRow& AnimRow = BaseMeshCategoryBuilder->AddProperty(GET_MEMBER_NAME_CHECKED(UMLDeformerModel, AnimSequence), UMLDeformerModel::StaticClass());
+		AnimRow.CustomWidget()
+		.NameContent()
+		[
+			AnimRow.GetPropertyHandle()->CreatePropertyNameWidget()
+		]
+		.ValueContent()
+		[
+			SNew(SObjectPropertyEntryBox)
+			.PropertyHandle(AnimRow.GetPropertyHandle())
+			.AllowedClass(UAnimSequence::StaticClass())
+			.ObjectPath(Model ? Model->GetAnimSequence()->GetPathName() : FString())
+			.ThumbnailPool(DetailBuilder.GetThumbnailPool())
+			.OnShouldFilterAsset(
+				this, 
+				&FMLDeformerModelDetails::FilterAnimSequences, 
+				Model->GetSkeletalMesh() ? Model->GetSkeletalMesh()->GetSkeleton() : nullptr
+			)
+		];
+
+		AddAnimSequenceErrors();
+
+		const FText AnimErrorText = EditorModel->GetIncompatibleSkeletonErrorText(Model->GetSkeletalMesh(), Model->AnimSequence);
+		FDetailWidgetRow& AnimErrorRow = BaseMeshCategoryBuilder->AddCustomRow(FText::FromString("AnimSkeletonMisMatchError"))
+			.Visibility(!AnimErrorText.IsEmpty() ? EVisibility::Visible : EVisibility::Collapsed)
+			.WholeRowContent()
+			[
+				SNew(SBox)
+				.Padding(FMargin(0.0f, 4.0f))
+				[
+					SNew(SWarningOrErrorBox)
+					.MessageStyle(EMessageStyle::Warning)
+					.Message(AnimErrorText)
+				]
+			];
+
+
+		AddTargetMesh();
+
+		TargetMeshCategoryBuilder->AddProperty(GET_MEMBER_NAME_CHECKED(UMLDeformerModel, AlignmentTransform), UMLDeformerModel::StaticClass());
+
+		// Input and output.
+		InputOutputCategoryBuilder->AddProperty(GET_MEMBER_NAME_CHECKED(UMLDeformerModel, TrainingInputs), UMLDeformerModel::StaticClass());
+
+		AddTrainingInputErrors();
+
+		const FText ErrorText = EditorModel->GetInputsErrorText();
+		FDetailWidgetRow& ErrorRow = InputOutputCategoryBuilder->AddCustomRow(FText::FromString("InputsError"))
+			.Visibility(!ErrorText.IsEmpty() ? EVisibility::Visible : EVisibility::Collapsed)
+			.WholeRowContent()
+			[
+				SNew(SBox)
+				.Padding(FMargin(0.0f, 4.0f))
+				[
+					SNew(SWarningOrErrorBox)
+					.MessageStyle(EMessageStyle::Error)
+					.Message(ErrorText)
+				]
+			];
+
+		InputOutputCategoryBuilder->AddProperty(GET_MEMBER_NAME_CHECKED(UMLDeformerModel, MaxTrainingFrames), UMLDeformerModel::StaticClass());
+		InputOutputCategoryBuilder->AddProperty(GET_MEMBER_NAME_CHECKED(UMLDeformerModel, DeltaCutoffLength), UMLDeformerModel::StaticClass());
+
+		// Bone include list group.
+		IDetailGroup& BoneIncludeGroup = InputOutputCategoryBuilder->AddGroup("BoneIncludeGroup", LOCTEXT("BoneIncludeGroup", "Bones"), false, false);
+		BoneIncludeGroup.AddWidgetRow()
+			.ValueContent()
+			[
+				SNew(SButton)
+				.HAlign(HAlign_Center)
+				.VAlign(VAlign_Center)
+				.Text(LOCTEXT("AnimatedBonesButton", "Animated Bones Only"))
+				.OnClicked(FOnClicked::CreateSP(this, &FMLDeformerModelDetails::OnFilterAnimatedBonesOnly))
+				.IsEnabled_Lambda([this](){ return (Model->GetTrainingInputs() == EMLDeformerTrainingInputFilter::BonesAndCurves) || (Model->GetTrainingInputs() == EMLDeformerTrainingInputFilter::BonesOnly); })
+			];
+		BoneIncludeGroup.AddPropertyRow(DetailBuilder.GetProperty(GET_MEMBER_NAME_CHECKED(UMLDeformerModel, BoneIncludeList), UMLDeformerModel::StaticClass()));
+
+		AddBoneInputErrors();
+
+		// Curve include list group.
+		IDetailGroup& CurveIncludeGroup = InputOutputCategoryBuilder->AddGroup("CurveIncludeGroup", LOCTEXT("CurveIncludeGroup", "Curves"), false, false);
+		CurveIncludeGroup.AddWidgetRow()
+			.ValueContent()
+			[
+				SNew(SButton)
+				.HAlign(HAlign_Center)
+				.VAlign(VAlign_Center)
+				.Text(LOCTEXT("AnimatedCurvesButton", "Animated Curves Only"))
+				.OnClicked(FOnClicked::CreateSP(this, &FMLDeformerModelDetails::OnFilterAnimatedCurvesOnly))
+				.IsEnabled_Lambda([this](){ return (Model->GetTrainingInputs() == EMLDeformerTrainingInputFilter::BonesAndCurves) || (Model->GetTrainingInputs() == EMLDeformerTrainingInputFilter::CurvesOnly); })
+			];
+		CurveIncludeGroup.AddPropertyRow(DetailBuilder.GetProperty(GET_MEMBER_NAME_CHECKED(UMLDeformerModel, CurveIncludeList), UMLDeformerModel::StaticClass()));
+
+		AddCurveInputErrors();
+
+		SettingsCategoryBuilder->AddProperty(GET_MEMBER_NAME_CHECKED(UMLDeformerModel, MaxCacheSizeGB), UMLDeformerModel::StaticClass());
+
+		AddCacheSizeErrors();
+
+		// Show a warning when no neural network has been set.
+		{		
+			UNeuralNetwork* NeuralNetwork = Model->GetNeuralNetwork();
+			FDetailWidgetRow& NeuralNetErrorRow = SettingsCategoryBuilder->AddCustomRow(FText::FromString("NeuralNetError"))
+				.Visibility((NeuralNetwork == nullptr) ? EVisibility::Visible : EVisibility::Collapsed)
+				.WholeRowContent()
+				[
+					SNew(SBox)
+					.Padding(FMargin(0.0f, 4.0f))
+					[
+						SNew(SWarningOrErrorBox)
+						.MessageStyle(EMessageStyle::Warning)
+						.Message(FText::FromString("Model still needs to be trained."))
+					]
+				];
+
+			// Check if our network is compatible with the skeletal mesh.
+			if (Model->GetSkeletalMesh() && NeuralNetwork)
+			{
+				FDetailWidgetRow& NeuralNetIncompatibleErrorRow = SettingsCategoryBuilder->AddCustomRow(FText::FromString("NeuralNetIncompatibleError"))
+					.Visibility(!Model->GetInputInfo()->IsCompatible(Model->GetSkeletalMesh()) ? EVisibility::Visible : EVisibility::Collapsed)
+					.WholeRowContent()
+					[
+						SNew(SBox)
+						.Padding(FMargin(0.0f, 4.0f))
+						[
+							SNew(SWarningOrErrorBox)
+							.MessageStyle(EMessageStyle::Error)
+							.Message(FText::FromString("Trained neural network is incompatible with selected SkeletalMesh."))
+						]
+					];
+			}
+		}
+	}
+
+	bool FMLDeformerModelDetails::FilterAnimSequences(const FAssetData& AssetData, USkeleton* Skeleton)
+	{
+		if (Skeleton && Skeleton->IsCompatibleSkeletonByAssetData(AssetData))
+		{
+			return false;
+		}
+
+		return true;
+	}
+
+	FReply FMLDeformerModelDetails::OnFilterAnimatedBonesOnly() const
+	{
+		EditorModel->InitBoneIncludeListToAnimatedBonesOnly();
+		DetailLayoutBuilder->ForceRefreshDetails();
+		return FReply::Handled();
+	}
+
+	FReply FMLDeformerModelDetails::OnFilterAnimatedCurvesOnly() const
+	{
+		EditorModel->InitCurveIncludeListToAnimatedCurvesOnly();
+		DetailLayoutBuilder->ForceRefreshDetails();
+		return FReply::Handled();
+	}
+}	// namespace UE::MLDeformer
+
+#undef LOCTEXT_NAMESPACE
diff --git a/Engine/Plugins/Experimental/MLDeformer/MLDeformerFramework/Source/MLDeformerFrameworkEditor/Private/MLDeformerTimelineTabSummoner.cpp b/Engine/Plugins/Experimental/MLDeformer/MLDeformerFramework/Source/MLDeformerFrameworkEditor/Private/MLDeformerTimelineTabSummoner.cpp
new file mode 100644
--- /dev/null
+++ b/Engine/Plugins/Experimental/MLDeformer/MLDeformerFramework/Source/MLDeformerFrameworkEditor/Private/MLDeformerTimelineTabSummoner.cpp
@@ -0,0 +1,117 @@
+// Copyright Epic Games, Inc. All Rights Reserved.
+
+#include "MLDeformerTimelineTabSummoner.h"
+#include "MLDeformerEditorToolkit.h"
+#include "MLDeformerEditorStyle.h"
+#include "MLDeformerEditorModel.h"
+#include "MLDeformerModel.h"
+
+#include "IDocumentation.h"
+#include "SSimpleTimeSlider.h"
+#include "Widgets/SBoxPanel.h"
+#include "Widgets/Input/SButton.h"
+#include "Widgets/Images/SImage.h"
+
+#define LOCTEXT_NAMESPACE "MLDeformerTimelineTabSummoner"
+
+namespace UE::MLDeformer
+{
+	const FName FMLDeformerTimelineTabSummoner::TabID(TEXT("MLDeformerTimeline"));
+
+	FMLDeformerTimelineTabSummoner::FMLDeformerTimelineTabSummoner(const TSharedRef<FMLDeformerEditorToolkit>& InEditor)
+		: FWorkflowTabFactory(TabID, InEditor)
+		, Editor(&InEditor.Get())
+	{
+		bIsSingleton = true; // only allow a single instance of this tab
+		TabLabel = LOCTEXT("TimelineTabLabel", "Timeline");
+		TabIcon = FSlateIcon(FMLDeformerEditorStyle::Get().GetStyleSetName(), "MLDeformer.Timeline.TabIcon");
+		ViewMenuDescription = LOCTEXT("ViewMenu_Desc", "Timeline");
+		ViewMenuTooltip = LOCTEXT("ViewMenu_ToolTip", "Show the ML Deformer timeline.");
+	}
+
+	TSharedPtr<SToolTip> FMLDeformerTimelineTabSummoner::CreateTabToolTipWidget(const FWorkflowTabSpawnInfo& Info) const
+	{
+		return IDocumentation::Get()->CreateToolTip(
+			LOCTEXT("TimelineTooltip", "The timeline widget that controls the offset in the training or test anim sequence."), 
+			NULL, 
+			TEXT("Shared/Editors/Persona"), 
+			TEXT("MLDeformerTimeline_Window"));
+	}
+
+	TSharedRef<SWidget> FMLDeformerTimelineTabSummoner::CreateTabBody(const FWorkflowTabSpawnInfo& Info) const
+	{
+		// Create and setup the time slider widget.
+		TSharedRef<SSimpleTimeSlider> TimeSlider = SNew(SSimpleTimeSlider)
+			.ScrubPosition_Raw(Editor, &FMLDeformerEditorToolkit::CalcTimelinePosition)
+			.OnScrubPositionChanged_Raw(Editor, &FMLDeformerEditorToolkit::OnTimeSliderScrubPositionChanged);
+
+		Editor->SetTimeSlider(TimeSlider);
+
+		// Add the time slider.
+		TSharedRef<SHorizontalBox> Content = SNew(SHorizontalBox)
+			+SHorizontalBox::Slot()
+			.VAlign(EVerticalAlignment::VAlign_Top)
+			[
+				TimeSlider
+			]
+			+SHorizontalBox::Slot()
+			.AutoWidth()
+			.HAlign(EHorizontalAlignment::HAlign_Right)
+			.VAlign(EVerticalAlignment::VAlign_Top)
+			[
+				SNew(SBox)
+				.WidthOverride(25)
+				.HeightOverride(25)
+				.Padding(FMargin(0.0f))
+				.Visibility_Lambda
+				(
+					[this]()
+					{
+						using namespace UE::MLDeformer;
+						const FMLDeformerEditorModel* ActiveModel = Editor->GetActiveModel();
+						if (ActiveModel)
+						{
+							return (Editor->GetActiveModel()->GetModel()->GetVizSettings()->GetVisualizationMode() == EMLDeformerVizMode::TestData) ? EVisibility::Visible : EVisibility::Collapsed; 
+						}
+						return EVisibility::Collapsed;
+					} 
+				)
+				[
+					SNew(SButton)
+					.ButtonStyle(FAppStyle::Get(), "SimpleButton")
+					.HAlign(EHorizontalAlignment::HAlign_Center)
+					.VAlign(EVerticalAlignment::VAlign_Center)
+					.ToolTipText(LOCTEXT("PlayButtonToolTip", "Play or pause the test animation sequence"))
+					.ContentPadding(FMargin(0.0f))
+					.OnClicked_Lambda([this]()
+					{ 
+							using namespace UE::MLDeformer;
+							FMLDeformerEditorModel* EditorModel = Editor->GetActiveModel();
+							if (EditorModel)
+							{ 
+								EditorModel->OnPlayButtonPressed(); 
+							} 
+							return FReply::Handled(); 
+					})
+					[
+						SNew(SImage)
+						.Image_Lambda
+						(
+							[this]()
+							{
+								using namespace UE::MLDeformer;
+								const FMLDeformerEditorModel* EditorModel = Editor->GetActiveModel();
+								return (EditorModel != nullptr && EditorModel->IsPlayingAnim())
+									? FMLDeformerEditorStyle::Get().GetBrush("MLDeformer.Timeline.PauseIcon")
+									: FMLDeformerEditorStyle::Get().GetBrush("MLDeformer.Timeline.PlayIcon");
+							}
+						)
+					]
+				]
+			];
+
+		return Content;
+	}
+}	// namespace UE::MLDeformer
+
+#undef LOCTEXT_NAMESPACE 
diff --git a/Engine/Plugins/Experimental/MLDeformer/MLDeformerFramework/Source/MLDeformerFrameworkEditor/Private/MLDeformerTrainingModel.cpp b/Engine/Plugins/Experimental/MLDeformer/MLDeformerFramework/Source/MLDeformerFrameworkEditor/Private/MLDeformerTrainingModel.cpp
new file mode 100644
--- /dev/null
+++ b/Engine/Plugins/Experimental/MLDeformer/MLDeformerFramework/Source/MLDeformerFrameworkEditor/Private/MLDeformerTrainingModel.cpp
@@ -0,0 +1,66 @@
+// Copyright Epic Games, Inc. All Rights Reserved.
+
+#include "MLDeformerTrainingModel.h"
+#include "MLDeformerEditorModel.h"
+#include "MLDeformerModel.h"
+#include "MLDeformerInputInfo.h"
+#include "MLDeformerSampler.h"
+
+UMLDeformerModel* UMLDeformerTrainingModel::GetModel() const
+{
+	return EditorModel->GetModel();
+}
+
+void UMLDeformerTrainingModel::Init(UE::MLDeformer::FMLDeformerEditorModel* InEditorModel)
+{
+	EditorModel = InEditorModel;
+}
+
+int32 UMLDeformerTrainingModel::GetNumberSampleTransforms() const
+{
+	return GetModel()->GetInputInfo()->GetNumBones();
+}
+
+int32 UMLDeformerTrainingModel::GetNumberSampleCurves() const
+{
+	return GetModel()->GetInputInfo()->GetNumCurves();
+}
+
+int32 UMLDeformerTrainingModel::NumSamples() const
+{
+	return EditorModel->GetNumFramesForTraining();
+}
+
+int32 UMLDeformerTrainingModel::GetNumberSampleDeltas() const
+{
+	return GetModel()->GetNumBaseMeshVerts();
+}
+
+bool UMLDeformerTrainingModel::SetCurrentSampleIndex(int32 Index)
+{
+	return SampleFrame(Index);
+}
+
+bool UMLDeformerTrainingModel::SampleFrame(int32 Index)
+{
+	using namespace UE::MLDeformer;
+
+	// Make sure we have a valid frame number.
+	if (Index < 0 || Index >= EditorModel->GetNumFrames())
+	{
+		UE_LOG(LogMLDeformer, Warning, TEXT("Sample index must range from %d to %d, but a value of %d was provided."), 0, EditorModel->GetNumFrames()-1, Index);
+		return false;
+	}
+
+	// Sample the frame.
+	FMLDeformerSampler* Sampler = EditorModel->GetSampler();
+	Sampler->SetVertexDeltaSpace(EVertexDeltaSpace::PreSkinning);
+	Sampler->Sample(Index);
+
+	// Copy sampled values.
+	SampleDeltas = Sampler->GetVertexDeltas();
+	SampleBoneRotations = Sampler->GetBoneRotations();
+	SampleCurveValues = Sampler->GetCurveValues();
+
+	return true;
+}
diff --git a/Engine/Plugins/Experimental/MLDeformer/MLDeformerFramework/Source/MLDeformerFrameworkEditor/Private/MLDeformerVizSettingsDetails.cpp b/Engine/Plugins/Experimental/MLDeformer/MLDeformerFramework/Source/MLDeformerFrameworkEditor/Private/MLDeformerVizSettingsDetails.cpp
new file mode 100644
--- /dev/null
+++ b/Engine/Plugins/Experimental/MLDeformer/MLDeformerFramework/Source/MLDeformerFrameworkEditor/Private/MLDeformerVizSettingsDetails.cpp
@@ -0,0 +1,219 @@
+// Copyright Epic Games, Inc. All Rights Reserved.
+
+#include "MLDeformerVizSettingsDetails.h"
+#include "MLDeformerModule.h"
+#include "MLDeformerModel.h"
+#include "MLDeformerEditorModel.h"
+#include "MLDeformerEditorModule.h"
+#include "MLDeformerVizSettings.h"
+#include "DetailLayoutBuilder.h"
+#include "DetailCategoryBuilder.h"
+#include "DetailWidgetRow.h"
+#include "IDetailsView.h"
+#include "IDetailChildrenBuilder.h"
+#include "IDetailGroup.h"
+#include "Modules/ModuleManager.h"
+#include "Layout/Margin.h"
+#include "Widgets/Layout/SBorder.h"
+#include "Widgets/Text/STextBlock.h"
+#include "Widgets/Layout/SBox.h"
+#include "PropertyCustomizationHelpers.h"
+#include "AssetRegistry/AssetData.h"
+#include "PropertyHandle.h"
+#include "DetailLayoutBuilder.h"
+#include "SWarningOrErrorBox.h"
+#include "GeometryCache.h"
+#include "Animation/AnimSequence.h"
+#include "Animation/MeshDeformer.h"
+#include "Engine/SkeletalMesh.h"
+
+#define LOCTEXT_NAMESPACE "MLDeformerVizSettingsDetails"
+
+namespace UE::MLDeformer
+{
+	TSharedRef<IDetailCustomization> FMLDeformerVizSettingsDetails::MakeInstance()
+	{
+		return MakeShareable(new FMLDeformerVizSettingsDetails());
+	}
+
+	bool FMLDeformerVizSettingsDetails::UpdateMemberPointers(const TArray<TWeakObjectPtr<UObject>>& Objects)
+	{
+		Model = nullptr;
+		VizSettings = nullptr;
+		EditorModel = nullptr;
+
+		if (Objects.Num() == 1)
+		{
+			FMLDeformerEditorModule& EditorModule = FModuleManager::GetModuleChecked<FMLDeformerEditorModule>("MLDeformerFrameworkEditor");
+			VizSettings = Cast<UMLDeformerVizSettings>(Objects[0]);	
+			Model = VizSettings ? Cast<UMLDeformerModel>(VizSettings->GetOuter()) : nullptr;
+			EditorModel = Model ? EditorModule.GetModelRegistry().GetEditorModel(Model) : nullptr;
+		}
+
+		return (Model != nullptr && VizSettings != nullptr && EditorModel != nullptr);
+	}
+
+	void FMLDeformerVizSettingsDetails::CreateCategories()
+	{
+		SharedCategoryBuilder = &DetailLayoutBuilder->EditCategory("Shared Settings", FText::GetEmpty(), ECategoryPriority::Important);
+		TestAssetsCategory = &DetailLayoutBuilder->EditCategory("Test Assets", FText::GetEmpty(), ECategoryPriority::Important);
+		LiveSettingsCategory = &DetailLayoutBuilder->EditCategory("Live Settings", FText::GetEmpty(), ECategoryPriority::Important);
+		TrainingMeshesCategoryBuilder = &DetailLayoutBuilder->EditCategory("Training Meshes", FText::GetEmpty(), ECategoryPriority::Important);
+	}
+
+	void FMLDeformerVizSettingsDetails::CustomizeDetails(IDetailLayoutBuilder& DetailBuilder)
+	{
+		DetailLayoutBuilder = &DetailBuilder;
+
+		// Try update the member model, editormodel and viz settings pointers.
+		TArray<TWeakObjectPtr<UObject>> Objects;
+		DetailLayoutBuilder->GetObjectsBeingCustomized(Objects);
+		if (!UpdateMemberPointers(Objects))
+		{
+			return;
+		}
+
+		// Create the categories.
+		CreateCategories();
+
+		const bool bShowTrainingData = VizSettings ? (VizSettings->GetVisualizationMode() == EMLDeformerVizMode::TrainingData) : true;
+		const bool bShowTestData = VizSettings ? (VizSettings->GetVisualizationMode() == EMLDeformerVizMode::TestData) : true;
+
+		// Shared settings.
+		SharedCategoryBuilder->AddProperty(GET_MEMBER_NAME_CHECKED(UMLDeformerVizSettings, bDrawLabels), UMLDeformerVizSettings::StaticClass());
+		SharedCategoryBuilder->AddProperty(GET_MEMBER_NAME_CHECKED(UMLDeformerVizSettings, LabelHeight), UMLDeformerVizSettings::StaticClass());
+		SharedCategoryBuilder->AddProperty(GET_MEMBER_NAME_CHECKED(UMLDeformerVizSettings, LabelScale), UMLDeformerVizSettings::StaticClass());
+		SharedCategoryBuilder->AddProperty(GET_MEMBER_NAME_CHECKED(UMLDeformerVizSettings, MeshSpacing), UMLDeformerVizSettings::StaticClass());
+
+		// Test Assets.
+		TestAssetsCategory->SetCategoryVisibility(bShowTestData);
+	
+		IDetailPropertyRow& TestAnimRow = TestAssetsCategory->AddProperty(GET_MEMBER_NAME_CHECKED(UMLDeformerVizSettings, TestAnimSequence), UMLDeformerVizSettings::StaticClass());
+		TestAnimRow.CustomWidget()
+		.NameContent()
+		[
+			TestAnimRow.GetPropertyHandle()->CreatePropertyNameWidget()
+		]
+		.ValueContent()
+		[
+			SNew(SObjectPropertyEntryBox)
+			.PropertyHandle(TestAnimRow.GetPropertyHandle())
+			.AllowedClass(UAnimSequence::StaticClass())
+			.ObjectPath(VizSettings ? VizSettings->GetTestAnimSequence()->GetPathName() : FString())
+			.ThumbnailPool(DetailBuilder.GetThumbnailPool())
+			.OnShouldFilterAsset(
+				this, 
+				&FMLDeformerVizSettingsDetails::FilterAnimSequences, 
+				Model->GetSkeletalMesh() ? Model->GetSkeletalMesh()->GetSkeleton() : nullptr
+			)
+		];
+
+		AddTestSequenceErrors();
+
+		const FText AnimErrorText = EditorModel->GetIncompatibleSkeletonErrorText(Model->GetSkeletalMesh(), VizSettings->GetTestAnimSequence());
+		FDetailWidgetRow& AnimErrorRow = TestAssetsCategory->AddCustomRow(FText::FromString("AnimSkeletonMisMatchError"))
+			.Visibility(!AnimErrorText.IsEmpty() ? EVisibility::Visible : EVisibility::Collapsed)
+			.WholeRowContent()
+			[
+				SNew(SBox)
+				.Padding(FMargin(0.0f, 4.0f))
+				[
+					SNew(SWarningOrErrorBox)
+					.MessageStyle(EMessageStyle::Warning)
+					.Message(AnimErrorText)
+				]
+			];
+
+		FIsResetToDefaultVisible IsResetVisible = FIsResetToDefaultVisible::CreateSP(this, &FMLDeformerVizSettingsDetails::IsResetToDefaultDeformerGraphVisible);
+		FResetToDefaultHandler ResetHandler = FResetToDefaultHandler::CreateSP(this, &FMLDeformerVizSettingsDetails::OnResetToDefaultDeformerGraph);
+		FResetToDefaultOverride ResetOverride = FResetToDefaultOverride::Create(IsResetVisible, ResetHandler);
+		TestAssetsCategory->AddProperty(GET_MEMBER_NAME_CHECKED(UMLDeformerVizSettings, DeformerGraph), UMLDeformerVizSettings::StaticClass()).OverrideResetToDefault(ResetOverride);
+
+		AddDeformerGraphErrors();
+
+		// Show a warning when no deformer graph has been selected.
+		UObject* Graph = nullptr;
+		TSharedRef<IPropertyHandle> DeformerGraphProperty = DetailBuilder.GetProperty(GET_MEMBER_NAME_CHECKED(UMLDeformerVizSettings, DeformerGraph), UMLDeformerVizSettings::StaticClass());
+		if (DeformerGraphProperty->GetValue(Graph) == FPropertyAccess::Result::Success)
+		{
+			FDetailWidgetRow& GraphErrorRow = TestAssetsCategory->AddCustomRow(FText::FromString("GraphError"))
+				.Visibility((Graph == nullptr) ? EVisibility::Visible : EVisibility::Collapsed)
+				.WholeRowContent()
+				[
+					SNew(SBox)
+					.Padding(FMargin(0.0f, 4.0f))
+					[
+						SNew(SWarningOrErrorBox)
+						.MessageStyle(EMessageStyle::Warning)
+						.Message(FText::FromString("Please select a deformer graph.\nOtherwise only linear skinning is used."))
+					]
+				];
+		}
+
+		AddGroundTruth();
+
+		LiveSettingsCategory->SetCategoryVisibility(bShowTestData);
+		LiveSettingsCategory->AddProperty(GET_MEMBER_NAME_CHECKED(UMLDeformerVizSettings, Weight), UMLDeformerVizSettings::StaticClass());
+		LiveSettingsCategory->AddProperty(GET_MEMBER_NAME_CHECKED(UMLDeformerVizSettings, AnimPlaySpeed), UMLDeformerVizSettings::StaticClass());
+
+		IDetailGroup& HeatMapGroup = LiveSettingsCategory->AddGroup("HeatMap", LOCTEXT("HeatMap", "Heat Map"), false, true);
+		HeatMapGroup.AddPropertyRow(DetailBuilder.GetProperty(GET_MEMBER_NAME_CHECKED(UMLDeformerVizSettings, bShowHeatMap), UMLDeformerVizSettings::StaticClass()));
+		HeatMapGroup.AddPropertyRow(DetailBuilder.GetProperty(GET_MEMBER_NAME_CHECKED(UMLDeformerVizSettings, HeatMapMode), UMLDeformerVizSettings::StaticClass()));
+		HeatMapGroup.AddPropertyRow(DetailBuilder.GetProperty(GET_MEMBER_NAME_CHECKED(UMLDeformerVizSettings, HeatMapMax), UMLDeformerVizSettings::StaticClass()));
+		HeatMapGroup.AddPropertyRow(DetailBuilder.GetProperty(GET_MEMBER_NAME_CHECKED(UMLDeformerVizSettings, GroundTruthLerp), UMLDeformerVizSettings::StaticClass()));
+
+		IDetailGroup& VisGroup = LiveSettingsCategory->AddGroup("Visibility", LOCTEXT("VisibilityLabel", "Visibility"), false, true);
+		VisGroup.AddPropertyRow(DetailBuilder.GetProperty(GET_MEMBER_NAME_CHECKED(UMLDeformerVizSettings, bDrawLinearSkinnedActor), UMLDeformerVizSettings::StaticClass()));
+		VisGroup.AddPropertyRow(DetailBuilder.GetProperty(GET_MEMBER_NAME_CHECKED(UMLDeformerVizSettings, bDrawMLDeformedActor), UMLDeformerVizSettings::StaticClass()));
+		VisGroup.AddPropertyRow(DetailBuilder.GetProperty(GET_MEMBER_NAME_CHECKED(UMLDeformerVizSettings, bDrawGroundTruthActor), UMLDeformerVizSettings::StaticClass()))
+			.EditCondition(VizSettings->HasTestGroundTruth(), nullptr);
+
+		// Training data.
+		TrainingMeshesCategoryBuilder->SetCategoryVisibility(bShowTrainingData);
+		TrainingMeshesCategoryBuilder->AddProperty(GET_MEMBER_NAME_CHECKED(UMLDeformerVizSettings, FrameNumber), UMLDeformerVizSettings::StaticClass());
+		TrainingMeshesCategoryBuilder->AddProperty(GET_MEMBER_NAME_CHECKED(UMLDeformerVizSettings, bDrawDeltas), UMLDeformerVizSettings::StaticClass());
+		TrainingMeshesCategoryBuilder->AddProperty(GET_MEMBER_NAME_CHECKED(UMLDeformerVizSettings, bXRayDeltas), UMLDeformerVizSettings::StaticClass());
+	}
+
+	bool FMLDeformerVizSettingsDetails::FilterAnimSequences(const FAssetData& AssetData, USkeleton* Skeleton)
+	{
+		if (Skeleton && Skeleton->IsCompatibleSkeletonByAssetData(AssetData))
+		{
+			return false;
+		}
+
+		return true;
+	}
+
+	void FMLDeformerVizSettingsDetails::OnResetToDefaultDeformerGraph(TSharedPtr<IPropertyHandle> PropertyHandle)
+	{
+		if (EditorModel)
+		{
+			UMeshDeformer* MeshDeformer = EditorModel->LoadDefaultDeformerGraph();
+			PropertyHandle->SetValue(MeshDeformer);
+		}
+	}
+
+	bool FMLDeformerVizSettingsDetails::IsResetToDefaultDeformerGraphVisible(TSharedPtr<IPropertyHandle> PropertyHandle)
+	{
+		UObject* CurrentGraph = nullptr;
+		PropertyHandle->GetValue(CurrentGraph);
+		if (CurrentGraph == nullptr)
+		{
+			return true;
+		}
+
+		if (EditorModel)
+		{
+			// Check if we already assigned the default asset.
+			const FAssetData CurrentGraphAssetData(CurrentGraph);
+			const FString CurrentPath = CurrentGraphAssetData.ObjectPath.ToString();
+			const FString DefaultPath = EditorModel->GetDefaultDeformerGraphAssetPath();
+			return (DefaultPath != CurrentPath);
+		}
+
+		return false;
+	}
+}	// namespace UE::MLDeformer
+
+#undef LOCTEXT_NAMESPACE
diff --git a/Engine/Plugins/Experimental/MLDeformer/MLDeformerFramework/Source/MLDeformerFrameworkEditor/Private/SMLDeformerCurvePickerWidget.cpp b/Engine/Plugins/Experimental/MLDeformer/MLDeformerFramework/Source/MLDeformerFrameworkEditor/Private/SMLDeformerCurvePickerWidget.cpp
new file mode 100644
--- /dev/null
+++ b/Engine/Plugins/Experimental/MLDeformer/MLDeformerFramework/Source/MLDeformerFrameworkEditor/Private/SMLDeformerCurvePickerWidget.cpp
@@ -0,0 +1,196 @@
+// Copyright Epic Games, Inc. All Rights Reserved.
+
+#include "SMLDeformerCurvePickerWidget.h"
+#include "EditorStyleSet.h"
+#include "DetailLayoutBuilder.h"
+#include "Widgets/Input/SComboButton.h"
+#include "Widgets/Input/SSearchBox.h"
+#include "Widgets/Views/SListView.h"
+#include "Widgets/SBoxPanel.h"
+#include "Widgets/Text/STextBlock.h"
+#include "Widgets/Layout/SBox.h"
+#include "Animation/Skeleton.h"
+
+#define LOCTEXT_NAMESPACE "SCurvePickerWidget"
+
+namespace UE::MLDeformer
+{
+	void SMLDeformerCurvePickerWidget::Construct(const FArguments& InArgs)
+	{
+		OnCurveNamePicked = InArgs._OnCurveNamePicked;
+		OnGetSkeleton = InArgs._OnGetSkeleton;
+
+		if (OnGetSkeleton.IsBound())
+		{
+			Skeleton = OnGetSkeleton.Execute();
+		}
+
+		ChildSlot
+		[
+			SNew(SVerticalBox)
+			+SVerticalBox::Slot()
+			.AutoHeight()
+			[
+				SAssignNew(SearchBox, SSearchBox)
+					.HintText(LOCTEXT("SearchBoxHint", "Search Curves"))
+					.OnTextChanged(this, &SMLDeformerCurvePickerWidget::HandleFilterTextChanged)
+			]
+			+SVerticalBox::Slot()
+			.FillHeight(1.0f)
+			[
+				SAssignNew(NameListView, SListView<TSharedPtr<FString>>)
+				.SelectionMode(ESelectionMode::Single)
+				.ItemHeight(20.0f)
+				.ListItemsSource(&CurveNames)
+				.OnSelectionChanged(this, &SMLDeformerCurvePickerWidget::HandleSelectionChanged)
+				.OnGenerateRow(this, &SMLDeformerCurvePickerWidget::HandleGenerateRow)
+			]
+		];
+
+		RefreshListItems();
+	}
+
+	void SMLDeformerCurvePickerWidget::HandleSelectionChanged(TSharedPtr<FString> InItem, ESelectInfo::Type InSelectionType)
+	{
+		OnCurveNamePicked.ExecuteIfBound(**InItem.Get());
+	}
+
+	TSharedRef<ITableRow> SMLDeformerCurvePickerWidget::HandleGenerateRow(TSharedPtr<FString> InItem, const TSharedRef<STableViewBase>& InOwnerTable)
+	{
+		return 
+			SNew(STableRow<TSharedPtr<FString>>, InOwnerTable)
+			[
+				SNew(SBox)
+				.MinDesiredHeight(20.0f)
+				.VAlign(VAlign_Center)
+				[
+					SNew(STextBlock)
+					.Text(FText::FromString(*InItem.Get()))
+					.HighlightText_Lambda([this]() { return FText::FromString(FilterText); })
+				]
+			];
+	}
+
+	void SMLDeformerCurvePickerWidget::RefreshListItems()
+	{
+		CurveNames.Empty();
+		UniqueCurveNames.Empty();
+
+		// Get the curve names.
+		if (Skeleton)
+		{
+			const FSmartNameMapping* Mapping = Skeleton->GetSmartNameContainer(USkeleton::AnimCurveMappingName);
+			if (Mapping)
+			{
+				TArray<FName> SkeletonCurveNames;
+				Mapping->FillNameArray(SkeletonCurveNames);
+
+				UniqueCurveNames.Reserve(SkeletonCurveNames.Num());
+				for (const FName& CurveName : SkeletonCurveNames)
+				{
+					UniqueCurveNames.Add(CurveName.ToString());
+				}
+			}
+		}
+
+		FilterAvailableCurves();
+	}
+
+	void SMLDeformerCurvePickerWidget::FilterAvailableCurves()
+	{
+		CurveNames.Empty();
+
+		for (const FString& UniqueCurveName : UniqueCurveNames)
+		{
+			if (FilterText.IsEmpty() || UniqueCurveName.Contains(FilterText))
+			{
+				CurveNames.Add(MakeShared<FString>(UniqueCurveName));
+			}
+		}
+
+		NameListView->RequestListRefresh();
+	}
+
+	void SMLDeformerCurvePickerWidget::HandleFilterTextChanged(const FText& InFilterText)
+	{
+		FilterText = InFilterText.ToString();
+		FilterAvailableCurves();
+	}
+
+	/////////////////////////////////////////////////////
+
+	void SCurveSelectionWidget::Construct(const FArguments& InArgs)
+	{
+		OnCurveSelectionChanged = InArgs._OnCurveSelectionChanged;
+		OnGetSelectedCurve = InArgs._OnGetSelectedCurve;
+		OnGetSkeleton = InArgs._OnGetSkeleton;
+
+		SuppliedToolTip = InArgs._ToolTipText.Get();
+
+		this->ChildSlot
+		[
+			SAssignNew(CurvePickerButton, SComboButton)
+			.OnGetMenuContent(FOnGetContent::CreateSP(this, &SCurveSelectionWidget::CreateSkeletonWidgetMenu))
+			.ContentPadding(FMargin(4.0f, 2.0f, 4.0f, 2.0f))
+			.ButtonContent()
+			[
+				SNew(STextBlock)
+					.Text(this, &SCurveSelectionWidget::GetCurrentCurveName)
+					.Font(IDetailLayoutBuilder::GetDetailFont())
+					.ToolTipText(this, &SCurveSelectionWidget::GetFinalToolTip)
+			]
+		];
+	}
+
+	TSharedRef<SWidget> SCurveSelectionWidget::CreateSkeletonWidgetMenu()
+	{
+		bool bMultipleValues = false;
+
+		TSharedRef<SMLDeformerCurvePickerWidget> ListWidget = SNew(SMLDeformerCurvePickerWidget)
+			.OnCurveNamePicked(this, &SCurveSelectionWidget::OnSelectionChanged)
+			.OnGetSkeleton(OnGetSkeleton);
+
+		CurvePickerButton->SetMenuContentWidgetToFocus(ListWidget->GetFilterTextWidget());
+
+		return ListWidget;
+	}
+
+	void SCurveSelectionWidget::OnSelectionChanged(const FString& CurveName)
+	{
+		SelectedCurve = FText::FromString(CurveName);
+		if (OnCurveSelectionChanged.IsBound())
+		{
+			OnCurveSelectionChanged.Execute(CurveName);
+		}
+		CurvePickerButton->SetIsOpen(false);
+	}
+
+	FText SCurveSelectionWidget::GetCurrentCurveName() const
+	{
+		if (OnGetSelectedCurve.IsBound())
+		{
+			const FString Name = OnGetSelectedCurve.Execute();
+			return FText::FromString(Name);
+		}
+
+		return FText::GetEmpty();
+	}
+
+	FText SCurveSelectionWidget::GetFinalToolTip() const
+	{
+		return FText::Format(LOCTEXT("CurveClickToolTip", "Curve: {0}\n\n{1}"), GetCurrentCurveName(), SuppliedToolTip);
+	}
+
+	USkeleton* SCurveSelectionWidget::GetSkeleton() const
+	{
+		if (OnGetSkeleton.IsBound())
+		{
+			return OnGetSkeleton.Execute();
+		}
+
+		return nullptr;
+	}
+
+}	// namespace UE::MLDeformer
+
+#undef LOCTEXT_NAMESPACE
diff --git a/Engine/Plugins/Experimental/MLDeformer/MLDeformerFramework/Source/MLDeformerFrameworkEditor/Public/MLDeformerApplicationMode.h b/Engine/Plugins/Experimental/MLDeformer/MLDeformerFramework/Source/MLDeformerFrameworkEditor/Public/MLDeformerApplicationMode.h
new file mode 100644
--- /dev/null
+++ b/Engine/Plugins/Experimental/MLDeformer/MLDeformerFramework/Source/MLDeformerFrameworkEditor/Public/MLDeformerApplicationMode.h
@@ -0,0 +1,35 @@
+// Copyright Epic Games, Inc. All Rights Reserved.
+
+#pragma once
+
+#include "WorkflowOrientedApp/ApplicationMode.h"
+#include "WorkflowOrientedApp/WorkflowTabManager.h"
+
+class FWorkflowCentricApplication;
+class IPersonaPreviewScene;
+
+namespace UE::MLDeformer
+{
+	class FMLDeformerEditorToolkit;
+
+	class MLDEFORMERFRAMEWORKEDITOR_API FMLDeformerApplicationMode
+		: public FApplicationMode
+	{
+	public:
+		static FName ModeName;
+
+		FMLDeformerApplicationMode(TSharedRef<FWorkflowCentricApplication> InHostingApp, TSharedRef<IPersonaPreviewScene> InPreviewScene);
+
+		// FApplicationMode overrides.
+		virtual void RegisterTabFactories(TSharedPtr<FTabManager> InTabManager) override;
+		// ~END FApplicationMode overrides.
+
+	protected:
+		/** The hosting app. */
+		TWeakPtr<FMLDeformerEditorToolkit> EditorToolkit = nullptr;
+
+		/** The tab factories we support. */
+		FWorkflowAllowedTabSet TabFactories;
+	};
+
+}	// namespace UE::MLDeformer
\ No newline at end of file
diff --git a/Engine/Plugins/Experimental/MLDeformer/MLDeformerFramework/Source/MLDeformerFrameworkEditor/Public/MLDeformerCurveReferenceCustomization.h b/Engine/Plugins/Experimental/MLDeformer/MLDeformerFramework/Source/MLDeformerFrameworkEditor/Public/MLDeformerCurveReferenceCustomization.h
new file mode 100644
--- /dev/null
+++ b/Engine/Plugins/Experimental/MLDeformer/MLDeformerFramework/Source/MLDeformerFrameworkEditor/Public/MLDeformerCurveReferenceCustomization.h
@@ -0,0 +1,42 @@
+// Copyright Epic Games, Inc. All Rights Reserved.
+
+#pragma once
+
+#include "CoreMinimal.h"
+#include "Widgets/DeclarativeSyntaxSupport.h"
+#include "IPropertyTypeCustomization.h"
+#include "PropertyHandle.h"
+#include "DetailWidgetRow.h"
+
+class USkeleton;
+
+namespace UE::MLDeformer
+{
+	class MLDEFORMERFRAMEWORKEDITOR_API FMLDeformerCurveReferenceCustomization
+		: public IPropertyTypeCustomization
+	{
+	public:
+		static TSharedRef<IPropertyTypeCustomization> MakeInstance();
+
+		// IPropertyTypeCustomization overrides.
+		virtual void CustomizeHeader(TSharedRef<IPropertyHandle> StructPropertyHandle, class FDetailWidgetRow& HeaderRow, IPropertyTypeCustomizationUtils& StructCustomizationUtils) override;
+		virtual void CustomizeChildren(TSharedRef<IPropertyHandle> StructPropertyHandle, class IDetailChildrenBuilder& ChildBuilder, IPropertyTypeCustomizationUtils& StructCustomizationUtils) override {}
+		// ~END IPropertyTypeCustomization overrides.
+
+	protected:
+		void SetSkeleton(TSharedRef<IPropertyHandle> StructPropertyHandle);
+		virtual void SetPropertyHandle(TSharedRef<IPropertyHandle> StructPropertyHandle);
+		TSharedPtr<IPropertyHandle> FindStructMemberProperty(TSharedRef<IPropertyHandle> PropertyHandle, const FName& PropertyName);
+
+		/** Property to change after curve has been picked. */
+		TSharedPtr<IPropertyHandle> CurveNameProperty = nullptr;
+
+		/** The Skeleton we get the curves from. */
+		TObjectPtr<USkeleton> Skeleton = nullptr;
+
+	private:
+		virtual void OnCurveSelectionChanged(const FString& Name);
+		virtual FString OnGetSelectedCurve() const;
+		virtual USkeleton* OnGetSkeleton() const;
+	};
+}	// namespace UE::MLDeformer
diff --git a/Engine/Plugins/Experimental/MLDeformer/MLDeformerFramework/Source/MLDeformerFrameworkEditor/Public/MLDeformerEditorActor.h b/Engine/Plugins/Experimental/MLDeformer/MLDeformerFramework/Source/MLDeformerFrameworkEditor/Public/MLDeformerEditorActor.h
new file mode 100644
--- /dev/null
+++ b/Engine/Plugins/Experimental/MLDeformer/MLDeformerFramework/Source/MLDeformerFrameworkEditor/Public/MLDeformerEditorActor.h
@@ -0,0 +1,88 @@
+// Copyright Epic Games, Inc. All Rights Reserved.
+
+#pragma once
+
+#include "CoreMinimal.h"
+#include "GameFramework/Actor.h"
+
+class UDebugSkelMeshComponent;
+class UTextRenderComponent;
+class UMLDeformerComponent;
+
+namespace UE::MLDeformer
+{
+	enum : int32
+	{
+		ActorID_Train_Base,
+		ActorID_Train_GroundTruth,
+		ActorID_Test_Base,
+		ActorID_Test_MLDeformed,
+		ActorID_Test_GroundTruth
+	};
+
+	class MLDEFORMERFRAMEWORKEDITOR_API FMLDeformerEditorActor
+	{
+	public:
+		struct MLDEFORMERFRAMEWORKEDITOR_API FConstructSettings
+		{
+			AActor* Actor = nullptr;
+			int32 TypeID = -1;
+			FLinearColor LabelColor = FLinearColor(1.0f, 0.0f, 0.0f);
+			FText LabelText;
+			bool bIsTrainingActor = false;
+		};
+		FMLDeformerEditorActor(const FConstructSettings& Settings);
+		virtual ~FMLDeformerEditorActor() = default;
+
+		// Main methods you can override.
+		virtual void SetVisibility(bool bIsVisible);
+		virtual bool IsVisible() const;
+		virtual void SetPlayPosition(float TimeInSeconds, bool bAutoPause=true);
+		virtual float GetPlayPosition() const;
+		virtual void SetPlaySpeed(float PlaySpeed);
+		virtual FBox GetBoundingBox() const;
+		virtual void Pause(bool bPaused);
+
+		AActor* GetActor() const { return Actor; }
+		int32 GetTypeID() const { return TypeID; }
+
+		UDebugSkelMeshComponent* GetSkeletalMeshComponent() const { return SkeletalMeshComponent.Get(); }
+		UTextRenderComponent* GetLabelComponent() const { return LabelComponent.Get(); }
+		UMLDeformerComponent* GetMLDeformerComponent() const { return MLDeformerComponent.Get(); }
+		void SetSkeletalMeshComponent(UDebugSkelMeshComponent* SkelMeshComponent) { SkeletalMeshComponent = SkelMeshComponent; }
+		void SetMLDeformerComponent(UMLDeformerComponent* Component) { MLDeformerComponent = Component; }
+		bool IsTrainingActor() const { return bIsTrainingActor; }
+		bool IsTestActor() const { return !bIsTrainingActor; }
+		void SetCanDestroyActor(bool bCanDestroy) { bCanDestroyActor = bCanDestroy; }
+		bool GetCanDestroyActor() const { return bCanDestroyActor; }
+
+	protected:
+		UTextRenderComponent* CreateLabelComponent(AActor* InActor, FLinearColor Color, const FText& Text) const;
+
+	protected:
+		/**
+		 * The ID of the editor actor type.
+		 * This can be used to identify what actor we are dealing with, for example the base actor or ML Deformed one etc.
+		 */
+		int32 TypeID = -1;
+
+		/** The label component, which shows above the actor. */
+		TObjectPtr<UTextRenderComponent> LabelComponent = nullptr;
+
+		/** The actual actor pointer. */
+		TObjectPtr<AActor> Actor = nullptr;
+
+		/** The skeletal mesh component (can be nullptr). */
+		TObjectPtr<UDebugSkelMeshComponent> SkeletalMeshComponent = nullptr;
+
+		/** The ML Deformer component (can be nullptr). */
+		TObjectPtr<UMLDeformerComponent> MLDeformerComponent = nullptr;
+
+		/** Can the actor we created be destroyed during cleanup? */
+		bool bCanDestroyActor = true;
+
+		/** Is this actor used for training? */
+		bool bIsTrainingActor = true;
+	};
+
+}	// namespace UE::MLDeformer
diff --git a/Engine/Plugins/Experimental/MLDeformer/MLDeformerFramework/Source/MLDeformerFrameworkEditor/Public/MLDeformerEditorMode.h b/Engine/Plugins/Experimental/MLDeformer/MLDeformerFramework/Source/MLDeformerFrameworkEditor/Public/MLDeformerEditorMode.h
new file mode 100644
--- /dev/null
+++ b/Engine/Plugins/Experimental/MLDeformer/MLDeformerFramework/Source/MLDeformerFrameworkEditor/Public/MLDeformerEditorMode.h
@@ -0,0 +1,43 @@
+// Copyright Epic Games, Inc. All Rights Reserved.
+
+#pragma once
+
+#include "CoreMinimal.h"
+#include "IPersonaEditMode.h"
+
+class UMLDeformerVizSettings;
+
+namespace UE::MLDeformer
+{
+	class FMLDeformerEditorToolkit;
+	class FMLDeformerPreviewScene;
+
+	class MLDEFORMERFRAMEWORKEDITOR_API FMLDeformerEditorMode
+		: public IPersonaEditMode
+	{
+	public:
+		static FName ModeName;
+
+		void SetEditorToolkit(FMLDeformerEditorToolkit* InToolkit) { DeformerEditorToolkit = InToolkit; }
+
+		// IPersonaEditMode overrides.
+		virtual bool GetCameraTarget(FSphere& OutTarget) const override;
+		virtual class IPersonaPreviewScene& GetAnimPreviewScene() const override;
+		virtual void GetOnScreenDebugInfo(TArray<FText>& OutDebugInfo) const override {}
+		// ~END IPersonaEditMode overrides.
+
+		// FEdMode overrides.
+		virtual void Tick(FEditorViewportClient* ViewportClient, float DeltaTime) override;
+		virtual void Render(const FSceneView* View, FViewport* Viewport, FPrimitiveDrawInterface* PDI) override;
+		virtual void DrawHUD(FEditorViewportClient* ViewportClient, FViewport* Viewport, const FSceneView* View, FCanvas* Canvas) override;
+		virtual bool IsCompatibleWith(FEditorModeID OtherModeID) const override { return true; }
+		virtual bool AllowWidgetMove() override { return false; }
+		virtual bool ShouldDrawWidget() const override { return false; }
+		virtual bool UsesTransformWidget() const override { return false; }
+		virtual bool UsesTransformWidget(UE::Widget::EWidgetMode CheckMode) const override { return false; }
+		// ~END FEdMode overrides.
+
+	protected:
+		FMLDeformerEditorToolkit* DeformerEditorToolkit = nullptr;
+	};
+}	// namespace UE::MLDeformer
