MLDeformer: Rename deformer graphs to have DG_ suffix.
#rnx

diff --git a/Engine/Plugins/Experimental/MLDeformer/MLDeformerFramework/Source/MLDeformerFrameworkEditor/Private/MLDeformerEditorModel.cpp b/Engine/Plugins/Experimental/MLDeformer/MLDeformerFramework/Source/MLDeformerFrameworkEditor/Private/MLDeformerEditorModel.cpp
--- a/Engine/Plugins/Experimental/MLDeformer/MLDeformerFramework/Source/MLDeformerFrameworkEditor/Private/MLDeformerEditorModel.cpp
+++ b/Engine/Plugins/Experimental/MLDeformer/MLDeformerFramework/Source/MLDeformerFrameworkEditor/Private/MLDeformerEditorModel.cpp
@@ -1406,12 +1406,12 @@ namespace UE::MLDeformer
 
 	FString FMLDeformerEditorModel::GetHeatMapDeformerGraphPath() const
 	{
-		return FString(TEXT("/MLDeformerFramework/Deformers/MLDeformerHeatMapGraph.MLDeformerHeatMapGraph"));
+		return FString(TEXT("/MLDeformerFramework/Deformers/DG_MLDeformerModel_HeatMap.DG_MLDeformerModel_HeatMap"));
 	}
 
 	FString FMLDeformerEditorModel::GetDefaultDeformerGraphAssetPath() const
 	{
-		return FString(TEXT("/MLDeformerFramework/Deformers/MLDeformerDefaultGraph.MLDeformerDefaultGraph"));
+		return FString(TEXT("/MLDeformerFramework/Deformers/DG_MLDeformerModel.DG_MLDeformerModel"));
 	}
 
 	FString FMLDeformerEditorModel::GetTrainedNetworkOnnxFile() const
diff --git a/Engine/Plugins/Experimental/MLDeformer/VertexDeltaModel/Source/VertexDeltaModelEditor/Private/VertexDeltaEditorModel.cpp b/Engine/Plugins/Experimental/MLDeformer/VertexDeltaModel/Source/VertexDeltaModelEditor/Private/VertexDeltaEditorModel.cpp
--- a/Engine/Plugins/Experimental/MLDeformer/VertexDeltaModel/Source/VertexDeltaModelEditor/Private/VertexDeltaEditorModel.cpp
+++ b/Engine/Plugins/Experimental/MLDeformer/VertexDeltaModel/Source/VertexDeltaModelEditor/Private/VertexDeltaEditorModel.cpp
@@ -293,12 +293,12 @@ namespace UE::VertexDeltaModel
 
 	FString FVertexDeltaEditorModel::GetDefaultDeformerGraphAssetPath() const
 	{
-		return FString(TEXT("/VertexDeltaModel/Deformers/VertexDeltaModelGraph.VertexDeltaModelGraph"));
+		return FString(TEXT("/VertexDeltaModel/Deformers/DG_VertexDeltaModel.DG_VertexDeltaModel"));
 	}
 
 	FString FVertexDeltaEditorModel::GetHeatMapDeformerGraphPath() const
 	{
-		return FString(TEXT("/VertexDeltaModel/Deformers/VertexDeltaHeatMapGraph.VertexDeltaHeatMapGraph"));
+		return FString(TEXT("/VertexDeltaModel/Deformers/DG_VertexDeltaModel_HeatMap.DG_VertexDeltaModel_HeatMap"));
 	}
 }	// namespace UE::VertexDeltaModel
 
