[MLDeformer] Bug fixes:
- Alignment transform not applied to actor for ground truths, label positions could be incorrect.
- Shrinkage fields not getting hidden.
- Few small refactors.
- Crash when switching models that are trained.
- Crash when closing, sometimes.
- Fix where the heatmap didn't show correctly after changing some assets.

diff --git a/Engine/Plugins/Experimental/MLDeformer/LegacyVertexDeltaModel/Source/LegacyVertexDeltaModelEditor/Private/LegacyVertexDeltaModelDetails.cpp b/Engine/Plugins/Experimental/MLDeformer/LegacyVertexDeltaModel/Source/LegacyVertexDeltaModelEditor/Private/LegacyVertexDeltaModelDetails.cpp
--- a/Engine/Plugins/Experimental/MLDeformer/LegacyVertexDeltaModel/Source/LegacyVertexDeltaModelEditor/Private/LegacyVertexDeltaModelDetails.cpp
+++ b/Engine/Plugins/Experimental/MLDeformer/LegacyVertexDeltaModel/Source/LegacyVertexDeltaModelEditor/Private/LegacyVertexDeltaModelDetails.cpp
@@ -59,9 +59,9 @@ namespace UE::LegacyVertexDeltaModel
 		FMLDeformerModelDetails::CustomizeDetails(DetailBuilder);
 
 		// Training settings.
-		SettingsCategoryBuilder->AddProperty(GET_MEMBER_NAME_CHECKED(ULegacyVertexDeltaModel, NumHiddenLayers), UMLDeformerModel::StaticClass());
-		SettingsCategoryBuilder->AddProperty(GET_MEMBER_NAME_CHECKED(ULegacyVertexDeltaModel, NumNeuronsPerLayer), UMLDeformerModel::StaticClass());
-		SettingsCategoryBuilder->AddProperty(GET_MEMBER_NAME_CHECKED(ULegacyVertexDeltaModel, Epochs), UMLDeformerModel::StaticClass());
+		SettingsCategoryBuilder->AddProperty(GET_MEMBER_NAME_CHECKED(ULegacyVertexDeltaModel, NumHiddenLayers));
+		SettingsCategoryBuilder->AddProperty(GET_MEMBER_NAME_CHECKED(ULegacyVertexDeltaModel, NumNeuronsPerLayer));
+		SettingsCategoryBuilder->AddProperty(GET_MEMBER_NAME_CHECKED(ULegacyVertexDeltaModel, Epochs));
 
 		// Check whether shrinkage settings should be visible or not.
 		auto IsShrinkageVisible = [this]()
@@ -70,15 +70,15 @@ namespace UE::LegacyVertexDeltaModel
 		};
 	 
 		// Advanced settings.
-		SettingsCategoryBuilder->AddProperty(GET_MEMBER_NAME_CHECKED(ULegacyVertexDeltaModel, BatchSize), UMLDeformerModel::StaticClass());
-		SettingsCategoryBuilder->AddProperty(GET_MEMBER_NAME_CHECKED(ULegacyVertexDeltaModel, LearningRate), UMLDeformerModel::StaticClass());
-		SettingsCategoryBuilder->AddProperty(GET_MEMBER_NAME_CHECKED(ULegacyVertexDeltaModel, ActivationFunction), UMLDeformerModel::StaticClass());
-		SettingsCategoryBuilder->AddProperty(GET_MEMBER_NAME_CHECKED(ULegacyVertexDeltaModel, LossFunction), UMLDeformerModel::StaticClass());
-		SettingsCategoryBuilder->AddProperty(GET_MEMBER_NAME_CHECKED(ULegacyVertexDeltaModel, ShrinkageSpeed), UMLDeformerModel::StaticClass())
+		SettingsCategoryBuilder->AddProperty(GET_MEMBER_NAME_CHECKED(ULegacyVertexDeltaModel, BatchSize));
+		SettingsCategoryBuilder->AddProperty(GET_MEMBER_NAME_CHECKED(ULegacyVertexDeltaModel, LearningRate));
+		SettingsCategoryBuilder->AddProperty(GET_MEMBER_NAME_CHECKED(ULegacyVertexDeltaModel, ActivationFunction));
+		SettingsCategoryBuilder->AddProperty(GET_MEMBER_NAME_CHECKED(ULegacyVertexDeltaModel, LossFunction));
+		SettingsCategoryBuilder->AddProperty(GET_MEMBER_NAME_CHECKED(ULegacyVertexDeltaModel, ShrinkageSpeed))
 			.Visibility(TAttribute<EVisibility>::CreateLambda(IsShrinkageVisible));
-		SettingsCategoryBuilder->AddProperty(GET_MEMBER_NAME_CHECKED(ULegacyVertexDeltaModel, ShrinkageThreshold), UMLDeformerModel::StaticClass())
+		SettingsCategoryBuilder->AddProperty(GET_MEMBER_NAME_CHECKED(ULegacyVertexDeltaModel, ShrinkageThreshold))
 			.Visibility(TAttribute<EVisibility>::CreateLambda(IsShrinkageVisible));
-		SettingsCategoryBuilder->AddProperty(GET_MEMBER_NAME_CHECKED(ULegacyVertexDeltaModel, MaxCacheSizeGB), UMLDeformerModel::StaticClass());
+		SettingsCategoryBuilder->AddProperty(GET_MEMBER_NAME_CHECKED(ULegacyVertexDeltaModel, MaxCacheSizeGB));
 	}
 
 	void FLegacyVertexDeltaModelDetails::AddBaseMeshErrors()
@@ -118,7 +118,7 @@ namespace UE::LegacyVertexDeltaModel
 
 	void FLegacyVertexDeltaModelDetails::AddTargetMesh()
 	{
-		TargetMeshCategoryBuilder->AddProperty(GET_MEMBER_NAME_CHECKED(ULegacyVertexDeltaModel, GeometryCache), UMLDeformerModel::StaticClass());
+		TargetMeshCategoryBuilder->AddProperty(GET_MEMBER_NAME_CHECKED(ULegacyVertexDeltaModel, GeometryCache));
 
 		const FText TargetMeshErrorText = GetGeomCacheErrorText(VertexModel->GetSkeletalMesh(), VertexModel->GetGeometryCache());
 		TargetMeshCategoryBuilder->AddCustomRow(FText::FromString("TargetMeshError"))
diff --git a/Engine/Plugins/Experimental/MLDeformer/MLDeformerFramework/Source/MLDeformerFramework/Private/MLDeformerComponent.cpp b/Engine/Plugins/Experimental/MLDeformer/MLDeformerFramework/Source/MLDeformerFramework/Private/MLDeformerComponent.cpp
--- a/Engine/Plugins/Experimental/MLDeformer/MLDeformerFramework/Source/MLDeformerFramework/Private/MLDeformerComponent.cpp
+++ b/Engine/Plugins/Experimental/MLDeformer/MLDeformerFramework/Source/MLDeformerFramework/Private/MLDeformerComponent.cpp
@@ -38,8 +38,10 @@ void UMLDeformerComponent::Init()
 		{
 			ModelInstance = TUniquePtr<FMLDeformerModelInstance>(Model->CreateModelInstance());
 		}
-
-		ModelInstance->Release();
+		else
+		{
+			ModelInstance->Release();
+		}
 		ModelInstance->Init(SkelMeshComponent);
 	}
 	else
@@ -68,27 +70,38 @@ void UMLDeformerComponent::AddNeuralNetworkModifyDelegate()
 		return;
 	}
 
-	UMLDeformerModel* Model = DeformerAsset->GetModel();		
-	NeuralNetworkModifyDelegateHandle = DeformerAsset->NeuralNetworkModifyDelegate.AddLambda
-	(
-		([this]()
-		{
-			ModelInstance->Release();
-			ModelInstance->Init(SkelMeshComponent);
-		})
-	);
+	UMLDeformerModel* Model = DeformerAsset->GetModel();
+	if (Model)
+	{
+		NeuralNetworkModifyDelegateHandle = Model->NeuralNetworkModifyDelegate.AddLambda
+		(
+			([this]()
+			{
+				ModelInstance->Release();
+				ModelInstance->Init(SkelMeshComponent);
+			})
+		);
+	}
 }
 
 void UMLDeformerComponent::RemoveNeuralNetworkModifyDelegate()
 {
-	if (DeformerAsset != nullptr && NeuralNetworkModifyDelegateHandle != FDelegateHandle())
+	if (DeformerAsset != nullptr && 
+		NeuralNetworkModifyDelegateHandle != FDelegateHandle() && 
+		DeformerAsset->GetModel() != nullptr)
 	{
-		DeformerAsset->NeuralNetworkModifyDelegate.Remove(NeuralNetworkModifyDelegateHandle);
+		DeformerAsset->GetModel()->NeuralNetworkModifyDelegate.Remove(NeuralNetworkModifyDelegateHandle);
 	}
 	
 	NeuralNetworkModifyDelegateHandle = FDelegateHandle();
 }
 
+void UMLDeformerComponent::BeginDestroy()
+{
+	RemoveNeuralNetworkModifyDelegate();
+	Super::BeginDestroy();
+}
+
 void UMLDeformerComponent::Activate(bool bReset)
 {
 	// If we haven't pointed to some skeletal mesh component to use, then try to find one on the actor.
diff --git a/Engine/Plugins/Experimental/MLDeformer/MLDeformerFramework/Source/MLDeformerFramework/Private/MLDeformerModel.cpp b/Engine/Plugins/Experimental/MLDeformer/MLDeformerFramework/Source/MLDeformerFramework/Private/MLDeformerModel.cpp
--- a/Engine/Plugins/Experimental/MLDeformer/MLDeformerFramework/Source/MLDeformerFramework/Private/MLDeformerModel.cpp
+++ b/Engine/Plugins/Experimental/MLDeformer/MLDeformerFramework/Source/MLDeformerFramework/Private/MLDeformerModel.cpp
@@ -106,7 +106,7 @@ void UMLDeformerModel::PostLoad()
 
 void UMLDeformerModel::SetNeuralNetwork(UNeuralNetwork* InNeuralNetwork)
 {
-	DeformerAsset->NeuralNetworkModifyDelegate.Broadcast();
+	NeuralNetworkModifyDelegate.Broadcast();
 	NeuralNetwork = InNeuralNetwork;
 }
 
@@ -125,6 +125,8 @@ USkeleton* UMLDeformerModel::GetSkeleton(bool& bInvalidSkeletonIsError, const IP
 
 void UMLDeformerModel::BeginDestroy()
 {
+	NeuralNetworkModifyDelegate.Broadcast();
+
 	BeginReleaseResource(&VertexMapBuffer);
 	RenderResourceDestroyFence.BeginFence();
 	Super::BeginDestroy();
diff --git a/Engine/Plugins/Experimental/MLDeformer/MLDeformerFramework/Source/MLDeformerFramework/Public/MLDeformerAsset.h b/Engine/Plugins/Experimental/MLDeformer/MLDeformerFramework/Source/MLDeformerFramework/Public/MLDeformerAsset.h
--- a/Engine/Plugins/Experimental/MLDeformer/MLDeformerFramework/Source/MLDeformerFramework/Public/MLDeformerAsset.h
+++ b/Engine/Plugins/Experimental/MLDeformer/MLDeformerFramework/Source/MLDeformerFramework/Public/MLDeformerAsset.h
@@ -26,8 +26,4 @@ public:
 	/** The ML Deformer model, used to deform the mesh. */
 	UPROPERTY()
 	TObjectPtr<UMLDeformerModel> Model;
-
-	/** Delegate that will be called immediately before the NeuralNetwork is changed. */
-	DECLARE_MULTICAST_DELEGATE(FNeuralNetworkModifyDelegate);
-	FNeuralNetworkModifyDelegate NeuralNetworkModifyDelegate;
 };
diff --git a/Engine/Plugins/Experimental/MLDeformer/MLDeformerFramework/Source/MLDeformerFramework/Public/MLDeformerComponent.h b/Engine/Plugins/Experimental/MLDeformer/MLDeformerFramework/Source/MLDeformerFramework/Public/MLDeformerComponent.h
--- a/Engine/Plugins/Experimental/MLDeformer/MLDeformerFramework/Source/MLDeformerFramework/Public/MLDeformerComponent.h
+++ b/Engine/Plugins/Experimental/MLDeformer/MLDeformerFramework/Source/MLDeformerFramework/Public/MLDeformerComponent.h
@@ -27,6 +27,10 @@ class MLDEFORMERFRAMEWORK_API UMLDeformerComponent : public UActorComponent
 	GENERATED_UCLASS_BODY()
 
 public:
+	// UObject overrides.
+	void BeginDestroy() override;
+	// ~END UObject overrides.
+
 	// UActorComponent overrides.
 	void Activate(bool bReset=false) override;
 	void Deactivate() override;
@@ -49,10 +53,10 @@ protected:
 
 	void Init();
 
-	/** Bind to the DeformerAsset NeuralNetworkModifyDelegate. */
+	/** Bind to the MLDeformerModel's NeuralNetworkModifyDelegate. */
 	void AddNeuralNetworkModifyDelegate();
 
-	/** Unbind from the DeformerAsset NeuralNetworkModifyDelegate. */
+	/** Unbind from the MLDeformerModel's NeuralNetworkModifyDelegate. */
 	void RemoveNeuralNetworkModifyDelegate();
 
 protected:
diff --git a/Engine/Plugins/Experimental/MLDeformer/MLDeformerFramework/Source/MLDeformerFramework/Public/MLDeformerModel.h b/Engine/Plugins/Experimental/MLDeformer/MLDeformerFramework/Source/MLDeformerFramework/Public/MLDeformerModel.h
--- a/Engine/Plugins/Experimental/MLDeformer/MLDeformerFramework/Source/MLDeformerFramework/Public/MLDeformerModel.h
+++ b/Engine/Plugins/Experimental/MLDeformer/MLDeformerFramework/Source/MLDeformerFramework/Public/MLDeformerModel.h
@@ -166,6 +166,10 @@ public:
 	/** Fence used in render thread cleanup on destruction. */
 	FRenderCommandFence RenderResourceDestroyFence;
 
+	/** Delegate that will be called immediately before the NeuralNetwork is changed. */
+	DECLARE_MULTICAST_DELEGATE(FNeuralNetworkModifyDelegate);
+	FNeuralNetworkModifyDelegate NeuralNetworkModifyDelegate;
+
 #if WITH_EDITORONLY_DATA
 	UPROPERTY()
 	TObjectPtr<UMLDeformerVizSettings> VizSettings = nullptr;
diff --git a/Engine/Plugins/Experimental/MLDeformer/MLDeformerFramework/Source/MLDeformerFrameworkEditor/Private/MLDeformerEditorActor.cpp b/Engine/Plugins/Experimental/MLDeformer/MLDeformerFramework/Source/MLDeformerFrameworkEditor/Private/MLDeformerEditorActor.cpp
--- a/Engine/Plugins/Experimental/MLDeformer/MLDeformerFramework/Source/MLDeformerFrameworkEditor/Private/MLDeformerEditorActor.cpp
+++ b/Engine/Plugins/Experimental/MLDeformer/MLDeformerFramework/Source/MLDeformerFrameworkEditor/Private/MLDeformerEditorActor.cpp
@@ -102,4 +102,9 @@ namespace UE::MLDeformer
 		Box.Init();
 		return Box;
 	}
+
+	bool FMLDeformerEditorActor::IsGroundTruthActor() const
+	{
+		return (TypeID == ActorID_Test_GroundTruth || TypeID == ActorID_Train_GroundTruth);
+	}
 }	// namespace UE::MLDeformer
diff --git a/Engine/Plugins/Experimental/MLDeformer/MLDeformerFramework/Source/MLDeformerFrameworkEditor/Private/MLDeformerEditorModel.cpp b/Engine/Plugins/Experimental/MLDeformer/MLDeformerFramework/Source/MLDeformerFrameworkEditor/Private/MLDeformerEditorModel.cpp
--- a/Engine/Plugins/Experimental/MLDeformer/MLDeformerFramework/Source/MLDeformerFrameworkEditor/Private/MLDeformerEditorModel.cpp
+++ b/Engine/Plugins/Experimental/MLDeformer/MLDeformerFramework/Source/MLDeformerFrameworkEditor/Private/MLDeformerEditorModel.cpp
@@ -120,6 +120,7 @@ namespace UE::MLDeformer
 		FMLDeformerEditorActor* EditorActor = CreateEditorActor(Settings);
 		EditorActor->SetSkeletalMeshComponent(SkelMeshComponent);
 		EditorActor->SetCanDestroyActor(false);	// Crash will occur when destroying the Persona actor, so disable this.
+		EditorActor->SetMeshOffsetFactor(0.0f);
 		EditorActors.Add(EditorActor);
 	}
 
@@ -149,6 +150,7 @@ namespace UE::MLDeformer
 		Settings.bIsTrainingActor = false;
 		FMLDeformerEditorActor* EditorActor = CreateEditorActor(Settings);
 		EditorActor->SetSkeletalMeshComponent(SkelMeshComponent);
+		EditorActor->SetMeshOffsetFactor(0.0f);
 		EditorActors.Add(EditorActor);
 	}
 
@@ -188,6 +190,7 @@ namespace UE::MLDeformer
 		FMLDeformerEditorActor* EditorActor = static_cast<FMLDeformerEditorActor*>(CreateEditorActor(Settings));
 		EditorActor->SetSkeletalMeshComponent(SkelMeshComponent);
 		EditorActor->SetMLDeformerComponent(MLDeformerComponent);
+		EditorActor->SetMeshOffsetFactor(1.0f);
 		EditorActors.Add(EditorActor);
 	}
 
@@ -199,6 +202,23 @@ namespace UE::MLDeformer
 		CreateTestMLDeformedActor(World);
 		CreateTrainingGroundTruthActor(World);
 		CreateTestGroundTruthActor(World);
+
+		// Set the default mesh translation offsets for our ground truth actors.
+		for (FMLDeformerEditorActor* EditorActor : EditorActors)
+		{
+			if (EditorActor->IsGroundTruthActor())
+			{			
+				// The mesh offset factor basically just offsets the actor position by a given factor.
+				// The amount the actor is moved from the origin is: (MeshSpacing * MeshOffsetFactor).
+				// In test mode we have 3 actors (Linear, ML Deformed, Ground Truth), so it's mesh offset factor will be 2.0 for the ground truth.
+				// It is 2.0 because the ground truth actor in testing mode is all the way on the right, next to the ML Deformed model.
+				// In training mode we have only the Linear Skinned actor and the ground truth, so there the spacing factor is 1.0.
+				// TLDR: Basically 1.0 means its the first actor next to the linear skinned actor while 2.0 means its the second character, etc.
+				EditorActor->SetMeshOffsetFactor(EditorActor->IsTestActor() ? 2.0f : 1.0f);
+			}
+		}
+
+		OnPostCreateActors();
 	}
 
 	void FMLDeformerEditorModel::ClearWorld()
@@ -289,7 +309,7 @@ namespace UE::MLDeformer
 			{
 				const AActor* Actor = EditorActor->GetActor();
 				const FVector ActorLocation = Actor->GetActorLocation();
-				const FVector AlignmentOffset = Model->HasTrainingGroundTruth() ? Model->GetAlignmentTransform().GetTranslation() : FVector::ZeroVector;
+				const FVector AlignmentOffset = EditorActor->IsGroundTruthActor() ? Model->GetAlignmentTransform().GetTranslation() : FVector::ZeroVector;
 
 				LabelComponent->SetRelativeLocation(ActorLocation + FVector(0.0f, 0.0f, VizSettings->GetLabelHeight()) - AlignmentOffset);
 				LabelComponent->SetRelativeRotation(FQuat(FVector(0.0f, 0.0f, 1.0f), FMath::DegreesToRadians(90.0f)));
@@ -315,45 +335,10 @@ namespace UE::MLDeformer
 	void FMLDeformerEditorModel::UpdateActorTransforms()
 	{
 		const FVector MeshSpacingVector = Model->GetVizSettings()->GetMeshSpacingOffsetVector();
-
-		// Update the linear skinned training model.
-		FMLDeformerEditorActor* EditorActor = FindEditorActor(ActorID_Train_Base);
-		if (EditorActor)
-		{
-			EditorActor->GetActor()->SetActorTransform(FTransform::Identity);
-		}
-
-		// Update training ground truth actor's transform.
-		EditorActor = FindEditorActor(ActorID_Train_GroundTruth);
-		if (EditorActor)
-		{
-			FTransform Transform = FTransform::Identity;
-			Transform.AddToTranslation(MeshSpacingVector);
-			EditorActor->GetActor()->SetActorTransform(Transform);
-		}
-
-		// Update the linear skinned test model.
-		EditorActor = FindEditorActor(ActorID_Test_Base);
-		if (EditorActor)
-		{
-			EditorActor->GetActor()->SetActorTransform(FTransform::Identity);
-		}
-
-		// Update the ML deformed test actor.
-		EditorActor = FindEditorActor(ActorID_Test_MLDeformed);
-		if (EditorActor)
-		{
-			FTransform Transform = FTransform::Identity;
-			Transform.AddToTranslation(MeshSpacingVector);
-			EditorActor->GetActor()->SetActorTransform(Transform);
-		}
-
-		// Update the ground truth test actor.
-		EditorActor = FindEditorActor(ActorID_Test_GroundTruth);
-		if (EditorActor)
+		for (FMLDeformerEditorActor* EditorActor : EditorActors)
 		{
-			FTransform Transform = FTransform::Identity;
-			Transform.AddToTranslation(MeshSpacingVector * 2.0f);
+			FTransform Transform = EditorActor->IsGroundTruthActor() ? Model->GetAlignmentTransform() : FTransform::Identity;
+			Transform.AddToTranslation(MeshSpacingVector * EditorActor->GetMeshOffsetFactor());
 			EditorActor->GetActor()->SetActorTransform(Transform);
 		}
 	}
@@ -363,35 +348,10 @@ namespace UE::MLDeformer
 		UMLDeformerVizSettings* VizSettings = Model->GetVizSettings();
 		const bool bShowTrainingData = (VizSettings->GetVisualizationMode() == EMLDeformerVizMode::TrainingData);
 		const bool bShowTestData = (VizSettings->GetVisualizationMode() == EMLDeformerVizMode::TestData);
-
-		FMLDeformerEditorActor* EditorActor = FindEditorActor(ActorID_Train_Base);
-		if (EditorActor)
-		{
-			EditorActor->SetVisibility(bShowTrainingData);
-		}
-	
-		EditorActor = FindEditorActor(ActorID_Train_GroundTruth);
-		if (EditorActor)
-		{
-			EditorActor->SetVisibility(bShowTrainingData);
-		}
-
-		EditorActor = FindEditorActor(ActorID_Test_Base);
-		if (EditorActor)
-		{
-			EditorActor->SetVisibility(bShowTestData && VizSettings->GetDrawLinearSkinnedActor());
-		}
-
-		EditorActor = FindEditorActor(ActorID_Test_MLDeformed);
-		if (EditorActor)
-		{
-			EditorActor->SetVisibility(bShowTestData && VizSettings->GetDrawMLDeformedActor());
-		}
-
-		EditorActor = FindEditorActor(ActorID_Test_GroundTruth);
-		if (EditorActor)
+		for (FMLDeformerEditorActor* EditorActor : EditorActors)
 		{
-			EditorActor->SetVisibility(bShowTestData && VizSettings->GetDrawGroundTruthActor());
+			const bool bIsVisible = (EditorActor->IsTestActor() && bShowTestData) || (EditorActor->IsTrainingActor() && bShowTrainingData);
+			EditorActor->SetVisibility(bIsVisible);
 		}
 	}
 
@@ -520,7 +480,7 @@ namespace UE::MLDeformer
 			TriggerInputAssetChanged();
 			Model->InitVertexMap();
 			Model->InitGPUData();
-			SetComputeGraphDataProviders();
+			UpdateDeformerGraph();
 		}
 		else
 		if (Property->GetFName() == GET_MEMBER_NAME_CHECKED(UMLDeformerModel, AnimSequence) ||
@@ -582,7 +542,7 @@ namespace UE::MLDeformer
 		} 
 		else if (Property->GetFName() == GET_MEMBER_NAME_CHECKED(UMLDeformerVizSettings, DeformerGraph))
 		{
-			SetComputeGraphDataProviders();
+			UpdateDeformerGraph();
 			GetEditor()->GetVizSettingsDetailsView()->ForceRefresh();
 		}
 	}
@@ -1271,22 +1231,6 @@ namespace UE::MLDeformer
 		}
 	}
 
-	void FMLDeformerEditorModel::SetComputeGraphDataProviders() const
-	{
-		const UMLDeformerVizSettings* VizSettings = Model->GetVizSettings();
-		UMeshDeformer* MeshDeformer = VizSettings->GetDeformerGraph();
-
-		FMLDeformerEditorActor* EditorActor = FindEditorActor(ActorID_Test_MLDeformed);
-		if (EditorActor)
-		{
-			UDebugSkelMeshComponent* SkelMeshComponent = EditorActor->GetSkeletalMeshComponent();
-			if (SkelMeshComponent)
-			{
-				SkelMeshComponent->SetMeshDeformer(MeshDeformer);
-			}
-		}
-	}
-
 	void FMLDeformerEditorModel::OnPostTraining(ETrainingResult TrainingResult)
 	{
 		for (FMLDeformerEditorActor* EditorActor : EditorActors)
@@ -1332,7 +1276,6 @@ namespace UE::MLDeformer
 			if (Result->Load(OnnxFile))
 			{
 				Result->SetDeviceType(ENeuralDeviceType::GPU, ENeuralDeviceType::CPU, ENeuralDeviceType::GPU);	
-				SetComputeGraphDataProviders();
 				UE_LOG(LogMLDeformer, Display, TEXT("Successfully loaded Onnx file '%s'..."), *OnnxFile);
 				return Result;
 			}
diff --git a/Engine/Plugins/Experimental/MLDeformer/MLDeformerFramework/Source/MLDeformerFrameworkEditor/Private/MLDeformerEditorToolkit.cpp b/Engine/Plugins/Experimental/MLDeformer/MLDeformerFramework/Source/MLDeformerFrameworkEditor/Private/MLDeformerEditorToolkit.cpp
--- a/Engine/Plugins/Experimental/MLDeformer/MLDeformerFramework/Source/MLDeformerFrameworkEditor/Private/MLDeformerEditorToolkit.cpp
+++ b/Engine/Plugins/Experimental/MLDeformer/MLDeformerFramework/Source/MLDeformerFrameworkEditor/Private/MLDeformerEditorToolkit.cpp
@@ -531,6 +531,7 @@ namespace UE::MLDeformer
 			FMessageDialog::Open(EAppMsgType::Ok, WindowMessage, &WindowTitle);
 		}
 
+		ActiveModel->UpdateDeformerGraph();
 		return bMarkDirty;
 	}
 
diff --git a/Engine/Plugins/Experimental/MLDeformer/MLDeformerFramework/Source/MLDeformerFrameworkEditor/Public/MLDeformerEditorActor.h b/Engine/Plugins/Experimental/MLDeformer/MLDeformerFramework/Source/MLDeformerFrameworkEditor/Public/MLDeformerEditorActor.h
--- a/Engine/Plugins/Experimental/MLDeformer/MLDeformerFramework/Source/MLDeformerFrameworkEditor/Public/MLDeformerEditorActor.h
+++ b/Engine/Plugins/Experimental/MLDeformer/MLDeformerFramework/Source/MLDeformerFrameworkEditor/Public/MLDeformerEditorActor.h
@@ -42,6 +42,7 @@ namespace UE::MLDeformer
 		virtual void SetPlaySpeed(float PlaySpeed);
 		virtual FBox GetBoundingBox() const;
 		virtual void Pause(bool bPaused);
+		virtual bool IsGroundTruthActor() const;
 
 		AActor* GetActor() const { return Actor; }
 		int32 GetTypeID() const { return TypeID; }
@@ -51,10 +52,12 @@ namespace UE::MLDeformer
 		UMLDeformerComponent* GetMLDeformerComponent() const { return MLDeformerComponent.Get(); }
 		void SetSkeletalMeshComponent(UDebugSkelMeshComponent* SkelMeshComponent) { SkeletalMeshComponent = SkelMeshComponent; }
 		void SetMLDeformerComponent(UMLDeformerComponent* Component) { MLDeformerComponent = Component; }
+		void SetMeshOffsetFactor(float OffsetFactor) { MeshOffsetFactor = OffsetFactor; }
 		bool IsTrainingActor() const { return bIsTrainingActor; }
 		bool IsTestActor() const { return !bIsTrainingActor; }
 		void SetCanDestroyActor(bool bCanDestroy) { bCanDestroyActor = bCanDestroy; }
 		bool GetCanDestroyActor() const { return bCanDestroyActor; }
+		float GetMeshOffsetFactor() const { return MeshOffsetFactor; }
 
 	protected:
 		UTextRenderComponent* CreateLabelComponent(AActor* InActor, FLinearColor Color, const FText& Text) const;
@@ -78,6 +81,9 @@ namespace UE::MLDeformer
 		/** The ML Deformer component (can be nullptr). */
 		TObjectPtr<UMLDeformerComponent> MLDeformerComponent = nullptr;
 
+		/** The position offset factor of the actor. A value of 1.0 would offset the actor with the Mesh Spacing amount, a value of 2.0 would be two times the mesh spacing offset, etc. */
+		float MeshOffsetFactor = 0.0f;
+
 		/** Can the actor we created be destroyed during cleanup? */
 		bool bCanDestroyActor = true;
 
diff --git a/Engine/Plugins/Experimental/MLDeformer/MLDeformerFramework/Source/MLDeformerFrameworkEditor/Public/MLDeformerEditorModel.h b/Engine/Plugins/Experimental/MLDeformer/MLDeformerFramework/Source/MLDeformerFrameworkEditor/Public/MLDeformerEditorModel.h
--- a/Engine/Plugins/Experimental/MLDeformer/MLDeformerFramework/Source/MLDeformerFrameworkEditor/Public/MLDeformerEditorModel.h
+++ b/Engine/Plugins/Experimental/MLDeformer/MLDeformerFramework/Source/MLDeformerFrameworkEditor/Public/MLDeformerEditorModel.h
@@ -68,6 +68,7 @@ namespace UE::MLDeformer
 		// Optional overrides.
 		virtual void Init(const InitSettings& Settings);
 		virtual void CreateActors(const TSharedRef<IPersonaPreviewScene>& InPersonaPreviewScene);
+		virtual void OnPostCreateActors() {}
 		virtual void ClearWorld();
 		virtual FMLDeformerEditorActor* CreateEditorActor(const FMLDeformerEditorActor::FConstructSettings& Settings) const;
 		virtual FMLDeformerSampler* CreateSampler() const;
@@ -107,7 +108,6 @@ namespace UE::MLDeformer
 		virtual UMeshDeformer* LoadDefaultDeformerGraph();
 		virtual void SetDefaultDeformerGraphIfNeeded();
 		virtual void UpdateDeformerGraph();
-		virtual void SetComputeGraphDataProviders() const;
 		virtual void SampleDeltas();
 		virtual bool LoadTrainedNetwork() const;
 		virtual bool IsTrained() const;
diff --git a/Engine/Plugins/Experimental/MLDeformer/VertexDeltaModel/Source/VertexDeltaModel/Public/VertexDeltaModel.h b/Engine/Plugins/Experimental/MLDeformer/VertexDeltaModel/Source/VertexDeltaModel/Public/VertexDeltaModel.h
--- a/Engine/Plugins/Experimental/MLDeformer/VertexDeltaModel/Source/VertexDeltaModel/Public/VertexDeltaModel.h
+++ b/Engine/Plugins/Experimental/MLDeformer/VertexDeltaModel/Source/VertexDeltaModel/Public/VertexDeltaModel.h
@@ -43,7 +43,7 @@ public:
 	UGeometryCache* GetGeometryCache() { return GeometryCache; }
 	int32 GetNumHiddenLayers() const { return NumHiddenLayers; }
 	int32 GetNumNeuronsPerLayer() const { return NumNeuronsPerLayer; }
-	int32 GetNumIerations() const { return NumIterations; }
+	int32 GetNumIterations() const { return NumIterations; }
 	int32 GetBatchSize() const { return BatchSize; }
 	float GetLearningRate() const { return LearningRate; }
 	TArray<UE::MLDeformer::FMLDeformerGeomCacheMeshMapping>& GetGeomCacheMeshMappings() { return MeshMappings; }
@@ -78,6 +78,5 @@ public:
 	/** The learning rate used during the model training. */
 	UPROPERTY(EditAnywhere, BlueprintReadWrite, Category = "Training Settings", meta = (ClampMin = "0.000001", ClampMax = "1.0"))
 	float LearningRate = 0.001f;
-
 #endif // WITH_EDITORONLY_DATA
 };
diff --git a/Engine/Plugins/Experimental/MLDeformer/VertexDeltaModel/Source/VertexDeltaModelEditor/Private/VertexDeltaModelDetails.cpp b/Engine/Plugins/Experimental/MLDeformer/VertexDeltaModel/Source/VertexDeltaModelEditor/Private/VertexDeltaModelDetails.cpp
--- a/Engine/Plugins/Experimental/MLDeformer/VertexDeltaModel/Source/VertexDeltaModelEditor/Private/VertexDeltaModelDetails.cpp
+++ b/Engine/Plugins/Experimental/MLDeformer/VertexDeltaModel/Source/VertexDeltaModelEditor/Private/VertexDeltaModelDetails.cpp
@@ -59,11 +59,11 @@ namespace UE::VertexDeltaModel
 		FMLDeformerModelDetails::CustomizeDetails(DetailBuilder);
 
 		// Training settings.
-		SettingsCategoryBuilder->AddProperty(GET_MEMBER_NAME_CHECKED(UVertexDeltaModel, NumHiddenLayers), UMLDeformerModel::StaticClass());
-		SettingsCategoryBuilder->AddProperty(GET_MEMBER_NAME_CHECKED(UVertexDeltaModel, NumNeuronsPerLayer), UMLDeformerModel::StaticClass());
-		SettingsCategoryBuilder->AddProperty(GET_MEMBER_NAME_CHECKED(UVertexDeltaModel, NumIterations), UMLDeformerModel::StaticClass());
-		SettingsCategoryBuilder->AddProperty(GET_MEMBER_NAME_CHECKED(UVertexDeltaModel, BatchSize), UMLDeformerModel::StaticClass());
-		SettingsCategoryBuilder->AddProperty(GET_MEMBER_NAME_CHECKED(UVertexDeltaModel, LearningRate), UMLDeformerModel::StaticClass());
+		SettingsCategoryBuilder->AddProperty(GET_MEMBER_NAME_CHECKED(UVertexDeltaModel, NumHiddenLayers));
+		SettingsCategoryBuilder->AddProperty(GET_MEMBER_NAME_CHECKED(UVertexDeltaModel, NumNeuronsPerLayer));
+		SettingsCategoryBuilder->AddProperty(GET_MEMBER_NAME_CHECKED(UVertexDeltaModel, NumIterations));
+		SettingsCategoryBuilder->AddProperty(GET_MEMBER_NAME_CHECKED(UVertexDeltaModel, BatchSize));
+		SettingsCategoryBuilder->AddProperty(GET_MEMBER_NAME_CHECKED(UVertexDeltaModel, LearningRate));
 	}
 
 	void FVertexDeltaModelDetails::AddBaseMeshErrors()
@@ -103,7 +103,7 @@ namespace UE::VertexDeltaModel
 
 	void FVertexDeltaModelDetails::AddTargetMesh()
 	{
-		TargetMeshCategoryBuilder->AddProperty(GET_MEMBER_NAME_CHECKED(UVertexDeltaModel, GeometryCache), UMLDeformerModel::StaticClass());
+		TargetMeshCategoryBuilder->AddProperty(GET_MEMBER_NAME_CHECKED(UVertexDeltaModel, GeometryCache));
 
 		const FText TargetMeshErrorText = GetGeomCacheErrorText(VertexModel->GetSkeletalMesh(), VertexModel->GetGeometryCache());
 		TargetMeshCategoryBuilder->AddCustomRow(FText::FromString("TargetMeshError"))
