MLDeformer: Add some redirects to help with rigs that were made before migration to MLDeformerFramework.
#rnx

diff --git a/Engine/Plugins/Experimental/MLDeformer/MLDeformerFramework/Config/BaseMLDeformerFramework.ini b/Engine/Plugins/Experimental/MLDeformer/MLDeformerFramework/Config/BaseMLDeformerFramework.ini
new file mode 100644
--- /dev/null
+++ b/Engine/Plugins/Experimental/MLDeformer/MLDeformerFramework/Config/BaseMLDeformerFramework.ini
@@ -0,0 +1,4 @@
+[CoreRedirects]
+; Redirects from migrationof MLDeformer to MLDeformerFramework
++ClassRedirects=(OldName="/Script/MLDeformer.MLDeformerAsset",NewName="/Script/MLDeformerFramework.MLDeformerAsset")
++PackageRedirects=(OldName="MLDeformer",NewName="/MLDeformerFramework/",MatchSubstring=true)
