[MLDeformer]
- Made the FMLDeformerInstance into UMLDeformerInstance and a UCLASS to solve some garbage collection issues that triggered crashes.
- Fixed a tooltip on the "Setup Component" ML Deformer function.
- Renamed "Deformer Asset" in the ML Deformer Component to "ML Deformer Asset".
- Fix unit tests failing.
- Changed the default number of iterations in the new VertexDeltaModel to 10k instead of 500k.
- Add support for master pose components.
- Lower default mesh spacing to 125 units and lowered the default height and scale of the labels.

diff --git a/Engine/Plugins/Experimental/MLDeformer/LegacyVertexDeltaModel/Source/LegacyVertexDeltaModel/Private/Tests/LegacyVertexDeltaModelTest.cpp b/Engine/Plugins/Experimental/MLDeformer/LegacyVertexDeltaModel/Source/LegacyVertexDeltaModel/Private/Tests/LegacyVertexDeltaModelTest.cpp
--- a/Engine/Plugins/Experimental/MLDeformer/LegacyVertexDeltaModel/Source/LegacyVertexDeltaModel/Private/Tests/LegacyVertexDeltaModelTest.cpp
+++ b/Engine/Plugins/Experimental/MLDeformer/LegacyVertexDeltaModel/Source/LegacyVertexDeltaModel/Private/Tests/LegacyVertexDeltaModelTest.cpp
@@ -81,7 +81,7 @@ namespace UE::LegacyVertexDeltaModelTests
 		const int64 ExpectedNetworkOutput = NumGeomCacheVerts * 3;	// 3 floats per vertex.
 		UTEST_EQUAL(TEXT("NeuralNet output size check"), NumNetworkOutputs, ExpectedNetworkOutput);
 
-		UTEST_EQUAL(TEXT("Loss function check"), VertexDeltaModel->GetLossFunction(), ELegacyVertexDeltaModelLossFunction::MSE);
+		UTEST_EQUAL(TEXT("Loss function check"), VertexDeltaModel->GetLossFunction(), ELegacyVertexDeltaModelLossFunction::L1);
 		UTEST_EQUAL(TEXT("Activation function check"), VertexDeltaModel->GetActivationFunction(), ELegacyVertexDeltaModelActivationFunction::LRelu);
 		UTEST_EQUAL(TEXT("Hidden Layers check"), VertexDeltaModel->GetNumHiddenLayers(), 2);
 
@@ -113,7 +113,7 @@ namespace UE::LegacyVertexDeltaModelTests
 		UTEST_TRUE(TEXT("MLDeformerComponent SkelMeshComponent check"), MLDeformerComponent->GetSkeletalMeshComponent() == SkelMeshComponent);
 
 		// Create and init the model instance.
-		FMLDeformerModelInstance* ModelInstance = MLDeformerComponent->GetModelInstance();
+		UMLDeformerModelInstance* ModelInstance = MLDeformerComponent->GetModelInstance();
 		UTEST_NOT_NULL(TEXT("CreateModelInstance check"), ModelInstance);
 		UTEST_TRUE(TEXT("ModelInstance GetModel check"), ModelInstance->GetModel() == VertexDeltaModel);
 		UTEST_TRUE(TEXT("ModelInstance compatible check"), ModelInstance->IsCompatible());
diff --git a/Engine/Plugins/Experimental/MLDeformer/MLDeformerFramework/Source/MLDeformerFramework/Private/MLDeformerComponent.cpp b/Engine/Plugins/Experimental/MLDeformer/MLDeformerFramework/Source/MLDeformerFramework/Private/MLDeformerComponent.cpp
--- a/Engine/Plugins/Experimental/MLDeformer/MLDeformerFramework/Source/MLDeformerFramework/Private/MLDeformerComponent.cpp
+++ b/Engine/Plugins/Experimental/MLDeformer/MLDeformerFramework/Source/MLDeformerFramework/Private/MLDeformerComponent.cpp
@@ -17,15 +17,10 @@ UMLDeformerComponent::UMLDeformerComponent(const FObjectInitializer& ObjectIniti
 
 void UMLDeformerComponent::Init()
 {
-	using namespace UE::MLDeformer;
-
 	// If there is no deformer asset linked, release what we currently have.
 	if (DeformerAsset == nullptr)
 	{
-		if (ModelInstance)
-		{
-			ModelInstance->Release();
-		}
+		ModelInstance = nullptr;
 		return;
 	}
 
@@ -33,26 +28,23 @@ void UMLDeformerComponent::Init()
 	UMLDeformerModel* Model = DeformerAsset->GetModel();
 	if (Model)
 	{
-		// Create the model instance if it isn't there yet.
-		if (ModelInstance == nullptr)
-		{
-			ModelInstance = TUniquePtr<FMLDeformerModelInstance>(Model->CreateModelInstance());
-		}
-		else
-		{
-			ModelInstance->Release();
-		}
+		ModelInstance = Model->CreateModelInstance(this);
+		ModelInstance->SetModel(Model);
 		ModelInstance->Init(SkelMeshComponent);
 	}
 	else
 	{
+		ModelInstance = nullptr;
 		UE_LOG(LogMLDeformer, Warning, TEXT("ML Deformer component on '%s' has a deformer asset that has no ML model setup."), *GetOuter()->GetName());
 	}
 }
 
 void UMLDeformerComponent::SetupComponent(UMLDeformerAsset* InDeformerAsset, USkeletalMeshComponent* InSkelMeshComponent)
 {
-	AddTickPrerequisiteComponent(SkelMeshComponent);
+	if (InSkelMeshComponent)
+	{		
+		AddTickPrerequisiteComponent(InSkelMeshComponent);
+	}
 
 	DeformerAsset = InDeformerAsset;
 	SkelMeshComponent = InSkelMeshComponent;
@@ -77,8 +69,7 @@ void UMLDeformerComponent::AddNeuralNetworkModifyDelegate()
 		(
 			([this]()
 			{
-				ModelInstance->Release();
-				ModelInstance->Init(SkelMeshComponent);
+				Init();
 			})
 		);
 	}
@@ -118,14 +109,12 @@ void UMLDeformerComponent::Activate(bool bReset)
 void UMLDeformerComponent::Deactivate()
 {
 	RemoveNeuralNetworkModifyDelegate();
-	if (ModelInstance)
-	{
-		ModelInstance->Release();
-	}
+	ModelInstance = nullptr;
 }
 
 void UMLDeformerComponent::TickComponent(float DeltaTime, enum ELevelTick TickType, FActorComponentTickFunction* ThisTickFunction)
 {
+	Super::TickComponent(DeltaTime, TickType, ThisTickFunction);
 	if (TickType != ELevelTick::LEVELTICK_PauseTick)
 	{
 		if (ModelInstance &&
@@ -135,6 +124,22 @@ void UMLDeformerComponent::TickComponent(float DeltaTime, enum ELevelTick TickTy
 			ModelInstance->Tick(DeltaTime);
 		}
 	}
-
-	Super::TickComponent(DeltaTime, TickType, ThisTickFunction);
 }
+
+#if WITH_EDITOR
+	void UMLDeformerComponent::PostEditChangeProperty(FPropertyChangedEvent& PropertyChangedEvent)
+	{
+		const FProperty* Property = PropertyChangedEvent.Property;
+		if (Property == nullptr)
+		{
+			return;
+		}
+
+		if (Property->GetFName() == GET_MEMBER_NAME_CHECKED(UMLDeformerComponent, DeformerAsset))
+		{
+			RemoveNeuralNetworkModifyDelegate();
+			Init();
+			AddNeuralNetworkModifyDelegate();
+		}
+	}
+#endif
diff --git a/Engine/Plugins/Experimental/MLDeformer/MLDeformerFramework/Source/MLDeformerFramework/Private/MLDeformerGraphDataInterface.cpp b/Engine/Plugins/Experimental/MLDeformer/MLDeformerFramework/Source/MLDeformerFramework/Private/MLDeformerGraphDataInterface.cpp
--- a/Engine/Plugins/Experimental/MLDeformer/MLDeformerFramework/Source/MLDeformerFramework/Private/MLDeformerGraphDataInterface.cpp
+++ b/Engine/Plugins/Experimental/MLDeformer/MLDeformerFramework/Source/MLDeformerFramework/Private/MLDeformerGraphDataInterface.cpp
@@ -78,7 +78,7 @@ namespace UE::MLDeformer
 
 		const UMLDeformerAsset* DeformerAsset = DeformerComponent->GetDeformerAsset();
 		const UMLDeformerModel* Model = DeformerAsset->GetModel();
-		const FMLDeformerModelInstance* ModelInstance = DeformerComponent->GetModelInstance();
+		const UMLDeformerModelInstance* ModelInstance = DeformerComponent->GetModelInstance();
 
 		NeuralNetwork = Model->GetNeuralNetwork();
 		NeuralNetworkInferenceHandle = ModelInstance->GetNeuralNetworkInferenceHandle();
diff --git a/Engine/Plugins/Experimental/MLDeformer/MLDeformerFramework/Source/MLDeformerFramework/Private/MLDeformerModel.cpp b/Engine/Plugins/Experimental/MLDeformer/MLDeformerFramework/Source/MLDeformerFramework/Private/MLDeformerModel.cpp
--- a/Engine/Plugins/Experimental/MLDeformer/MLDeformerFramework/Source/MLDeformerFramework/Private/MLDeformerModel.cpp
+++ b/Engine/Plugins/Experimental/MLDeformer/MLDeformerFramework/Source/MLDeformerFramework/Private/MLDeformerModel.cpp
@@ -5,6 +5,7 @@
 #include "MLDeformerModelInstance.h"
 #include "MLDeformerVizSettings.h"
 #include "MLDeformerAsset.h"
+#include "MLDeformerComponent.h"
 #include "Engine/SkeletalMesh.h"
 #include "Rendering/SkeletalMeshModel.h"
 #include "Rendering/SkeletalMeshLODModel.h"
@@ -12,6 +13,7 @@
 #include "GeometryCacheMeshData.h"
 #include "GeometryCacheTrack.h"
 #include "NeuralNetwork.h"
+#include "UObject/UObjectGlobals.h"
 
 namespace UE::MLDeformer
 {
@@ -45,9 +47,9 @@ UMLDeformerInputInfo* UMLDeformerModel::CreateInputInfo()
 	return NewObject<UMLDeformerInputInfo>(this);
 }
 
-UE::MLDeformer::FMLDeformerModelInstance* UMLDeformerModel::CreateModelInstance()
+UMLDeformerModelInstance* UMLDeformerModel::CreateModelInstance(UMLDeformerComponent* Component)
 { 
-	return new UE::MLDeformer::FMLDeformerModelInstance(this);
+	return NewObject<UMLDeformerModelInstance>(Component);
 }
 
 void UMLDeformerModel::Init(UMLDeformerAsset* InDeformerAsset) 
@@ -88,6 +90,9 @@ void UMLDeformerModel::PostLoad()
 		UpdateCachedNumVertices();
 	#endif
 
+	UMLDeformerAsset* MLDeformerAsset = Cast<UMLDeformerAsset>(GetOuter());
+	Init(MLDeformerAsset);
+
 	if (InputInfo)
 	{
 		InputInfo->OnPostLoad();
@@ -125,8 +130,6 @@ USkeleton* UMLDeformerModel::GetSkeleton(bool& bInvalidSkeletonIsError, const IP
 
 void UMLDeformerModel::BeginDestroy()
 {
-	NeuralNetworkModifyDelegate.Broadcast();
-
 	BeginReleaseResource(&VertexMapBuffer);
 	RenderResourceDestroyFence.BeginFence();
 	Super::BeginDestroy();
diff --git a/Engine/Plugins/Experimental/MLDeformer/MLDeformerFramework/Source/MLDeformerFramework/Private/MLDeformerModelInstance.cpp b/Engine/Plugins/Experimental/MLDeformer/MLDeformerFramework/Source/MLDeformerFramework/Private/MLDeformerModelInstance.cpp
--- a/Engine/Plugins/Experimental/MLDeformer/MLDeformerFramework/Source/MLDeformerFramework/Private/MLDeformerModelInstance.cpp
+++ b/Engine/Plugins/Experimental/MLDeformer/MLDeformerFramework/Source/MLDeformerFramework/Private/MLDeformerModelInstance.cpp
@@ -9,261 +9,289 @@
 #include "Animation/AnimInstance.h"
 #include "NeuralNetwork.h"
 
-namespace UE::MLDeformer
+void UMLDeformerModelInstance::BeginDestroy()
 {
-	FMLDeformerModelInstance::FMLDeformerModelInstance(UMLDeformerModel* InModel)
-		: Model(InModel)
+	Release();
+	Super::BeginDestroy();
+}
+
+void UMLDeformerModelInstance::Release()
+{
+	UNeuralNetwork* NeuralNetwork = Model.Get() ? Model->GetNeuralNetwork() : nullptr;
+	if (NeuralNetwork != nullptr && NeuralNetworkInferenceHandle != -1)
 	{
-		check(InModel);
+		NeuralNetwork->DestroyInferenceContext(NeuralNetworkInferenceHandle);
+		NeuralNetworkInferenceHandle = -1;
 	}
+}
+
+void UMLDeformerModelInstance::Init(USkeletalMeshComponent* SkelMeshComponent)
+{
+	SkeletalMeshComponent = SkelMeshComponent;
 
-	FMLDeformerModelInstance::~FMLDeformerModelInstance()
+	if (SkelMeshComponent == nullptr)
 	{
-		Release();
+		AssetBonesToSkelMeshMappings.Empty();
+		return;
 	}
 
-	void FMLDeformerModelInstance::Release()
+	USkeletalMesh* SkelMesh = SkelMeshComponent->SkeletalMesh;
+	if (SkelMesh)
 	{
-		UNeuralNetwork* NeuralNetwork = Model->GetNeuralNetwork();
-		if (NeuralNetwork != nullptr && NeuralNetworkInferenceHandle != -1)
+		// Init the bone mapping table.
+		const UMLDeformerInputInfo* InputInfo = Model->GetInputInfo();
+		const int32 NumAssetBones = InputInfo->GetNumBones();
+		AssetBonesToSkelMeshMappings.Reset();
+		AssetBonesToSkelMeshMappings.AddUninitialized(NumAssetBones);
+		TrainingBoneTransforms.SetNumUninitialized(NumAssetBones);
+
+		// For each bone in the deformer asset, find the matching bone index inside the skeletal mesh component.
+		for (int32 Index = 0; Index < NumAssetBones; ++Index)
 		{
-			NeuralNetwork->DestroyInferenceContext(NeuralNetworkInferenceHandle);
-			NeuralNetworkInferenceHandle = -1;
+			const FName BoneName = InputInfo->GetBoneName(Index);
+			const int32 SkelMeshBoneIndex = SkeletalMeshComponent->GetBaseComponent()->GetBoneIndex(BoneName);
+			AssetBonesToSkelMeshMappings[Index] = SkelMeshBoneIndex;
 		}
 	}
 
-	void FMLDeformerModelInstance::Init(USkeletalMeshComponent* SkelMeshComponent)
-	{
-		SkeletalMeshComponent = SkelMeshComponent;
+	// Perform a compatibility check.
+	UpdateCompatibilityStatus();
+}
+
+void UMLDeformerModelInstance::UpdateCompatibilityStatus()
+{
+	bIsCompatible = SkeletalMeshComponent->SkeletalMesh && CheckCompatibility(SkeletalMeshComponent, true).IsEmpty();
+}
+
+FString UMLDeformerModelInstance::CheckCompatibility(USkeletalMeshComponent* InSkelMeshComponent, bool LogIssues)
+{
+	ErrorText = FString();
 
-		if (SkelMeshComponent == nullptr)
+	// If we're not compatible, generate a compatibility string.
+	USkeletalMesh* SkelMesh = InSkelMeshComponent ? InSkelMeshComponent->SkeletalMesh.Get() : nullptr;
+	UMLDeformerInputInfo* InputInfo = Model->GetInputInfo();
+	if (SkelMesh && !InputInfo->IsCompatible(SkelMesh) && Model->GetDeformerAsset())
+	{
+		ErrorText += InputInfo->GenerateCompatibilityErrorString(SkelMesh);
+		ErrorText += "\n";
+		check(!ErrorText.IsEmpty());
+		if (LogIssues)
 		{
-			AssetBonesToSkelMeshMappings.Empty();
-			return;
+			UE_LOG(LogMLDeformer, Error, TEXT("ML Deformer '%s' isn't compatible with Skeletal Mesh '%s'.\nReason(s):\n%s"), 
+				*Model->GetDeformerAsset()->GetName(), 
+				*SkelMesh->GetName(), 
+				*ErrorText);
 		}
+	}
 
-		USkeletalMesh* SkelMesh = SkelMeshComponent->SkeletalMesh;
-		if (SkelMesh)
+	UNeuralNetwork* NeuralNetwork = Model->GetNeuralNetwork();
+	if (NeuralNetwork != nullptr && NeuralNetwork->IsLoaded() && Model->GetDeformerAsset())
+	{
+		const int64 NumNeuralNetInputs = NeuralNetwork->GetInputTensor().Num();
+		const int64 NumDeformerAssetInputs = static_cast<int64>(Model->GetInputInfo()->CalcNumNeuralNetInputs());
+		if (NumNeuralNetInputs != NumDeformerAssetInputs)
 		{
-			// Init the bone mapping table.
-			const UMLDeformerInputInfo* InputInfo = Model->GetInputInfo();
-			const int32 NumAssetBones = InputInfo->GetNumBones();
-			AssetBonesToSkelMeshMappings.Reset();
-			AssetBonesToSkelMeshMappings.AddUninitialized(NumAssetBones);
-
-			// For each bone in the deformer asset, find the matching bone index inside the skeletal mesh component.
-			for (int32 Index = 0; Index < NumAssetBones; ++Index)
+			const FString InputErrorString = "The number of network inputs doesn't match the asset. Please retrain the asset."; 
+			ErrorText += InputErrorString + "\n";
+			if (LogIssues)
 			{
-				const FName BoneName = InputInfo->GetBoneName(Index);
-				const int32 SkelMeshBoneIndex = SkeletalMeshComponent->GetBoneIndex(BoneName);
-				AssetBonesToSkelMeshMappings[Index] = SkelMeshBoneIndex;
+				UE_LOG(LogMLDeformer, Error, TEXT("Deformer '%s': %s"), *(Model->GetDeformerAsset()->GetName()), *InputErrorString);
 			}
 		}
-
-		// Perform a compatibility check.
-		UpdateCompatibilityStatus();
 	}
 
-	void FMLDeformerModelInstance::UpdateCompatibilityStatus()
-	{
-		bIsCompatible = SkeletalMeshComponent->SkeletalMesh && CheckCompatibility(SkeletalMeshComponent, true).IsEmpty();
-	}
+	return ErrorText;
+}
 
-	FString FMLDeformerModelInstance::CheckCompatibility(USkeletalMeshComponent* InSkelMeshComponent, bool LogIssues)
+void UMLDeformerModelInstance::UpdateBoneTransforms()
+{
+	const USkinnedMeshComponent* MasterPoseComponent = SkeletalMeshComponent->MasterPoseComponent.Get();
+	if (MasterPoseComponent)
 	{
-		ErrorText = FString();
-
-		// If we're not compatible, generate a compatibility string.
-		USkeletalMesh* SkelMesh = InSkelMeshComponent ? InSkelMeshComponent->SkeletalMesh.Get() : nullptr;
-		UMLDeformerInputInfo* InputInfo = Model->GetInputInfo();
-		if (SkelMesh && !InputInfo->IsCompatible(SkelMesh))
+		const TArray<FTransform>& MasterTransforms = MasterPoseComponent->GetComponentSpaceTransforms();
+		USkeletalMesh* Mesh = MasterPoseComponent->SkeletalMesh;
+		const int32 NumTrainingBones = AssetBonesToSkelMeshMappings.Num();
+		for (int32 Index = 0; Index < NumTrainingBones; ++Index)
 		{
-			ErrorText += InputInfo->GenerateCompatibilityErrorString(SkelMesh);
-			ErrorText += "\n";
-			check(!ErrorText.IsEmpty());
-			if (LogIssues)
+			const int32 ComponentBoneIndex = AssetBonesToSkelMeshMappings[Index];
+			const FTransform& ComponentSpaceTransform = MasterTransforms[ComponentBoneIndex];
+			const int32 ParentIndex = Mesh->GetRefSkeleton().GetParentIndex(ComponentBoneIndex);
+			if (ParentIndex != INDEX_NONE)
 			{
-				UE_LOG(LogMLDeformer, Error, TEXT("ML Deformer '%s' isn't compatible with Skeletal Mesh '%s'.\nReason(s):\n%s"), 
-					*Model->GetDeformerAsset()->GetName(), 
-					*SkelMesh->GetName(), 
-					*ErrorText);
+				TrainingBoneTransforms[Index] = ComponentSpaceTransform.GetRelativeTransform(MasterTransforms[ParentIndex]);
 			}
-		}
-
-		UNeuralNetwork* NeuralNetwork = Model->GetNeuralNetwork();
-		if (NeuralNetwork != nullptr && NeuralNetwork->IsLoaded())
-		{
-			const int64 NumNeuralNetInputs = NeuralNetwork->GetInputTensor().Num();
-			const int64 NumDeformerAssetInputs = static_cast<int64>(Model->GetInputInfo()->CalcNumNeuralNetInputs());
-			if (NumNeuralNetInputs != NumDeformerAssetInputs)
+			else
 			{
-				const FString InputErrorString = "The number of network inputs doesn't match the asset. Please retrain the asset."; 
-				ErrorText += InputErrorString + "\n";
-				if (LogIssues)
-				{
-					UE_LOG(LogMLDeformer, Error, TEXT("Deformer '%s': %s"), *(Model->GetDeformerAsset()->GetName()), *InputErrorString);
-				}
+				TrainingBoneTransforms[Index] = ComponentSpaceTransform;
 			}
+			TrainingBoneTransforms[Index].NormalizeRotation();
 		}
-
-		return ErrorText;
 	}
-
-	int64 FMLDeformerModelInstance::SetBoneTransforms(float* OutputBuffer, int64 OutputBufferSize, int64 StartIndex)
+	else
 	{
-		// Extract the bone space bone transforms from the component.
-		// Write the output transforms into the BoneTransforms array.
 		BoneTransforms = SkeletalMeshComponent->GetBoneSpaceTransforms();
-
-		// Write the transforms into the output buffer.
-		const UMLDeformerInputInfo* InputInfo = Model->GetInputInfo();
-		const int32 AssetNumBones = InputInfo->GetNumBones();
-		int64 Index = StartIndex;
-		check((Index + AssetNumBones * 6) <= OutputBufferSize); // Make sure we don't write past the OutputBuffer. (6 because of two columns of the 3x3 rotation matrix)
-		for (int32 BoneIndex = 0; BoneIndex < AssetNumBones; ++BoneIndex)
+		const int32 NumTrainingBones = AssetBonesToSkelMeshMappings.Num();
+		for (int32 Index = 0; Index < NumTrainingBones; ++Index)
 		{
-			const int32 SkelMeshBoneIndex = AssetBonesToSkelMeshMappings[BoneIndex];
-			const FMatrix RotationMatrix = (SkelMeshBoneIndex != INDEX_NONE) ? BoneTransforms[SkelMeshBoneIndex].GetRotation().ToMatrix() : FMatrix::Identity;
-			const FVector X = RotationMatrix.GetColumn(0);
-			const FVector Y = RotationMatrix.GetColumn(1);	
-			OutputBuffer[Index++] = X.X;
-			OutputBuffer[Index++] = X.Y;
-			OutputBuffer[Index++] = X.Z;
-			OutputBuffer[Index++] = Y.X;
-			OutputBuffer[Index++] = Y.Y;
-			OutputBuffer[Index++] = Y.Z;
+			const int32 ComponentBoneIndex = AssetBonesToSkelMeshMappings[Index];
+			TrainingBoneTransforms[Index] = BoneTransforms[ComponentBoneIndex];
 		}
-
-		return Index;
 	}
+}
 
-	int64 FMLDeformerModelInstance::SetCurveValues(float* OutputBuffer, int64 OutputBufferSize, int64 StartIndex)
+int64 UMLDeformerModelInstance::SetBoneTransforms(float* OutputBuffer, int64 OutputBufferSize, int64 StartIndex)
+{
+	// Get the transforms for the bones we used during training.
+	// These are in the space relative to their parent.
+	UpdateBoneTransforms();
+
+	// Write the transforms into the output buffer.
+	const UMLDeformerInputInfo* InputInfo = Model->GetInputInfo();
+	const int32 AssetNumBones = InputInfo->GetNumBones();
+	int64 Index = StartIndex;
+	check((Index + AssetNumBones * 6) <= OutputBufferSize); // Make sure we don't write past the OutputBuffer. (6 because of two columns of the 3x3 rotation matrix)
+	for (int32 BoneIndex = 0; BoneIndex < AssetNumBones; ++BoneIndex)
 	{
-		const UMLDeformerInputInfo* InputInfo = Model->GetInputInfo();
+		const FMatrix RotationMatrix = TrainingBoneTransforms[BoneIndex].GetRotation().ToMatrix();
+		const FVector X = RotationMatrix.GetColumn(0);
+		const FVector Y = RotationMatrix.GetColumn(1);	
+		OutputBuffer[Index++] = X.X;
+		OutputBuffer[Index++] = X.Y;
+		OutputBuffer[Index++] = X.Z;
+		OutputBuffer[Index++] = Y.X;
+		OutputBuffer[Index++] = Y.Y;
+		OutputBuffer[Index++] = Y.Z;
+	}
+
+	return Index;
+}
+
+int64 UMLDeformerModelInstance::SetCurveValues(float* OutputBuffer, int64 OutputBufferSize, int64 StartIndex)
+{
+	const UMLDeformerInputInfo* InputInfo = Model->GetInputInfo();
 
-		// Write the weights into the output buffer.
-		int64 Index = StartIndex;
-		const int32 AssetNumCurves = InputInfo->GetNumCurves();
-		check((Index + AssetNumCurves) <= OutputBufferSize); // Make sure we don't write past the OutputBuffer.
+	// Write the weights into the output buffer.
+	int64 Index = StartIndex;
+	const int32 AssetNumCurves = InputInfo->GetNumCurves();
+	check((Index + AssetNumCurves) <= OutputBufferSize); // Make sure we don't write past the OutputBuffer.
 
-		// Write the curve weights to the output buffer.
-		UAnimInstance* AnimInstance = SkeletalMeshComponent->GetAnimInstance();
-		if (AnimInstance)
+	// Write the curve weights to the output buffer.
+	UAnimInstance* AnimInstance = SkeletalMeshComponent->GetAnimInstance();
+	if (AnimInstance)
+	{
+		for (int32 CurveIndex = 0; CurveIndex < AssetNumCurves; ++CurveIndex)
 		{
-			for (int32 CurveIndex = 0; CurveIndex < AssetNumCurves; ++CurveIndex)
-			{
-				const FName CurveName = InputInfo->GetCurveName(CurveIndex);
-				const float CurveValue = AnimInstance->GetCurveValue(CurveName);	// Outputs 0.0 when not found.
-				OutputBuffer[Index++] = CurveValue;
-			}
+			const FName CurveName = InputInfo->GetCurveName(CurveIndex);
+			const float CurveValue = AnimInstance->GetCurveValue(CurveName);	// Outputs 0.0 when not found.
+			OutputBuffer[Index++] = CurveValue;
 		}
-		else
+	}
+	else
+	{
+		for (int32 CurveIndex = 0; CurveIndex < AssetNumCurves; ++CurveIndex)
 		{
-			for (int32 CurveIndex = 0; CurveIndex < AssetNumCurves; ++CurveIndex)
-			{
-				OutputBuffer[Index++] = 0.0f;
-			}
+			OutputBuffer[Index++] = 0.0f;
 		}
-
-		return Index;
 	}
 
-	void FMLDeformerModelInstance::SetNeuralNetworkInputValues(float* InputData, int64 NumInputFloats)
-	{
-		check(SkeletalMeshComponent);
+	return Index;
+}
 
-		// Feed data to the network inputs.
-		int64 BufferOffset = 0;
-		BufferOffset = SetBoneTransforms(InputData, NumInputFloats, BufferOffset);
-		BufferOffset = SetCurveValues(InputData, NumInputFloats, BufferOffset);
-		check(BufferOffset == NumInputFloats);
-	}
+void UMLDeformerModelInstance::SetNeuralNetworkInputValues(float* InputData, int64 NumInputFloats)
+{
+	check(SkeletalMeshComponent);
+
+	// Feed data to the network inputs.
+	int64 BufferOffset = 0;
+	BufferOffset = SetBoneTransforms(InputData, NumInputFloats, BufferOffset);
+	BufferOffset = SetCurveValues(InputData, NumInputFloats, BufferOffset);
+	check(BufferOffset == NumInputFloats);
+}
 
-	bool FMLDeformerModelInstance::IsValidForDataProvider() const
+bool UMLDeformerModelInstance::IsValidForDataProvider() const
+{
+	const UNeuralNetwork* NeuralNetwork = Model->GetNeuralNetwork();
+	if (NeuralNetwork == nullptr || 
+		!NeuralNetwork->IsLoaded() || 
+		NeuralNetwork->GetDeviceType() != ENeuralDeviceType::GPU || 
+		NeuralNetwork->GetOutputDeviceType() != ENeuralDeviceType::GPU)
 	{
-		const UNeuralNetwork* NeuralNetwork = Model->GetNeuralNetwork();
-		if (NeuralNetwork == nullptr || 
-			!NeuralNetwork->IsLoaded() || 
-			NeuralNetwork->GetDeviceType() != ENeuralDeviceType::GPU || 
-			NeuralNetwork->GetOutputDeviceType() != ENeuralDeviceType::GPU)
-		{
-			return false;
-		}
+		return false;
+	}
+
+	return
+		Model->GetVertexMapBuffer().ShaderResourceViewRHI != nullptr &&
+		GetNeuralNetworkInferenceHandle() != -1;
+}
 
-		return
-			Model->GetVertexMapBuffer().ShaderResourceViewRHI != nullptr &&
-			GetNeuralNetworkInferenceHandle() != -1;
+void UMLDeformerModelInstance::Tick(float DeltaTime)
+{
+	// Some safety checks.
+	if (Model == nullptr || 
+		SkeletalMeshComponent == nullptr || 
+		SkeletalMeshComponent->SkeletalMesh == nullptr || 
+		!bIsCompatible)
+	{
+		return;
 	}
 
-	void FMLDeformerModelInstance::Tick(float DeltaTime)
+	// Get the network and make sure it's loaded.
+	UNeuralNetwork* NeuralNetwork = Model->GetNeuralNetwork();
+	if (NeuralNetwork == nullptr || !NeuralNetwork->IsLoaded())
 	{
-		// Some safety checks.
-		if (Model == nullptr || 
-			SkeletalMeshComponent == nullptr || 
-			SkeletalMeshComponent->SkeletalMesh == nullptr || 
-			!bIsCompatible)
-		{
-			return;
-		}
+		return;
+	}
 
-		// Get the network and make sure it's loaded.
-		UNeuralNetwork* NeuralNetwork = Model->GetNeuralNetwork();
-		if (NeuralNetwork == nullptr || !NeuralNetwork->IsLoaded())
+	// If we're not on the GPU we can't continue really.
+	// This is needed as the deformer graph system needs it on the GPU.
+	// Some platforms might not support GPU yet.
+	// Only the inputs are on the CPU.
+	check(NeuralNetwork->GetInputDeviceType() == ENeuralDeviceType::CPU);
+	if (!bAllowCPU)
+	{
+		if (NeuralNetwork->GetDeviceType() != ENeuralDeviceType::GPU ||
+			NeuralNetwork->GetOutputDeviceType() != ENeuralDeviceType::GPU)
 		{
 			return;
 		}
+	}
 
-		// If we're not on the GPU we can't continue really.
-		// This is needed as the deformer graph system needs it on the GPU.
-		// Some platforms might not support GPU yet.
-		// Only the inputs are on the CPU.
-		check(NeuralNetwork->GetInputDeviceType() == ENeuralDeviceType::CPU);
-		if (!bAllowCPU)
-		{
-			if (NeuralNetwork->GetDeviceType() != ENeuralDeviceType::GPU ||
-				NeuralNetwork->GetOutputDeviceType() != ENeuralDeviceType::GPU)
-			{
-				return;
-			}
-		}
-
-		// Allocate an inference context if none has already been allocated.
+	// Allocate an inference context if none has already been allocated.
+	if (NeuralNetworkInferenceHandle == -1)
+	{
+		NeuralNetworkInferenceHandle = NeuralNetwork->CreateInferenceContext();
 		if (NeuralNetworkInferenceHandle == -1)
-		{
-			NeuralNetworkInferenceHandle = NeuralNetwork->CreateInferenceContext();
-			if (NeuralNetworkInferenceHandle == -1)
-			{
-				return;
-			}
-		}
-
-		// If the neural network expects a different number of inputs, do nothing.
-		const int64 NumNeuralNetInputs = NeuralNetwork->GetInputTensorForContext(NeuralNetworkInferenceHandle).Num();
-		const int64 NumDeformerAssetInputs = Model->GetInputInfo()->CalcNumNeuralNetInputs();
-		if (NumNeuralNetInputs != NumDeformerAssetInputs)
 		{
 			return;
 		}
+	}
 
-		// Update and write the input values directly into the input tensor.
-		float* InputDataPointer = static_cast<float*>(NeuralNetwork->GetInputDataPointerMutableForContext(NeuralNetworkInferenceHandle));
-		SetNeuralNetworkInputValues(InputDataPointer, NumNeuralNetInputs);
+	// If the neural network expects a different number of inputs, do nothing.
+	const int64 NumNeuralNetInputs = NeuralNetwork->GetInputTensorForContext(NeuralNetworkInferenceHandle).Num();
+	const int64 NumDeformerAssetInputs = Model->GetInputInfo()->CalcNumNeuralNetInputs();
+	if (NumNeuralNetInputs != NumDeformerAssetInputs)
+	{
+		return;
+	}
 
-		// Run the neural network.
-		if (NeuralNetwork->GetOutputDeviceType() == ENeuralDeviceType::GPU)
-		{
-			ENQUEUE_RENDER_COMMAND(RunNeuralNetwork)(
-				[NeuralNetwork, Handle = NeuralNetworkInferenceHandle](FRHICommandListImmediate& RHICmdList)
-			{
-					// Output deltas will be available on GPU for DeformerGraph via UMLDeformerDataProvider.
-					FRDGBuilder GraphBuilder(RHICmdList);
-					NeuralNetwork->Run(GraphBuilder, Handle);
-					GraphBuilder.Execute();
-			});
-		}
-		else
+	// Update and write the input values directly into the input tensor.
+	float* InputDataPointer = static_cast<float*>(NeuralNetwork->GetInputDataPointerMutableForContext(NeuralNetworkInferenceHandle));
+	SetNeuralNetworkInputValues(InputDataPointer, NumNeuralNetInputs);
+
+	// Run the neural network.
+	if (NeuralNetwork->GetOutputDeviceType() == ENeuralDeviceType::GPU)
+	{
+		ENQUEUE_RENDER_COMMAND(RunNeuralNetwork)(
+			[NeuralNetwork, Handle = NeuralNetworkInferenceHandle](FRHICommandListImmediate& RHICmdList)
 		{
-			NeuralNetwork->Run();
-		}
+				// Output deltas will be available on GPU for DeformerGraph via UMLDeformerDataProvider.
+				FRDGBuilder GraphBuilder(RHICmdList);
+				NeuralNetwork->Run(GraphBuilder, Handle);
+				GraphBuilder.Execute();
+		});
+	}
+	else
+	{
+		NeuralNetwork->Run();
 	}
-}	// namespace UE::MLDeformer
+}
diff --git a/Engine/Plugins/Experimental/MLDeformer/MLDeformerFramework/Source/MLDeformerFramework/Private/Tests/MLDeformerTestModel.cpp b/Engine/Plugins/Experimental/MLDeformer/MLDeformerFramework/Source/MLDeformerFramework/Private/Tests/MLDeformerTestModel.cpp
--- a/Engine/Plugins/Experimental/MLDeformer/MLDeformerFramework/Source/MLDeformerFramework/Private/Tests/MLDeformerTestModel.cpp
+++ b/Engine/Plugins/Experimental/MLDeformer/MLDeformerFramework/Source/MLDeformerFramework/Private/Tests/MLDeformerTestModel.cpp
@@ -1,18 +1,12 @@
 // Copyright Epic Games, Inc. All Rights Reserved.
 
 #include "MLDeformerTestModel.h"
+#include "MLDeformerComponent.h"
+#include "UObject/UObjectGlobals.h"
 
-namespace UE::MLDeformerTests
+UMLDeformerModelInstance* UMLDeformerTestModel::CreateModelInstance(UMLDeformerComponent* Component)
 {
-	FTestModelInstance::FTestModelInstance(UMLDeformerModel* InModel)
-		: FMLDeformerModelInstance(InModel)
-	{
-	}
-}	// namespace UE::MLDeformerTests
-
-UE::MLDeformer::FMLDeformerModelInstance* UMLDeformerTestModel::CreateModelInstance()
-{
-	return new UE::MLDeformerTests::FTestModelInstance(this);
+	return NewObject<UTestModelInstance>(Component);
 }
 
 FString UMLDeformerTestModel::GetDisplayName() const 
diff --git a/Engine/Plugins/Experimental/MLDeformer/MLDeformerFramework/Source/MLDeformerFramework/Private/Tests/MLDeformerTestModel.h b/Engine/Plugins/Experimental/MLDeformer/MLDeformerFramework/Source/MLDeformerFramework/Private/Tests/MLDeformerTestModel.h
--- a/Engine/Plugins/Experimental/MLDeformer/MLDeformerFramework/Source/MLDeformerFramework/Private/Tests/MLDeformerTestModel.h
+++ b/Engine/Plugins/Experimental/MLDeformer/MLDeformerFramework/Source/MLDeformerFramework/Private/Tests/MLDeformerTestModel.h
@@ -7,27 +7,26 @@
 #include "MLDeformerModelInstance.h"
 #include "MLDeformerTestModel.generated.h"
 
+class UMLDeformerComponent;
+
 UCLASS()
 class UMLDeformerTestModel 
-	:  public UMLDeformerModel
+	: public UMLDeformerModel
 {
 	GENERATED_BODY()
 
 public:
-	UE::MLDeformer::FMLDeformerModelInstance* CreateModelInstance() override;
+	UMLDeformerModelInstance* CreateModelInstance(UMLDeformerComponent* Component) override;
 	FString GetDisplayName() const override;
 	void UpdateNumTargetMeshVertices() override;
 };
 
-namespace UE::MLDeformerTests
+UCLASS()
+class UTestModelInstance
+	: public UMLDeformerModelInstance
 {
-	using namespace UE::MLDeformer;
+	GENERATED_BODY()
 
-	class FTestModelInstance
-		: public FMLDeformerModelInstance
-	{
-	public:
-		FTestModelInstance(UMLDeformerModel* InModel);
-		void Tick(float DeltaTime) {}
-	};
-}	// namespace UE::MLDeformerTests
+public:
+	void Tick(float DeltaTime) {}
+};
diff --git a/Engine/Plugins/Experimental/MLDeformer/MLDeformerFramework/Source/MLDeformerFramework/Private/Tests/ModelTests.cpp b/Engine/Plugins/Experimental/MLDeformer/MLDeformerFramework/Source/MLDeformerFramework/Private/Tests/ModelTests.cpp
--- a/Engine/Plugins/Experimental/MLDeformer/MLDeformerFramework/Source/MLDeformerFramework/Private/Tests/ModelTests.cpp
+++ b/Engine/Plugins/Experimental/MLDeformer/MLDeformerFramework/Source/MLDeformerFramework/Private/Tests/ModelTests.cpp
@@ -77,7 +77,7 @@ namespace UE::MLDeformerTests
 		UTEST_TRUE(TEXT("MLDeformerComponent SkelMeshComponent check"), MLDeformerComponent->GetSkeletalMeshComponent() == SkelMeshComponent);
 
 		// Create and init the model instance.
-		FMLDeformerModelInstance* ModelInstance = MLDeformerComponent->GetModelInstance();
+		UMLDeformerModelInstance* ModelInstance = MLDeformerComponent->GetModelInstance();
 		UTEST_NOT_NULL(TEXT("CreateModelInstance check"), ModelInstance);
 		UTEST_TRUE(TEXT("ModelInstance GetModel check"), ModelInstance->GetModel() == TestModel);
 		UTEST_TRUE(TEXT("ModelInstance compatible check"), ModelInstance->IsCompatible());
diff --git a/Engine/Plugins/Experimental/MLDeformer/MLDeformerFramework/Source/MLDeformerFramework/Public/MLDeformerComponent.h b/Engine/Plugins/Experimental/MLDeformer/MLDeformerFramework/Source/MLDeformerFramework/Public/MLDeformerComponent.h
--- a/Engine/Plugins/Experimental/MLDeformer/MLDeformerFramework/Source/MLDeformerFramework/Public/MLDeformerComponent.h
+++ b/Engine/Plugins/Experimental/MLDeformer/MLDeformerFramework/Source/MLDeformerFramework/Public/MLDeformerComponent.h
@@ -5,16 +5,11 @@
 #include "CoreMinimal.h"
 #include "UObject/Object.h"
 #include "Components/ActorComponent.h"
-#include "MLDeformerModelInstance.h"
 #include "MLDeformerComponent.generated.h"
 
 class UMLDeformerAsset;
 class USkeletalMeshComponent;
-
-namespace UE::MLDeformer
-{
-	class FMLDeformerModelInstance;
-}
+class UMLDeformerModelInstance;
 
 /**
  * The ML mesh deformer component.
@@ -22,7 +17,8 @@ namespace UE::MLDeformer
  * The component will perform runtime inference of the deformer model setup inside the asset.
  */
 UCLASS(Blueprintable, ClassGroup = Component, BlueprintType, meta = (BlueprintSpawnableComponent))
-class MLDEFORMERFRAMEWORK_API UMLDeformerComponent : public UActorComponent
+class MLDEFORMERFRAMEWORK_API UMLDeformerComponent
+	: public UActorComponent
 {
 	GENERATED_UCLASS_BODY()
 
@@ -36,14 +32,19 @@ public:
 	void Deactivate() override;
 	// ~END UActorComponent overrides.
 
+	/** Setup the ML Deformer, by picking the deformer asset and skeletal mesh component. */
 	UFUNCTION(BlueprintCallable, Category = "MLDeformer")
 	void SetupComponent(UMLDeformerAsset* InDeformerAsset, USkeletalMeshComponent* InSkelMeshComponent);
 
+#if WITH_EDITOR
+	void PostEditChangeProperty(FPropertyChangedEvent& PropertyChangedEvent) override;
+#endif
+
 	UMLDeformerAsset* GetDeformerAsset() { return DeformerAsset; }
 	void SetDeformerAsset(UMLDeformerAsset* InDeformerAsset) { DeformerAsset = InDeformerAsset; }
 	float GetWeight() const { return Weight; }
 	void SetWeight(float NormalizedWeightValue) { Weight = NormalizedWeightValue;  }
-	UE::MLDeformer::FMLDeformerModelInstance* GetModelInstance() const { return ModelInstance.Get(); }
+	UMLDeformerModelInstance* GetModelInstance() const { return ModelInstance; }
 	USkeletalMeshComponent* GetSkeletalMeshComponent() const { return SkelMeshComponent.Get(); }
 
 protected:
@@ -60,14 +61,6 @@ protected:
 	void RemoveNeuralNetworkModifyDelegate();
 
 protected:
-	/** The deformer asset to use. */
-	UPROPERTY(EditAnywhere, Category = "ML Deformer")
-	TObjectPtr<UMLDeformerAsset> DeformerAsset = nullptr;
-
-	/** How active is this deformer? Can be used to blend it in and out. */
-	UPROPERTY(EditAnywhere, Category = "ML Deformer", meta = (ClampMin = "0.0", ClampMax = "1.0"))
-	float Weight = 1.0f;
-
 	/** 
 	 * The skeletal mesh component we want to grab the bone transforms etc from. 
 	 * This can be a nullptr. When it is a nullptr then it will internally try to find the first skeletal mesh component on the actor.
@@ -75,9 +68,19 @@ protected:
 	 */
 	TObjectPtr<USkeletalMeshComponent> SkelMeshComponent = nullptr;
 
-	/** The deformation model instance. This is used to perform the runtime updates and run the inference. */
-	TUniquePtr<UE::MLDeformer::FMLDeformerModelInstance> ModelInstance = nullptr;
-
 	/** DelegateHandle for NeuralNetwork modification. */
 	FDelegateHandle NeuralNetworkModifyDelegateHandle;
+
+public:
+	/** The deformer asset to use. */
+	UPROPERTY(EditAnywhere, DisplayName = "ML Deformer Asset", Category = "ML Deformer")
+	TObjectPtr<UMLDeformerAsset> DeformerAsset = nullptr;
+
+	/** How active is this deformer? Can be used to blend it in and out. */
+	UPROPERTY(EditAnywhere, Category = "ML Deformer", meta = (ClampMin = "0.0", ClampMax = "1.0"))
+	float Weight = 1.0f;
+
+	/** The deformation model instance. This is used to perform the runtime updates and run the inference. */
+	UPROPERTY(Transient)
+	TObjectPtr<UMLDeformerModelInstance> ModelInstance = nullptr;
 };
diff --git a/Engine/Plugins/Experimental/MLDeformer/MLDeformerFramework/Source/MLDeformerFramework/Public/MLDeformerModel.h b/Engine/Plugins/Experimental/MLDeformer/MLDeformerFramework/Source/MLDeformerFramework/Public/MLDeformerModel.h
--- a/Engine/Plugins/Experimental/MLDeformer/MLDeformerFramework/Source/MLDeformerFramework/Public/MLDeformerModel.h
+++ b/Engine/Plugins/Experimental/MLDeformer/MLDeformerFramework/Source/MLDeformerFramework/Public/MLDeformerModel.h
@@ -15,14 +15,14 @@
 
 class UMLDeformerAsset;
 class UMLDeformerVizSettings;
+class UMLDeformerModelInstance;
+class UMLDeformerComponent;
 class UAnimSequence;
 class UGeometryCache;
 class UNeuralNetwork;
 
 namespace UE::MLDeformer
 {
-	class FMLDeformerModelInstance;
-
 	/** The vertex map on the GPU. */
 	class FVertexMapBuffer : public FVertexBufferWithSRV
 	{
@@ -64,7 +64,7 @@ public:
 	virtual void Init(UMLDeformerAsset* InDeformerAsset);
 	virtual void InitGPUData();
 	virtual UMLDeformerInputInfo* CreateInputInfo();
-	virtual UE::MLDeformer::FMLDeformerModelInstance* CreateModelInstance();
+	virtual UMLDeformerModelInstance* CreateModelInstance(UMLDeformerComponent* Component);
 #if WITH_EDITORONLY_DATA
 	virtual bool HasTrainingGroundTruth() const { return false; }
 	virtual void SampleGroundTruthPositions(float SampleTime, TArray<FVector3f>& OutPositions) {}
diff --git a/Engine/Plugins/Experimental/MLDeformer/MLDeformerFramework/Source/MLDeformerFramework/Public/MLDeformerModelInstance.h b/Engine/Plugins/Experimental/MLDeformer/MLDeformerFramework/Source/MLDeformerFramework/Public/MLDeformerModelInstance.h
--- a/Engine/Plugins/Experimental/MLDeformer/MLDeformerFramework/Source/MLDeformerFramework/Public/MLDeformerModelInstance.h
+++ b/Engine/Plugins/Experimental/MLDeformer/MLDeformerFramework/Source/MLDeformerFramework/Public/MLDeformerModelInstance.h
@@ -4,157 +4,175 @@
 #include "CoreMinimal.h"
 #include "MLDeformerModule.h"
 #include "UObject/ObjectPtr.h"
+#include "MLDeformerModelInstance.generated.h"
 
 class UMLDeformerModel;
+class UNeuralNetwork;
 class USkeletalMeshComponent;
 
-namespace UE::MLDeformer
+UCLASS()
+class MLDEFORMERFRAMEWORK_API UMLDeformerModelInstance
+	: public UObject
 {
-	class MLDEFORMERFRAMEWORK_API FMLDeformerModelInstance
-	{
-	public:
-		FMLDeformerModelInstance(UMLDeformerModel* InModel);
-		virtual ~FMLDeformerModelInstance();
+	GENERATED_BODY()
+
+public:
+	UMLDeformerModelInstance() = default;
+
+	// UObject overrides.
+	virtual void BeginDestroy() override;
+	// ~END UObject overrides.
 	
-		/**
-		 * Initialize the model instance.
-		 * This internally builds the structures to quickly grab bone transforms.
-		 * If you override this method, you probably also want to call the Init of this base class inside your overloaded method.
-		 * @param SkelMeshComponent The skeletal mesh component that we will grab bone transforms and curve values from.
-		 */
-		virtual void Init(USkeletalMeshComponent* SkelMeshComponent);
-
-		/**
-		 * Release all built lookup tables etc.
-		 * Also call this base class method if you override this method.
-		 */
-		virtual void Release();
-
-		/**
-		 * Update the deformer instance.
-		 * If you use a neural network, this should provide inputs to the neural network.
-		 * @param DeltaTime The delta time since the last frame, in seconds.
-		 */
-		virtual void Tick(float DeltaTime);
-
-		/**
-		 * Check whether the deformer is compatible with a given skeletal mesh component.
-		 * This internally also edits the value returned by GetCompatibilityErrorText().
-		 * @param InSkelMeshComponent The skeletal mesh component to check compatibility with.
-		 * @param LogIssues Set to true to automatically log any compatibility errors.
-		 * @return Returns the error string. When the returned string is empty, there were no errors and thus
-		 *         the specified skeletal mesh component is compatible.
-		 */
-		virtual FString CheckCompatibility(USkeletalMeshComponent* InSkelMeshComponent, bool LogIssues=false);
-
-		/**
-		 * Check if we are in a valid state for the deformer graph's data provider.
-		 * In the base class this will check whether the neural network is valid, if the vertex map buffer is valid
-		 * and whether the neural network instance handle is valid.
-		 * @return Returns true if the data provider in would be in a valid state, otherwise false is returned.
-		 */
-		virtual bool IsValidForDataProvider() const;
-
-		/**
-		 * Get the current component space bone transforms that we grabbed from the skeletal mesh component.
-		 * @return The cached bone transforms, in component space.
-		 */
-		const TArray<FTransform>& GetBoneTransforms() const { return BoneTransforms; }
-
-		/**
-		 * Is the deformer asset used compatible with the skeletal mesh component used during the Init call?
-		 * @return True if compatible, false if not.
-		 */
-		bool IsCompatible() const { return bIsCompatible; }
-
-		/** 
-		 * Get the compatibility error text. This will be a non-empty string in case IsCompatible() return false. 
-		 * @return Returns the string containing more information about the reasons for compatibility issues.
-		 */
-		const FString& GetCompatibilityErrorText() const { return ErrorText; }
-
-		/**
-		 * Get the skeletal mesh component we're working with.
-		 * @return A pointer to the skeletal mesh component.
-		 */
-		USkeletalMeshComponent* GetSkeletalMeshComponent() const { return SkeletalMeshComponent; }
-
-		/** Update the compatibility status, as returned by IsCompatible() and GetCompatibilityErrorText(). */
-		void UpdateCompatibilityStatus();
-
-		/** Get the model that this is an instance of. */
-		UMLDeformerModel* GetModel() const { return Model.Get(); }
-
-		/**
-		 * Get the neural network inference handle.
-		 * @return Returns the handle. This is -1 for an invalid handle.
-		 */
-		int32 GetNeuralNetworkInferenceHandle() const { return NeuralNetworkInferenceHandle; }
-
-		/**
-		 * Is running neural network inference on the CPU allowed for this instance?
-		 * @return Returns true if allowed, or false if only GPU is allowed.
-		 */
-		bool IsCPUInferenceAllowed() const { return bAllowCPU; }
-
-		/**
-		 * Specify whether running neural network inference on the CPU is allowed, or if it is only GPU.
-		 * This can be used by some automated tests that will need to run inference on the CPU to check the output results.
-		 * @param bAllowed Set to true when CPU inference is allowed, or false if only GPU inference is allowed.
-		 */
-		void SetCPUInferenceAllowed(bool bAllowed) { bAllowCPU = bAllowed; }
-
-	protected:
-		/**
-		 * Update the neural network input values directly inside its input tensor.
-		 * This will copy over the bone transforms (and possibly morph and curve weights) into this flat array.
-		 * @param InputData The buffer of floats to write the input values to.
-		 * @param NumInputFloats The number of floats of the InputData buffer. Do not write more than this number of floats to the InputData buffer.
-		 */
-		void SetNeuralNetworkInputValues(float* InputData, int64 NumInputFloats);
-
-		/**
-		 * Set the bone transformations inside a given output buffer, starting from a given StartIndex.
-		 * @param OutputBuffer The buffer we need to write the transformations to.
-		 * @param OutputBufferSize The number of floats inside the output buffer. Can be used to detect buffer overflows.
-		 * @param StartIndex The index where we want to start writing the transformations to.
-		 * @return The new start index into the output buffer, to be used when you write more data after this.
-		 */
-		int64 SetBoneTransforms(float* OutputBuffer, int64 OutputBufferSize, int64 StartIndex);
-
-		/**
-		 * Set the animation curve values inside a given output buffer, starting from a given StartIndex.
-		 * @param OutputBuffer The buffer we need to write the weights to.
-		 * @param OutputBufferSize The number of floats inside the output buffer. Can be used to detect buffer overflows.
-		 * @param StartIndex The index where we want to start writing the weights to.
-		 * @return The new start index into the output buffer, to be used when you write more data after this.
-		 */
-		int64 SetCurveValues(float* OutputBuffer, int64 OutputBufferSize, int64 StartIndex);
-
-	protected:
-		/** The ML Deformer model that this is an instance of. */
-		TObjectPtr<UMLDeformerModel> Model = nullptr;
-
-		/** The skeletal mesh component we work with. This is mainly used for compatibility checks. */
-		TObjectPtr<USkeletalMeshComponent> SkeletalMeshComponent = nullptr;
-
-		/** The cached current local space bone transforms for the current frame. */
-		TArray<FTransform> BoneTransforms;
-
-		/** Maps the ML deformer asset bone index to a skeletal mesh component bone index. */
-		TArray<int32> AssetBonesToSkelMeshMappings;
-
-		/** The compatibility error text, in case bIsCompatible is false. */
-		FString ErrorText;
-
-		/** Inference handle for neural network inference. */
-		int32 NeuralNetworkInferenceHandle = -1;
-
-		/** Allow running inference on the CPU? */
-		bool bAllowCPU = false;
-
-		/** Are the deformer asset and the used skeletal mesh component compatible? */
-		bool bIsCompatible = false;
-	};
-
-} // namespace UE::MLDeformer
+	/**
+	 * Initialize the model instance.
+	 * This internally builds the structures to quickly grab bone transforms.
+	 * If you override this method, you probably also want to call the Init of this base class inside your overloaded method.
+	 * @param SkelMeshComponent The skeletal mesh component that we will grab bone transforms and curve values from.
+	 */
+	virtual void Init(USkeletalMeshComponent* SkelMeshComponent);
+
+	/**
+	 * Release all built lookup tables etc.
+	 * Also call this base class method if you override this method.
+	 */
+	virtual void Release();
+
+	/**
+	 * Update the deformer instance.
+	 * If you use a neural network, this should provide inputs to the neural network.
+	 * @param DeltaTime The delta time since the last frame, in seconds.
+	 */
+	virtual void Tick(float DeltaTime);
+
+	/**
+	 * Check whether the deformer is compatible with a given skeletal mesh component.
+	 * This internally also edits the value returned by GetCompatibilityErrorText().
+	 * @param InSkelMeshComponent The skeletal mesh component to check compatibility with.
+	 * @param LogIssues Set to true to automatically log any compatibility errors.
+	 * @return Returns the error string. When the returned string is empty, there were no errors and thus
+	 *         the specified skeletal mesh component is compatible.
+	 */
+	virtual FString CheckCompatibility(USkeletalMeshComponent* InSkelMeshComponent, bool LogIssues=false);
+
+	/**
+	 * Check if we are in a valid state for the deformer graph's data provider.
+	 * In the base class this will check whether the neural network is valid, if the vertex map buffer is valid
+	 * and whether the neural network instance handle is valid.
+	 * @return Returns true if the data provider in would be in a valid state, otherwise false is returned.
+	 */
+	virtual bool IsValidForDataProvider() const;
+
+	/**
+	 * Get the current component space bone transforms that we grabbed from the skeletal mesh component.
+	 * @return The cached bone transforms, in component space.
+	 */
+	const TArray<FTransform>& GetBoneTransforms() const { return BoneTransforms; }
+
+	/**
+	 * Is the deformer asset used compatible with the skeletal mesh component used during the Init call?
+	 * @return True if compatible, false if not.
+	 */
+	bool IsCompatible() const { return bIsCompatible; }
+
+	/** 
+	 * Get the compatibility error text. This will be a non-empty string in case IsCompatible() return false. 
+	 * @return Returns the string containing more information about the reasons for compatibility issues.
+	 */
+	const FString& GetCompatibilityErrorText() const { return ErrorText; }
+
+	/**
+	 * Get the skeletal mesh component we're working with.
+	 * @return A pointer to the skeletal mesh component.
+	 */
+	USkeletalMeshComponent* GetSkeletalMeshComponent() const { return SkeletalMeshComponent; }
+
+	/** Update the compatibility status, as returned by IsCompatible() and GetCompatibilityErrorText(). */
+	void UpdateCompatibilityStatus();
+
+	/** Get the model that this is an instance of. */
+	UMLDeformerModel* GetModel() const { return Model.Get(); }
+
+	/** Set the deformer model that this is an instance of. */
+	void SetModel(UMLDeformerModel* InModel) { Model = InModel; }
+
+	/**
+	 * Get the neural network inference handle.
+	 * @return Returns the handle. This is -1 for an invalid handle.
+	 */
+	int32 GetNeuralNetworkInferenceHandle() const { return NeuralNetworkInferenceHandle; }
+
+	/**
+	 * Is running neural network inference on the CPU allowed for this instance?
+	 * @return Returns true if allowed, or false if only GPU is allowed.
+	 */
+	bool IsCPUInferenceAllowed() const { return bAllowCPU; }
+
+	/**
+	 * Specify whether running neural network inference on the CPU is allowed, or if it is only GPU.
+	 * This can be used by some automated tests that will need to run inference on the CPU to check the output results.
+	 * @param bAllowed Set to true when CPU inference is allowed, or false if only GPU inference is allowed.
+	 */
+	void SetCPUInferenceAllowed(bool bAllowed) { bAllowCPU = bAllowed; }
+
+protected:
+	/**
+	 * Update the neural network input values directly inside its input tensor.
+	 * This will copy over the bone transforms (and possibly morph and curve weights) into this flat array.
+	 * @param InputData The buffer of floats to write the input values to.
+	 * @param NumInputFloats The number of floats of the InputData buffer. Do not write more than this number of floats to the InputData buffer.
+	 */
+	void SetNeuralNetworkInputValues(float* InputData, int64 NumInputFloats);
+
+	/**
+	 * Set the bone transformations inside a given output buffer, starting from a given StartIndex.
+	 * @param OutputBuffer The buffer we need to write the transformations to.
+	 * @param OutputBufferSize The number of floats inside the output buffer. Can be used to detect buffer overflows.
+	 * @param StartIndex The index where we want to start writing the transformations to.
+	 * @return The new start index into the output buffer, to be used when you write more data after this.
+	 */
+	int64 SetBoneTransforms(float* OutputBuffer, int64 OutputBufferSize, int64 StartIndex);
+
+	/**
+	 * Set the animation curve values inside a given output buffer, starting from a given StartIndex.
+	 * @param OutputBuffer The buffer we need to write the weights to.
+	 * @param OutputBufferSize The number of floats inside the output buffer. Can be used to detect buffer overflows.
+	 * @param StartIndex The index where we want to start writing the weights to.
+	 * @return The new start index into the output buffer, to be used when you write more data after this.
+	 */
+	int64 SetCurveValues(float* OutputBuffer, int64 OutputBufferSize, int64 StartIndex);
+
+	/**
+	 * Updates the bone transforms array.
+	 */
+	void UpdateBoneTransforms();
+
+public:
+	/** The ML Deformer model that this is an instance of. */
+	UPROPERTY(Transient)
+	TObjectPtr<UMLDeformerModel> Model = nullptr;
+
+protected:
+	/** The skeletal mesh component we work with. This is mainly used for compatibility checks. */
+	TObjectPtr<USkeletalMeshComponent> SkeletalMeshComponent = nullptr;
+
+	/** The cached current local space bone transforms for the current frame. */
+	TArray<FTransform> TrainingBoneTransforms;
+
+	/** A temp array of bone transforms. */
+	TArray<FTransform> BoneTransforms;
+
+	/** Maps the ML deformer asset bone index to a skeletal mesh component bone index. */
+	TArray<int32> AssetBonesToSkelMeshMappings;
+
+	/** The compatibility error text, in case bIsCompatible is false. */
+	FString ErrorText;
+
+	/** Inference handle for neural network inference. */
+	int32 NeuralNetworkInferenceHandle = -1;
+
+	/** Allow running inference on the CPU? */
+	bool bAllowCPU = false;
+
+	/** Are the deformer asset and the used skeletal mesh component compatible? */
+	bool bIsCompatible = false;
+};
diff --git a/Engine/Plugins/Experimental/MLDeformer/MLDeformerFramework/Source/MLDeformerFramework/Public/MLDeformerVizSettings.h b/Engine/Plugins/Experimental/MLDeformer/MLDeformerFramework/Source/MLDeformerFramework/Public/MLDeformerVizSettings.h
--- a/Engine/Plugins/Experimental/MLDeformer/MLDeformerFramework/Source/MLDeformerFramework/Public/MLDeformerVizSettings.h
+++ b/Engine/Plugins/Experimental/MLDeformer/MLDeformerFramework/Source/MLDeformerFramework/Public/MLDeformerVizSettings.h
@@ -103,7 +103,7 @@ public:
 
 	/** The height in units to draw the labels at. */
 	UPROPERTY(EditAnywhere, Category = "Shared Settings", meta = (EditCondition = "bDrawLabels", ForceUnits="cm"))
-	float LabelHeight = 300.0f;
+	float LabelHeight = 200.0f;
 
 	/** The scale of the label text. */
 	UPROPERTY(EditAnywhere, Category = "Shared Settings", meta = (ClampMin = "0.001", EditCondition = "bDrawLabels", ForceUnits="Multiplier"))
@@ -111,7 +111,7 @@ public:
 
 	/** The spacing between meshes. */
 	UPROPERTY(EditAnywhere, Category = "Shared Settings", meta = (ClampMin = "0", ForceUnits="cm"))
-	float MeshSpacing = 300.0f;
+	float MeshSpacing = 125.0f;
 
 	/** Show the heat map? This will visualize the active areas of the deformer. */
 	UPROPERTY(EditAnywhere, Category = "Live Settings")
diff --git a/Engine/Plugins/Experimental/MLDeformer/MLDeformerFramework/Source/MLDeformerFrameworkEditor/Private/MLDeformerEditorModel.cpp b/Engine/Plugins/Experimental/MLDeformer/MLDeformerFramework/Source/MLDeformerFrameworkEditor/Private/MLDeformerEditorModel.cpp
--- a/Engine/Plugins/Experimental/MLDeformer/MLDeformerFramework/Source/MLDeformerFrameworkEditor/Private/MLDeformerEditorModel.cpp
+++ b/Engine/Plugins/Experimental/MLDeformer/MLDeformerFramework/Source/MLDeformerFrameworkEditor/Private/MLDeformerEditorModel.cpp
@@ -11,6 +11,7 @@
 #include "MLDeformerComponent.h"
 #include "MLDeformerInputInfo.h"
 #include "MLDeformerSampler.h"
+#include "MLDeformerModelInstance.h"
 #include "AnimationEditorPreviewActor.h"
 #include "AnimPreviewInstance.h"
 #include "Animation/MeshDeformer.h"
@@ -313,7 +314,7 @@ namespace UE::MLDeformer
 
 				LabelComponent->SetRelativeLocation(ActorLocation + FVector(0.0f, 0.0f, VizSettings->GetLabelHeight()) - AlignmentOffset);
 				LabelComponent->SetRelativeRotation(FQuat(FVector(0.0f, 0.0f, 1.0f), FMath::DegreesToRadians(90.0f)));
-				LabelComponent->SetRelativeScale3D(FVector(VizSettings->GetLabelScale()));
+				LabelComponent->SetRelativeScale3D(FVector(VizSettings->GetLabelScale() * 0.5f));
 
 				// Update visibility.
 				const bool bLabelIsVisible = (bDrawTrainingActors && EditorActor->IsTrainingActor()) || (bDrawTestActors && EditorActor->IsTestActor());
@@ -1108,7 +1109,7 @@ namespace UE::MLDeformer
 				USkeletalMeshComponent* SkelMeshComponent = EditorActor ? EditorActor->GetSkeletalMeshComponent() : nullptr;
 				UMLDeformerAsset* DeformerAsset = GetModel()->GetDeformerAsset();
 				EditorActor->GetMLDeformerComponent()->SetupComponent(DeformerAsset, SkelMeshComponent);
-				UE::MLDeformer::FMLDeformerModelInstance* ModelInstance = EditorActor->GetMLDeformerComponent()->GetModelInstance();
+				UMLDeformerModelInstance* ModelInstance = EditorActor->GetMLDeformerComponent()->GetModelInstance();
 				if (ModelInstance)
 				{
 					ModelInstance->UpdateCompatibilityStatus();
@@ -1198,7 +1199,7 @@ namespace UE::MLDeformer
 		const UMLDeformerComponent* DeformerComponent = EditorActor ? EditorActor->GetMLDeformerComponent() : nullptr;
 		if (DeformerComponent)
 		{
-			const FMLDeformerModelInstance* ModelInstance = DeformerComponent->GetModelInstance();
+			const UMLDeformerModelInstance* ModelInstance = DeformerComponent->GetModelInstance();
 			if (ModelInstance &&
 				ModelInstance->GetSkeletalMeshComponent() && 
 				ModelInstance->GetSkeletalMeshComponent()->SkeletalMesh &&
diff --git a/Engine/Plugins/Experimental/MLDeformer/VertexDeltaModel/Source/VertexDeltaModel/Private/Tests/VertexDeltaModelTest.cpp b/Engine/Plugins/Experimental/MLDeformer/VertexDeltaModel/Source/VertexDeltaModel/Private/Tests/VertexDeltaModelTest.cpp
--- a/Engine/Plugins/Experimental/MLDeformer/VertexDeltaModel/Source/VertexDeltaModel/Private/Tests/VertexDeltaModelTest.cpp
+++ b/Engine/Plugins/Experimental/MLDeformer/VertexDeltaModel/Source/VertexDeltaModel/Private/Tests/VertexDeltaModelTest.cpp
@@ -80,7 +80,7 @@ namespace UE::VertexDeltaModelTests
 		const int64 NumNetworkOutputs = NeuralNet->GetOutputTensor().Num();
 		const int64 ExpectedNetworkOutput = NumGeomCacheVerts * 3;	// 3 floats per vertex.
 		UTEST_EQUAL(TEXT("NeuralNet output size check"), NumNetworkOutputs, ExpectedNetworkOutput);
-		UTEST_EQUAL(TEXT("Hidden Layers check"), VertexDeltaModel->GetNumHiddenLayers(), 2);
+		UTEST_EQUAL(TEXT("Hidden Layers check"), VertexDeltaModel->GetNumHiddenLayers(), 3);
 
 		UVertexDeltaModelVizSettings* VizSettings = Cast<UVertexDeltaModelVizSettings>(VertexDeltaModel->GetVizSettings());
 		UTEST_NOT_NULL(TEXT("VizSettings check"), VizSettings);
@@ -110,7 +110,7 @@ namespace UE::VertexDeltaModelTests
 		UTEST_TRUE(TEXT("MLDeformerComponent SkelMeshComponent check"), MLDeformerComponent->GetSkeletalMeshComponent() == SkelMeshComponent);
 
 		// Create and init the model instance.
-		FMLDeformerModelInstance* ModelInstance = MLDeformerComponent->GetModelInstance();
+		UMLDeformerModelInstance* ModelInstance = MLDeformerComponent->GetModelInstance();
 		UTEST_NOT_NULL(TEXT("CreateModelInstance check"), ModelInstance);
 		UTEST_TRUE(TEXT("ModelInstance GetModel check"), ModelInstance->GetModel() == VertexDeltaModel);
 		UTEST_TRUE(TEXT("ModelInstance compatible check"), ModelInstance->IsCompatible());
diff --git a/Engine/Plugins/Experimental/MLDeformer/VertexDeltaModel/Source/VertexDeltaModel/Public/VertexDeltaModel.h b/Engine/Plugins/Experimental/MLDeformer/VertexDeltaModel/Source/VertexDeltaModel/Public/VertexDeltaModel.h
--- a/Engine/Plugins/Experimental/MLDeformer/VertexDeltaModel/Source/VertexDeltaModel/Public/VertexDeltaModel.h
+++ b/Engine/Plugins/Experimental/MLDeformer/VertexDeltaModel/Source/VertexDeltaModel/Public/VertexDeltaModel.h
@@ -69,7 +69,7 @@ public:
 
 	/** The number of iterations to train the model for. */
 	UPROPERTY(EditAnywhere, BlueprintReadWrite, Category = "Training Settings", meta = (ClampMin = "1"))
-	int32 NumIterations = 500000;
+	int32 NumIterations = 10000;
 
 	/** The number of frames per batch when training the model. */
 	UPROPERTY(EditAnywhere, BlueprintReadWrite, Category = "Training Settings", meta = (ClampMin = "1"))
