[MLDeformer] Some small comment and localization namespace fixes.

diff --git a/Engine/Plugins/Experimental/MLDeformer/LegacyVertexDeltaModel/Source/LegacyVertexDeltaModel/Public/LegacyVertexDeltaModelModule.h b/Engine/Plugins/Experimental/MLDeformer/LegacyVertexDeltaModel/Source/LegacyVertexDeltaModel/Public/LegacyVertexDeltaModelModule.h
--- a/Engine/Plugins/Experimental/MLDeformer/LegacyVertexDeltaModel/Source/LegacyVertexDeltaModel/Public/LegacyVertexDeltaModelModule.h
+++ b/Engine/Plugins/Experimental/MLDeformer/LegacyVertexDeltaModel/Source/LegacyVertexDeltaModel/Public/LegacyVertexDeltaModelModule.h
@@ -17,4 +17,4 @@ namespace UE::LegacyVertexDeltaModel
 		virtual void StartupModule() override;
 		// ~END IModuleInterface overrides.
 	};
-}	// namespace VertexDeltaModel
+}	// namespace UE::LegacyVertexDeltaModel
diff --git a/Engine/Plugins/Experimental/MLDeformer/LegacyVertexDeltaModel/Source/LegacyVertexDeltaModelEditor/Private/LegacyVertexDeltaEditorModel.cpp b/Engine/Plugins/Experimental/MLDeformer/LegacyVertexDeltaModel/Source/LegacyVertexDeltaModelEditor/Private/LegacyVertexDeltaEditorModel.cpp
--- a/Engine/Plugins/Experimental/MLDeformer/LegacyVertexDeltaModel/Source/LegacyVertexDeltaModelEditor/Private/LegacyVertexDeltaEditorModel.cpp
+++ b/Engine/Plugins/Experimental/MLDeformer/LegacyVertexDeltaModel/Source/LegacyVertexDeltaModelEditor/Private/LegacyVertexDeltaEditorModel.cpp
@@ -300,6 +300,6 @@ namespace UE::LegacyVertexDeltaModel
 	{
 		return FString(TEXT("/LegacyVertexDeltaModel/Deformers/DG_LegacyVertexDeltaModel_HeatMap.DG_LegacyVertexDeltaModel_HeatMap"));
 	}
-}	// namespace UE::VertexDeltaModel
+}	// namespace UE::LegacyVertexDeltaModel
 
 #undef LOCTEXT_NAMESPACE
diff --git a/Engine/Plugins/Experimental/MLDeformer/MLDeformerFramework/Source/MLDeformerFramework/Private/MLDeformerModule.cpp b/Engine/Plugins/Experimental/MLDeformer/MLDeformerFramework/Source/MLDeformerFramework/Private/MLDeformerModule.cpp
--- a/Engine/Plugins/Experimental/MLDeformer/MLDeformerFramework/Source/MLDeformerFramework/Private/MLDeformerModule.cpp
+++ b/Engine/Plugins/Experimental/MLDeformer/MLDeformerFramework/Source/MLDeformerFramework/Private/MLDeformerModule.cpp
@@ -9,7 +9,7 @@
 // Log category.
 MLDEFORMERFRAMEWORK_API DEFINE_LOG_CATEGORY(LogMLDeformer)
 
-#define LOCTEXT_NAMESPACE "VertexDeltaModelModule"
+#define LOCTEXT_NAMESPACE "MLDeformerModule"
 
 IMPLEMENT_MODULE(UE::MLDeformer::FMLDeformerModule, MLDeformerFramework)
 
diff --git a/Engine/Plugins/Experimental/MLDeformer/MLDeformerFramework/Source/MLDeformerFrameworkEditor/Private/MLDeformerGeomCacheSampler.cpp b/Engine/Plugins/Experimental/MLDeformer/MLDeformerFramework/Source/MLDeformerFrameworkEditor/Private/MLDeformerGeomCacheSampler.cpp
--- a/Engine/Plugins/Experimental/MLDeformer/MLDeformerFramework/Source/MLDeformerFrameworkEditor/Private/MLDeformerGeomCacheSampler.cpp
+++ b/Engine/Plugins/Experimental/MLDeformer/MLDeformerFramework/Source/MLDeformerFrameworkEditor/Private/MLDeformerGeomCacheSampler.cpp
@@ -139,4 +139,4 @@ namespace UE::MLDeformer
 	{
 		return GeometryCacheComponent.Get() ? GeometryCacheComponent->GetTimeAtFrame(InAnimFrameIndex) : 0.0f;
 	}
-}	// namespace UE::VertexDeltaModel
+}	// namespace UE::MLDeformer
diff --git a/Engine/Plugins/Experimental/MLDeformer/MLDeformerFramework/Source/MLDeformerFrameworkEditor/Private/MLDeformerSampler.cpp b/Engine/Plugins/Experimental/MLDeformer/MLDeformerFramework/Source/MLDeformerFrameworkEditor/Private/MLDeformerSampler.cpp
--- a/Engine/Plugins/Experimental/MLDeformer/MLDeformerFramework/Source/MLDeformerFrameworkEditor/Private/MLDeformerSampler.cpp
+++ b/Engine/Plugins/Experimental/MLDeformer/MLDeformerFramework/Source/MLDeformerFrameworkEditor/Private/MLDeformerSampler.cpp
@@ -239,4 +239,4 @@ namespace UE::MLDeformer
 		return NumBytes;
 	}
 
-}	// namespace UE::VertexDeltaModel
+}	// namespace UE::MLDeformer
diff --git a/Engine/Plugins/Experimental/MLDeformer/MLDeformerFramework/Source/MLDeformerFrameworkEditor/Public/MLDeformerModelDetails.h b/Engine/Plugins/Experimental/MLDeformer/MLDeformerFramework/Source/MLDeformerFrameworkEditor/Public/MLDeformerModelDetails.h
--- a/Engine/Plugins/Experimental/MLDeformer/MLDeformerFramework/Source/MLDeformerFrameworkEditor/Public/MLDeformerModelDetails.h
+++ b/Engine/Plugins/Experimental/MLDeformer/MLDeformerFramework/Source/MLDeformerFrameworkEditor/Public/MLDeformerModelDetails.h
@@ -51,4 +51,4 @@ namespace UE::MLDeformer
 		IDetailCategoryBuilder* InputOutputCategoryBuilder = nullptr;
 		IDetailCategoryBuilder* SettingsCategoryBuilder = nullptr;
 	};
-}	// namespace UE::VertexDeltaModel
+}	// namespace UE::MLDeformer
diff --git a/Engine/Plugins/Experimental/MLDeformer/MLDeformerFramework/Source/MLDeformerFrameworkEditor/Public/MLDeformerSampler.h b/Engine/Plugins/Experimental/MLDeformer/MLDeformerFramework/Source/MLDeformerFrameworkEditor/Public/MLDeformerSampler.h
--- a/Engine/Plugins/Experimental/MLDeformer/MLDeformerFramework/Source/MLDeformerFrameworkEditor/Public/MLDeformerSampler.h
+++ b/Engine/Plugins/Experimental/MLDeformer/MLDeformerFramework/Source/MLDeformerFrameworkEditor/Public/MLDeformerSampler.h
@@ -6,7 +6,6 @@
 #include "MLDeformerModel.h"
 #include "UObject/ObjectPtr.h"
 
-class UVertexDeltaModel;
 class UDebugSkelMeshComponent;
 class UGeometryCacheComponent;
 class UWorld;
@@ -119,4 +118,4 @@ namespace UE::MLDeformer
 		/** The vertex delta space (pre or post skinning) used when calculating the deltas. */
 		EVertexDeltaSpace VertexDeltaSpace = EVertexDeltaSpace::PreSkinning;
 	};
-}	// namespace UE::VertexDeltaModel
+}	// namespace UE::MLDeformer
diff --git a/Engine/Plugins/Experimental/MLDeformer/MLDeformerFramework/Source/MLDeformerFrameworkEditor/Public/MLDeformerVizSettingsDetails.h b/Engine/Plugins/Experimental/MLDeformer/MLDeformerFramework/Source/MLDeformerFrameworkEditor/Public/MLDeformerVizSettingsDetails.h
--- a/Engine/Plugins/Experimental/MLDeformer/MLDeformerFramework/Source/MLDeformerFrameworkEditor/Public/MLDeformerVizSettingsDetails.h
+++ b/Engine/Plugins/Experimental/MLDeformer/MLDeformerFramework/Source/MLDeformerFrameworkEditor/Public/MLDeformerVizSettingsDetails.h
@@ -54,4 +54,4 @@ namespace UE::MLDeformer
 		IDetailCategoryBuilder* LiveSettingsCategory = nullptr;
 		IDetailCategoryBuilder* TrainingMeshesCategoryBuilder = nullptr;
 	};
-}	// namespace UE::VertexDeltaModel
+}	// namespace UE::MLDeformer
diff --git a/Engine/Plugins/Experimental/MLDeformer/VertexDeltaModel/Source/VertexDeltaModelEditor/Public/VertexDeltaEditorModel.h b/Engine/Plugins/Experimental/MLDeformer/VertexDeltaModel/Source/VertexDeltaModelEditor/Public/VertexDeltaEditorModel.h
--- a/Engine/Plugins/Experimental/MLDeformer/VertexDeltaModel/Source/VertexDeltaModelEditor/Public/VertexDeltaEditorModel.h
+++ b/Engine/Plugins/Experimental/MLDeformer/VertexDeltaModel/Source/VertexDeltaModelEditor/Public/VertexDeltaEditorModel.h
@@ -61,9 +61,5 @@ namespace UE::VertexDeltaModel
 
 	protected:
 		void CreateGeomCacheActor(UWorld* World, int32 ActorID, const FName& Name, UGeometryCache* GeomCache, FLinearColor LabelColor, FLinearColor WireframeColor, const FText& LabelText, bool bIsTrainingActor);
-
-	protected:
-		FVector VertexDeltaMeanBackup = FVector::ZeroVector;
-		FVector VertexDeltaScaleBackup = FVector::OneVector;
 	};
 }	// namespace UE::VertexDeltaModel
