[MLDeformer] Compilation and test fail fixes.

diff --git a/Engine/Plugins/Experimental/MLDeformer/MLDeformerFramework/Source/MLDeformerFramework/Private/Tests/MLDeformerTestModel.cpp b/Engine/Plugins/Experimental/MLDeformer/MLDeformerFramework/Source/MLDeformerFramework/Private/Tests/MLDeformerTestModel.cpp
--- a/Engine/Plugins/Experimental/MLDeformer/MLDeformerFramework/Source/MLDeformerFramework/Private/Tests/MLDeformerTestModel.cpp
+++ b/Engine/Plugins/Experimental/MLDeformer/MLDeformerFramework/Source/MLDeformerFramework/Private/Tests/MLDeformerTestModel.cpp
@@ -14,7 +14,9 @@ FString UMLDeformerTestModel::GetDisplayName() const
 	return FString("Test Model");
 }
 
-void UMLDeformerTestModel::UpdateNumTargetMeshVertices()
-{
-	NumTargetMeshVerts = NumBaseMeshVerts;
-}
+#if WITH_EDITOR
+	void UMLDeformerTestModel::UpdateNumTargetMeshVertices()
+	{
+		NumTargetMeshVerts = NumBaseMeshVerts;
+	}
+#endif
diff --git a/Engine/Plugins/Experimental/MLDeformer/MLDeformerFramework/Source/MLDeformerFramework/Private/Tests/MLDeformerTestModel.h b/Engine/Plugins/Experimental/MLDeformer/MLDeformerFramework/Source/MLDeformerFramework/Private/Tests/MLDeformerTestModel.h
--- a/Engine/Plugins/Experimental/MLDeformer/MLDeformerFramework/Source/MLDeformerFramework/Private/Tests/MLDeformerTestModel.h
+++ b/Engine/Plugins/Experimental/MLDeformer/MLDeformerFramework/Source/MLDeformerFramework/Private/Tests/MLDeformerTestModel.h
@@ -18,7 +18,9 @@ class UMLDeformerTestModel
 public:
 	UMLDeformerModelInstance* CreateModelInstance(UMLDeformerComponent* Component) override;
 	FString GetDisplayName() const override;
+#if WITH_EDITOR
 	void UpdateNumTargetMeshVertices() override;
+#endif
 };
 
 UCLASS()
diff --git a/Engine/Plugins/Experimental/MLDeformer/MLDeformerFramework/Source/MLDeformerFramework/Public/MLDeformerModel.h b/Engine/Plugins/Experimental/MLDeformer/MLDeformerFramework/Source/MLDeformerFramework/Public/MLDeformerModel.h
--- a/Engine/Plugins/Experimental/MLDeformer/MLDeformerFramework/Source/MLDeformerFramework/Public/MLDeformerModel.h
+++ b/Engine/Plugins/Experimental/MLDeformer/MLDeformerFramework/Source/MLDeformerFramework/Public/MLDeformerModel.h
@@ -65,6 +65,7 @@ public:
 	virtual void InitGPUData();
 	virtual UMLDeformerInputInfo* CreateInputInfo();
 	virtual UMLDeformerModelInstance* CreateModelInstance(UMLDeformerComponent* Component);
+	virtual FString GetDisplayName() const PURE_VIRTUAL(UMLDeformerModel::GetDisplayName, return FString(););
 #if WITH_EDITORONLY_DATA
 	virtual bool HasTrainingGroundTruth() const { return false; }
 	virtual void SampleGroundTruthPositions(float SampleTime, TArray<FVector3f>& OutPositions) {}
@@ -72,14 +73,10 @@ public:
 #if WITH_EDITOR
 	virtual void UpdateNumBaseMeshVertices();
 	virtual void UpdateCachedNumVertices();
-
+	virtual void UpdateNumTargetMeshVertices() PURE_VIRTUAL(UMLDeformerModel::GetNumTargetMeshVertices; mTargetMeshVerts = 0;);
 	static int32 ExtractNumImportedSkinnedVertices(const USkeletalMesh* SkeletalMesh);
 #endif
 
-	// Methods you are required to override.
-	virtual FString GetDisplayName() const PURE_VIRTUAL(UMLDeformerModel::GetDisplayName, return FString(););
-	virtual void UpdateNumTargetMeshVertices() PURE_VIRTUAL(UMLDeformerModel::GetNumTargetMeshVertices; mTargetMeshVerts=0;);
-
 	// UObject overrides.
 	virtual void Serialize(FArchive& Archive) override;
 	virtual void PostLoad() override;
diff --git a/Engine/Plugins/Experimental/MLDeformer/MLDeformerFramework/Source/MLDeformerFrameworkEditor/Private/MLDeformerEditorModel.cpp b/Engine/Plugins/Experimental/MLDeformer/MLDeformerFramework/Source/MLDeformerFrameworkEditor/Private/MLDeformerEditorModel.cpp
--- a/Engine/Plugins/Experimental/MLDeformer/MLDeformerFramework/Source/MLDeformerFrameworkEditor/Private/MLDeformerEditorModel.cpp
+++ b/Engine/Plugins/Experimental/MLDeformer/MLDeformerFramework/Source/MLDeformerFrameworkEditor/Private/MLDeformerEditorModel.cpp
@@ -762,7 +762,6 @@ namespace UE::MLDeformer
 
 		BoneNames.Reset();
 		BoneNameStrings.Reset();
-		BoneNames.Reset();
 		CurveNames.Reset();
 
 		InputInfo->SetNumBaseVertices(Model->GetNumBaseMeshVerts());
