[MLDeformer] Fixed some static code analysis issues.

diff --git a/Engine/Plugins/Experimental/MLDeformer/LegacyVertexDeltaModel/Source/LegacyVertexDeltaModel/Private/LegacyVertexDeltaModel.cpp b/Engine/Plugins/Experimental/MLDeformer/LegacyVertexDeltaModel/Source/LegacyVertexDeltaModel/Private/LegacyVertexDeltaModel.cpp
--- a/Engine/Plugins/Experimental/MLDeformer/LegacyVertexDeltaModel/Source/LegacyVertexDeltaModel/Private/LegacyVertexDeltaModel.cpp
+++ b/Engine/Plugins/Experimental/MLDeformer/LegacyVertexDeltaModel/Source/LegacyVertexDeltaModel/Private/LegacyVertexDeltaModel.cpp
@@ -1,6 +1,4 @@
 // Copyright Epic Games, Inc. All Rights Reserved.
-#pragma once
-
 #include "LegacyVertexDeltaModel.h"
 #include "LegacyVertexDeltaModelVizSettings.h"
 #include "MLDeformerGeomCacheHelpers.h"
diff --git a/Engine/Plugins/Experimental/MLDeformer/MLDeformerFramework/Source/MLDeformerFrameworkEditor/Private/MLDeformerEditorMode.cpp b/Engine/Plugins/Experimental/MLDeformer/MLDeformerFramework/Source/MLDeformerFrameworkEditor/Private/MLDeformerEditorMode.cpp
--- a/Engine/Plugins/Experimental/MLDeformer/MLDeformerFramework/Source/MLDeformerFrameworkEditor/Private/MLDeformerEditorMode.cpp
+++ b/Engine/Plugins/Experimental/MLDeformer/MLDeformerFramework/Source/MLDeformerFrameworkEditor/Private/MLDeformerEditorMode.cpp
@@ -27,12 +27,15 @@ namespace UE::MLDeformer
 		UMLDeformerVizSettings* VizSettings = ActiveModel->GetModel()->GetVizSettings();
 		for (FMLDeformerEditorActor* EditorActor : ActiveModel->GetEditorActors())
 		{
-			bool bIncludeActor = (VizSettings->GetVisualizationMode() == EMLDeformerVizMode::TrainingData) && EditorActor->IsTrainingActor() && EditorActor->IsVisible();
-			bIncludeActor |= (VizSettings->GetVisualizationMode() == EMLDeformerVizMode::TestData) && EditorActor->IsTestActor() && EditorActor->IsVisible();
-
-			if (bIncludeActor)
+			if (EditorActor)
 			{
-				Box += EditorActor->GetBoundingBox();
+				bool bIncludeActor = (VizSettings->GetVisualizationMode() == EMLDeformerVizMode::TrainingData) && EditorActor->IsTrainingActor() && EditorActor->IsVisible();
+				bIncludeActor |= (VizSettings->GetVisualizationMode() == EMLDeformerVizMode::TestData) && EditorActor->IsTestActor() && EditorActor->IsVisible();
+
+				if (bIncludeActor)
+				{
+					Box += EditorActor->GetBoundingBox();
+				}
 			}
 		}
 	
diff --git a/Engine/Plugins/Experimental/MLDeformer/MLDeformerFramework/Source/MLDeformerFrameworkEditor/Private/MLDeformerEditorModel.cpp b/Engine/Plugins/Experimental/MLDeformer/MLDeformerFramework/Source/MLDeformerFrameworkEditor/Private/MLDeformerEditorModel.cpp
--- a/Engine/Plugins/Experimental/MLDeformer/MLDeformerFramework/Source/MLDeformerFrameworkEditor/Private/MLDeformerEditorModel.cpp
+++ b/Engine/Plugins/Experimental/MLDeformer/MLDeformerFramework/Source/MLDeformerFrameworkEditor/Private/MLDeformerEditorModel.cpp
@@ -207,7 +207,7 @@ namespace UE::MLDeformer
 		// Set the default mesh translation offsets for our ground truth actors.
 		for (FMLDeformerEditorActor* EditorActor : EditorActors)
 		{
-			if (EditorActor->IsGroundTruthActor())
+			if (EditorActor && EditorActor->IsGroundTruthActor())
 			{			
 				// The mesh offset factor basically just offsets the actor position by a given factor.
 				// The amount the actor is moved from the origin is: (MeshSpacing * MeshOffsetFactor).
@@ -229,10 +229,13 @@ namespace UE::MLDeformer
 		UWorld* World = PreviewScene->GetWorld();
 		for (FMLDeformerEditorActor* EditorActor : EditorActors)
 		{
-			World->RemoveActor(EditorActor->GetActor(), true);
-			if (EditorActor->GetCanDestroyActor())
+			if (EditorActor)
 			{
-				EditorActor->GetActor()->Destroy();
+				World->RemoveActor(EditorActor->GetActor(), true);
+				if (EditorActor->GetCanDestroyActor())
+				{
+					EditorActor->GetActor()->Destroy();
+				}
 			}
 		}
 
@@ -254,9 +257,9 @@ namespace UE::MLDeformer
 
 	void FMLDeformerEditorModel::DeleteEditorActors()
 	{
-		for (FMLDeformerEditorActor* Actor : EditorActors)
+		for (FMLDeformerEditorActor* EditorActor : EditorActors)
 		{
-			delete Actor;
+			delete EditorActor;
 		}
 		EditorActors.Empty();
 	}
@@ -265,7 +268,7 @@ namespace UE::MLDeformer
 	{
 		for (FMLDeformerEditorActor* EditorActor : EditorActors)
 		{
-			if (EditorActor->GetTypeID() == ActorTypeID)
+			if (EditorActor && EditorActor->GetTypeID() == ActorTypeID)
 			{
 				return EditorActor;
 			}
@@ -301,7 +304,7 @@ namespace UE::MLDeformer
 		for (FMLDeformerEditorActor* EditorActor : EditorActors)
 		{
 			UTextRenderComponent* LabelComponent = EditorActor->GetLabelComponent();
-			if (LabelComponent == nullptr)
+			if (LabelComponent == nullptr || EditorActor == nullptr)
 			{
 				continue;
 			}
@@ -338,9 +341,12 @@ namespace UE::MLDeformer
 		const FVector MeshSpacingVector = Model->GetVizSettings()->GetMeshSpacingOffsetVector();
 		for (FMLDeformerEditorActor* EditorActor : EditorActors)
 		{
-			FTransform Transform = EditorActor->IsGroundTruthActor() ? Model->GetAlignmentTransform() : FTransform::Identity;
-			Transform.AddToTranslation(MeshSpacingVector * EditorActor->GetMeshOffsetFactor());
-			EditorActor->GetActor()->SetActorTransform(Transform);
+			if (EditorActor)
+			{
+				FTransform Transform = EditorActor->IsGroundTruthActor() ? Model->GetAlignmentTransform() : FTransform::Identity;
+				Transform.AddToTranslation(MeshSpacingVector * EditorActor->GetMeshOffsetFactor());
+				EditorActor->GetActor()->SetActorTransform(Transform);
+			}
 		}
 	}
 
@@ -351,8 +357,11 @@ namespace UE::MLDeformer
 		const bool bShowTestData = (VizSettings->GetVisualizationMode() == EMLDeformerVizMode::TestData);
 		for (FMLDeformerEditorActor* EditorActor : EditorActors)
 		{
-			const bool bIsVisible = (EditorActor->IsTestActor() && bShowTestData) || (EditorActor->IsTrainingActor() && bShowTrainingData);
-			EditorActor->SetVisibility(bIsVisible);
+			if (EditorActor)
+			{
+				const bool bIsVisible = (EditorActor->IsTestActor() && bShowTestData) || (EditorActor->IsTrainingActor() && bShowTrainingData);
+				EditorActor->SetVisibility(bIsVisible);
+			}
 		}
 	}
 
@@ -433,7 +442,7 @@ namespace UE::MLDeformer
 		{
 			for (FMLDeformerEditorActor* EditorActor : EditorActors)
 			{
-				if (EditorActor->IsTrainingActor())
+				if (EditorActor && EditorActor->IsTrainingActor())
 				{
 					if (Model->HasTrainingGroundTruth())
 					{
@@ -448,7 +457,7 @@ namespace UE::MLDeformer
 		{
 			for (FMLDeformerEditorActor* EditorActor : EditorActors)
 			{
-				if (EditorActor->IsTestActor())
+				if (EditorActor && EditorActor->IsTestActor())
 				{
 					if (Model->GetVizSettings()->HasTestGroundTruth())
 					{
@@ -560,7 +569,7 @@ namespace UE::MLDeformer
 		const bool bMustPause = (BaseTestActor && BaseTestActor->GetSkeletalMeshComponent()) ? !BaseTestActor->GetSkeletalMeshComponent()->bPauseAnims : false;
 		for (FMLDeformerEditorActor* EditorActor : EditorActors)
 		{
-			if (EditorActor->IsTestActor())
+			if (EditorActor && EditorActor->IsTestActor())
 			{
 				EditorActor->Pause(bMustPause);
 			}
@@ -608,7 +617,7 @@ namespace UE::MLDeformer
 		const float Speed = VizSettings->GetAnimPlaySpeed();
 		for (FMLDeformerEditorActor* EditorActor : EditorActors)
 		{
-			if (EditorActor->IsTestActor()) // Only do test actors, no training actors.
+			if (EditorActor && EditorActor->IsTestActor()) // Only do test actors, no training actors.
 			{
 				EditorActor->SetPlaySpeed(Speed);
 			}
@@ -1216,7 +1225,7 @@ namespace UE::MLDeformer
 		for (FMLDeformerEditorActor* EditorActor : EditorActors)
 		{
 			USkeletalMeshComponent* SkelMeshComponent = EditorActor->GetSkeletalMeshComponent();
-			if (EditorActor->GetMLDeformerComponent() == nullptr || SkelMeshComponent == nullptr)
+			if (EditorActor == nullptr || EditorActor->GetMLDeformerComponent() == nullptr || SkelMeshComponent == nullptr)
 			{
 				continue;
 			}
@@ -1235,7 +1244,7 @@ namespace UE::MLDeformer
 	{
 		for (FMLDeformerEditorActor* EditorActor : EditorActors)
 		{
-			if (EditorActor->GetMLDeformerComponent())
+			if (EditorActor && EditorActor->GetMLDeformerComponent())
 			{
 				USkeletalMeshComponent* SkelMeshComponent = EditorActor->GetSkeletalMeshComponent();
 				UMLDeformerAsset* DeformerAsset = Model->GetDeformerAsset();
diff --git a/Engine/Plugins/Experimental/MLDeformer/VertexDeltaModel/Source/VertexDeltaModel/Private/VertexDeltaModel.cpp b/Engine/Plugins/Experimental/MLDeformer/VertexDeltaModel/Source/VertexDeltaModel/Private/VertexDeltaModel.cpp
--- a/Engine/Plugins/Experimental/MLDeformer/VertexDeltaModel/Source/VertexDeltaModel/Private/VertexDeltaModel.cpp
+++ b/Engine/Plugins/Experimental/MLDeformer/VertexDeltaModel/Source/VertexDeltaModel/Private/VertexDeltaModel.cpp
@@ -1,5 +1,4 @@
 // Copyright Epic Games, Inc. All Rights Reserved.
-#pragma once
 
 #include "VertexDeltaModel.h"
 #include "VertexDeltaModelVizSettings.h"
