Fix Clang Editor compile due to missing dependency in light of CL# 20282438
#rnx

diff --git a/Engine/Plugins/Experimental/MLDeformer/MLDeformerFramework/Source/MLDeformerFrameworkEditor/MLDeformerFrameworkEditor.Build.cs b/Engine/Plugins/Experimental/MLDeformer/MLDeformerFramework/Source/MLDeformerFrameworkEditor/MLDeformerFrameworkEditor.Build.cs
--- a/Engine/Plugins/Experimental/MLDeformer/MLDeformerFramework/Source/MLDeformerFrameworkEditor/MLDeformerFrameworkEditor.Build.cs
+++ b/Engine/Plugins/Experimental/MLDeformer/MLDeformerFramework/Source/MLDeformerFrameworkEditor/MLDeformerFrameworkEditor.Build.cs
@@ -30,6 +30,7 @@ namespace UnrealBuildTool.Rules
 					"MLDeformerFramework",
 					"Projects",
 					"PropertyEditor",
+					"AnimationEditMode",
 					"AnimGraph",
 					"ToolWidgets",
 					"GeometryCache",
