MLDeformer: Support component tags on MLDeformer data interface.
This allows us to map from multiple MLDeformers in a graph to multiple MLDeformer components.
We may remove this when we have the full Optimus binding system in place soon.

diff --git a/Engine/Plugins/Experimental/MLDeformer/MLDeformerFramework/Source/MLDeformerFramework/Private/MLDeformerGraphDataInterface.cpp b/Engine/Plugins/Experimental/MLDeformer/MLDeformerFramework/Source/MLDeformerFramework/Private/MLDeformerGraphDataInterface.cpp
--- a/Engine/Plugins/Experimental/MLDeformer/MLDeformerFramework/Source/MLDeformerFramework/Private/MLDeformerGraphDataInterface.cpp
+++ b/Engine/Plugins/Experimental/MLDeformer/MLDeformerFramework/Source/MLDeformerFramework/Private/MLDeformerGraphDataInterface.cpp
@@ -47,6 +47,25 @@ MLDEFORMER_GRAPH_IMPLEMENT_BASICS(
 	TEXT("#include \"/Plugin/MLDeformerFramework/Private/MLDeformerGraphDataInterface.ush\"\n"),
 	TEXT("ML Deformer"))
 
+UMLDeformerComponent* UMLDeformerGraphDataProvider::GetDeformerComponent() const
+{
+	UMLDeformerComponent* DeformerComponent = nullptr;
+	if (ComponentTag != NAME_None)
+	{
+		TArray<UActorComponent*> Components = SkeletalMeshComponent->GetOwner()->GetComponentsByTag(UMLDeformerComponent::StaticClass(), ComponentTag);
+		if (Components.Num())
+		{
+			DeformerComponent = Cast<UMLDeformerComponent>(Components[0]);
+		}
+	}
+	else
+	{
+		DeformerComponent = SkeletalMeshComponent->GetOwner()->FindComponentByClass<UMLDeformerComponent>();
+	}
+	
+	return DeformerComponent;
+}
+
 bool UMLDeformerGraphDataProvider::IsValid() const
 {
 	if (SkeletalMeshComponent == nullptr || SkeletalMeshComponent->MeshObject == nullptr)
@@ -54,7 +73,7 @@ bool UMLDeformerGraphDataProvider::IsValid() const
 		return false;
 	}
 	
-	UMLDeformerComponent* DeformerComponent = SkeletalMeshComponent->GetOwner()->FindComponentByClass<UMLDeformerComponent>();
+	UMLDeformerComponent* DeformerComponent = GetDeformerComponent();
 	if (DeformerComponent == nullptr || DeformerComponent->GetDeformerAsset() == nullptr)
 	{
 		return false;
diff --git a/Engine/Plugins/Experimental/MLDeformer/MLDeformerFramework/Source/MLDeformerFramework/Public/MLDeformerGraphDataInterface.h b/Engine/Plugins/Experimental/MLDeformer/MLDeformerFramework/Source/MLDeformerFramework/Public/MLDeformerGraphDataInterface.h
--- a/Engine/Plugins/Experimental/MLDeformer/MLDeformerFramework/Source/MLDeformerFramework/Public/MLDeformerGraphDataInterface.h
+++ b/Engine/Plugins/Experimental/MLDeformer/MLDeformerFramework/Source/MLDeformerFramework/Public/MLDeformerGraphDataInterface.h
@@ -64,12 +64,12 @@ class UMLDeformerModel;
 	{ \
 		DataProviderClassName* Provider = NewObject<DataProviderClassName>(); \
 		Provider->SkeletalMeshComponent = Cast<USkeletalMeshComponent>(InBinding); \
+		Provider->ComponentTag = ComponentTag; \
 		return Provider; \
 	} \
 	FComputeDataProviderRenderProxy* DataProviderClassName::GetRenderProxy() \
 	{ \
-		UMLDeformerComponent* DeformerComponent = SkeletalMeshComponent->GetOwner()->FindComponentByClass<UMLDeformerComponent>(); \
-		return new DataProviderProxyClassName(SkeletalMeshComponent, DeformerComponent); \
+		return new DataProviderProxyClassName(SkeletalMeshComponent, GetDeformerComponent()); \
 	}
 
 #if WITH_EDITORONLY_DATA
@@ -98,6 +98,10 @@ public:
 	virtual void GetHLSL(FString& OutHLSL) const override;
 	virtual UComputeDataProvider* CreateDataProvider(TObjectPtr<UObject> InBinding, uint64 InInputMask, uint64 InOutputMask) const override;
 	// ~END UComputeDataInterface overrides.
+
+	/** When set we bind only with to a MLDeformer Component with the same tag. */
+	UPROPERTY(EditAnywhere, Category = Binding)
+	FName ComponentTag;
 };
 
 /** Compute Framework Data Provider for MLDeformer data. */
@@ -111,10 +115,15 @@ public:
 	UPROPERTY(EditAnywhere, BlueprintReadWrite, Category = Binding)
 	TObjectPtr<USkeletalMeshComponent> SkeletalMeshComponent = nullptr;
 
+	FName ComponentTag;
+
 	// UComputeDataProvider overrides.
 	virtual bool IsValid() const override;
 	virtual FComputeDataProviderRenderProxy* GetRenderProxy() override;
 	// ~END UComputeDataProvider overrides.
+
+protected:
+	UMLDeformerComponent* GetDeformerComponent() const;
 };
 
 namespace UE::MLDeformer
