Fix static analysis warning (pointer dereference before null check)

diff --git a/Engine/Plugins/Experimental/MLDeformer/MLDeformerFramework/Source/MLDeformerFrameworkEditor/Private/MLDeformerEditorModel.cpp b/Engine/Plugins/Experimental/MLDeformer/MLDeformerFramework/Source/MLDeformerFrameworkEditor/Private/MLDeformerEditorModel.cpp
--- a/Engine/Plugins/Experimental/MLDeformer/MLDeformerFramework/Source/MLDeformerFrameworkEditor/Private/MLDeformerEditorModel.cpp
+++ b/Engine/Plugins/Experimental/MLDeformer/MLDeformerFramework/Source/MLDeformerFrameworkEditor/Private/MLDeformerEditorModel.cpp
@@ -303,8 +303,13 @@ namespace UE::MLDeformer
 
 		for (FMLDeformerEditorActor* EditorActor : EditorActors)
 		{
+			if (EditorActor == nullptr)
+			{
+				continue;
+			}
+
 			UTextRenderComponent* LabelComponent = EditorActor->GetLabelComponent();
-			if (LabelComponent == nullptr || EditorActor == nullptr)
+			if (LabelComponent == nullptr)
 			{
 				continue;
 			}
@@ -1224,8 +1229,13 @@ namespace UE::MLDeformer
 	{	
 		for (FMLDeformerEditorActor* EditorActor : EditorActors)
 		{
+			if (EditorActor == nullptr || EditorActor->GetMLDeformerComponent() == nullptr)
+			{
+				continue;
+			}
+
 			USkeletalMeshComponent* SkelMeshComponent = EditorActor->GetSkeletalMeshComponent();
-			if (EditorActor == nullptr || EditorActor->GetMLDeformerComponent() == nullptr || SkelMeshComponent == nullptr)
+			if (SkelMeshComponent == nullptr)
 			{
 				continue;
 			}
