[MLDeformer]

- Initial Neural Morph model plugin.
- There is now a testing frame number so you can go to an exact frame while testing.
- When changing some properties or after training, the time doesn't reset to 0.
- Fixed visibility flags of the actors, where checkbox state changes didn't do anything.
- Fixed an issue where aborting training and choosing not to use the partially trained model, still used it.
- Some small misc refactors.

diff --git a/Engine/Plugins/Experimental/MLDeformer/LegacyVertexDeltaModel/LegacyVertexDeltaModel.uplugin b/Engine/Plugins/Experimental/MLDeformer/LegacyVertexDeltaModel/LegacyVertexDeltaModel.uplugin
--- a/Engine/Plugins/Experimental/MLDeformer/LegacyVertexDeltaModel/LegacyVertexDeltaModel.uplugin
+++ b/Engine/Plugins/Experimental/MLDeformer/LegacyVertexDeltaModel/LegacyVertexDeltaModel.uplugin
@@ -2,14 +2,14 @@
 	"FileVersion": 3,
 	"Version": 1,
 	"VersionName": "0.2",
-	"FriendlyName": "LegacyVertexDeltaModel",
+	"FriendlyName": "ML Deformer Legacy Vertex Delta Model",
 	"Description": "Legacy Vertex Delta Model for the ML Deformer Framework",
 	"Category": "Animation",
 	"CreatedBy": "Epic Games, Inc.",
 	"CreatedByURL": "http://epicgames.com",
-	"DocsURL": "",
-	"MarketplaceURL": "",
-	"SupportURL": "",
+	"DocsURL": "https://docs.unrealengine.com/5.0/en-US/using-the-machine-learning-deformer-in-unreal-engine/",
+	"MarketplaceURL": "https://www.unrealengine.com/marketplace/en-US/product/ml-deformer",
+	"SupportURL": "https://forums.unrealengine.com/",
 	"CanContainContent": true,
 	"IsBetaVersion": false,
 	"IsExperimentalVersion": true,
diff --git a/Engine/Plugins/Experimental/MLDeformer/LegacyVertexDeltaModel/Source/LegacyVertexDeltaModelEditor/Private/LegacyVertexDeltaEditorModel.cpp b/Engine/Plugins/Experimental/MLDeformer/LegacyVertexDeltaModel/Source/LegacyVertexDeltaModelEditor/Private/LegacyVertexDeltaEditorModel.cpp
--- a/Engine/Plugins/Experimental/MLDeformer/LegacyVertexDeltaModel/Source/LegacyVertexDeltaModelEditor/Private/LegacyVertexDeltaEditorModel.cpp
+++ b/Engine/Plugins/Experimental/MLDeformer/LegacyVertexDeltaModel/Source/LegacyVertexDeltaModelEditor/Private/LegacyVertexDeltaEditorModel.cpp
@@ -73,7 +73,7 @@ namespace UE::LegacyVertexDeltaModel
 		{
 			if (PropertyChangedEvent.ChangeType == EPropertyChangeType::ValueSet)
 			{
-				bIsDataNormalized = false;
+				SetResamplingInputOutputsNeeded(true);
 				SampleDeltas();
 			}
 		}
@@ -188,7 +188,7 @@ namespace UE::LegacyVertexDeltaModel
 		return Cast<ULegacyVertexDeltaModelVizSettings>(GetVertexDeltaModel()->GetVizSettings());
 	}
 
-	double FLegacyVertexDeltaEditorModel::GetTimeAtFrame(int32 FrameNumber) const
+	double FLegacyVertexDeltaEditorModel::GetTrainingTimeAtFrame(int32 FrameNumber) const
 	{
 		const FLegacyVertexDeltaEditorModelActor* EditorActor = static_cast<FLegacyVertexDeltaEditorModelActor*>(GetTimelineEditorActor());
 		if (EditorActor && EditorActor->GetGeometryCacheComponent())
@@ -198,7 +198,7 @@ namespace UE::LegacyVertexDeltaModel
 		return 0.0;
 	}
 
-	int32 FLegacyVertexDeltaEditorModel::GetFrameAtTime(double TimeInSeconds) const
+	int32 FLegacyVertexDeltaEditorModel::GetTrainingFrameAtTime(double TimeInSeconds) const
 	{
 		const FLegacyVertexDeltaEditorModelActor* EditorActor = static_cast<FLegacyVertexDeltaEditorModelActor*>(GetTimelineEditorActor());
 		if (EditorActor && EditorActor->GetGeometryCacheComponent())
@@ -208,7 +208,7 @@ namespace UE::LegacyVertexDeltaModel
 		return 0;
 	}
 
-	int32 FLegacyVertexDeltaEditorModel::GetNumFrames() const
+	int32 FLegacyVertexDeltaEditorModel::GetNumTrainingFrames() const
 	{
 		const UGeometryCache* GeometryCache = GetVertexDeltaModel()->GetGeometryCache();
 		if (GeometryCache == nullptr)
@@ -278,12 +278,15 @@ namespace UE::LegacyVertexDeltaModel
 		VertexDeltaScaleBackup = GetVertexDeltaModel()->GetVertexDeltaScale();
 	}
 
-	void FLegacyVertexDeltaEditorModel::OnTrainingAborted()
+	void FLegacyVertexDeltaEditorModel::OnPostTraining(ETrainingResult TrainingResult, bool bUsePartiallyTrainedWhenAborted)
 	{
 		// Restore the vertex delta mean and scale, as we aborted, and they could have changed when training
 		// on a smaller subset of frames/samples. If we don't do this, the mesh will deform incorrectly.
-		GetVertexDeltaModel()->VertexDeltaMean = VertexDeltaMeanBackup;
-		GetVertexDeltaModel()->VertexDeltaScale = VertexDeltaScaleBackup;
+		if (TrainingResult == ETrainingResult::Aborted && !bUsePartiallyTrainedWhenAborted)
+		{
+			GetVertexDeltaModel()->VertexDeltaMean = VertexDeltaMeanBackup;
+			GetVertexDeltaModel()->VertexDeltaScale = VertexDeltaScaleBackup;
+		}
 	}
 
 	FString FLegacyVertexDeltaEditorModel::GetTrainedNetworkOnnxFile() const
diff --git a/Engine/Plugins/Experimental/MLDeformer/LegacyVertexDeltaModel/Source/LegacyVertexDeltaModelEditor/Private/LegacyVertexDeltaTrainingModel.cpp b/Engine/Plugins/Experimental/MLDeformer/LegacyVertexDeltaModel/Source/LegacyVertexDeltaModelEditor/Private/LegacyVertexDeltaTrainingModel.cpp
--- a/Engine/Plugins/Experimental/MLDeformer/LegacyVertexDeltaModel/Source/LegacyVertexDeltaModelEditor/Private/LegacyVertexDeltaTrainingModel.cpp
+++ b/Engine/Plugins/Experimental/MLDeformer/LegacyVertexDeltaModel/Source/LegacyVertexDeltaModelEditor/Private/LegacyVertexDeltaTrainingModel.cpp
@@ -56,7 +56,7 @@ bool ULegacyVertexDeltaTrainingModel::ComputeVertexDeltaStatistics(uint32 LODInd
 
 	FLegacyVertexDeltaEditorModel* VertexDeltaEditorModel = GetVertexDeltaEditorModel();
 	ULegacyVertexDeltaModel* VertexDeltaModel = VertexDeltaEditorModel->GetVertexDeltaModel();
-	if (VertexDeltaEditorModel->IsDataNormalized())
+	if (!VertexDeltaEditorModel->GetResamplingInputOutputsNeeded())
 	{
 		return true;
 	}
@@ -94,7 +94,7 @@ bool ULegacyVertexDeltaTrainingModel::ComputeVertexDeltaStatistics(uint32 LODInd
 	if (Count > 0.0f)
 	{
 		VertexDeltaModel->VertexDeltaScale = FVector::OneVector * Scale.GetMax();
-		VertexDeltaEditorModel->SetDataNormalized(true);
+		VertexDeltaEditorModel->SetResamplingInputOutputsNeeded(true);
 	}
 
 	return true;
diff --git a/Engine/Plugins/Experimental/MLDeformer/LegacyVertexDeltaModel/Source/LegacyVertexDeltaModelEditor/Public/LegacyVertexDeltaEditorModel.h b/Engine/Plugins/Experimental/MLDeformer/LegacyVertexDeltaModel/Source/LegacyVertexDeltaModelEditor/Public/LegacyVertexDeltaEditorModel.h
--- a/Engine/Plugins/Experimental/MLDeformer/LegacyVertexDeltaModel/Source/LegacyVertexDeltaModelEditor/Public/LegacyVertexDeltaEditorModel.h
+++ b/Engine/Plugins/Experimental/MLDeformer/LegacyVertexDeltaModel/Source/LegacyVertexDeltaModelEditor/Public/LegacyVertexDeltaEditorModel.h
@@ -43,11 +43,11 @@ namespace UE::LegacyVertexDeltaModel
 		virtual void CreateTrainingGroundTruthActor(UWorld* World) override;
 		virtual void CreateTestGroundTruthActor(UWorld* World) override;
 		virtual void OnPreTraining() override;
-		virtual void OnTrainingAborted() override;
+		virtual void OnPostTraining(ETrainingResult TrainingResult, bool bUsePartiallyTrainedWhenAborted) override;
 		virtual void OnInputAssetsChanged() override;
-		virtual int32 GetNumFrames() const override;
-		virtual double GetTimeAtFrame(int32 FrameNumber) const override;
-		virtual int32 GetFrameAtTime(double TimeInSeconds) const override;
+		virtual int32 GetNumTrainingFrames() const override;
+		virtual double GetTrainingTimeAtFrame(int32 FrameNumber) const override;
+		virtual int32 GetTrainingFrameAtTime(double TimeInSeconds) const override;
 		virtual void UpdateIsReadyForTrainingState() override;
 		virtual ETrainingResult Train() override;
 		virtual FString GetDefaultDeformerGraphAssetPath() const override;
diff --git a/Engine/Plugins/Experimental/MLDeformer/MLDeformerFramework/MLDeformerFramework.uplugin b/Engine/Plugins/Experimental/MLDeformer/MLDeformerFramework/MLDeformerFramework.uplugin
--- a/Engine/Plugins/Experimental/MLDeformer/MLDeformerFramework/MLDeformerFramework.uplugin
+++ b/Engine/Plugins/Experimental/MLDeformer/MLDeformerFramework/MLDeformerFramework.uplugin
@@ -2,14 +2,14 @@
 	"FileVersion": 3,
 	"Version": 1,
 	"VersionName": "0.1",
-	"FriendlyName": "MLDeformerFramework",
+	"FriendlyName": "ML Deformer Framework",
 	"Description": "Machine Learning Mesh Deformer Framework",
 	"Category": "Animation",
 	"CreatedBy": "Epic Games, Inc.",
 	"CreatedByURL": "http://epicgames.com",
-	"DocsURL": "",
-	"MarketplaceURL": "",
-	"SupportURL": "",
+	"DocsURL": "https://docs.unrealengine.com/5.0/en-US/using-the-machine-learning-deformer-in-unreal-engine/",
+	"MarketplaceURL": "https://www.unrealengine.com/marketplace/en-US/product/ml-deformer",
+	"SupportURL": "https://forums.unrealengine.com/",
 	"CanContainContent": true,
 	"IsBetaVersion": false,
 	"IsExperimentalVersion": true,
@@ -38,6 +38,10 @@
 		{
 			"Name": "NeuralNetworkInference",
 			"Enabled": true
+		},
+		{
+			"Name": "GeometryCache",
+			"Enabled": true
 		}
 	]
 }
\ No newline at end of file
diff --git a/Engine/Plugins/Experimental/MLDeformer/MLDeformerFramework/Source/MLDeformerFramework/Private/MLDeformerComponent.cpp b/Engine/Plugins/Experimental/MLDeformer/MLDeformerFramework/Source/MLDeformerFramework/Private/MLDeformerComponent.cpp
--- a/Engine/Plugins/Experimental/MLDeformer/MLDeformerFramework/Source/MLDeformerFramework/Private/MLDeformerComponent.cpp
+++ b/Engine/Plugins/Experimental/MLDeformer/MLDeformerFramework/Source/MLDeformerFramework/Private/MLDeformerComponent.cpp
@@ -35,6 +35,7 @@ void UMLDeformerComponent::Init()
 		ModelInstance = Model->CreateModelInstance(this);
 		ModelInstance->SetModel(Model);
 		ModelInstance->Init(SkelMeshComponent);
+		Model->PostMLDeformerComponentInit(ModelInstance);
 	}
 	else
 	{
@@ -45,6 +46,8 @@ void UMLDeformerComponent::Init()
 
 void UMLDeformerComponent::SetupComponent(UMLDeformerAsset* InDeformerAsset, USkeletalMeshComponent* InSkelMeshComponent)
 {
+	TRACE_CPUPROFILER_EVENT_SCOPE(UMLDeformerComponent::SetupComponent)
+
 	if (InSkelMeshComponent)
 	{		
 		AddTickPrerequisiteComponent(InSkelMeshComponent);
@@ -125,7 +128,8 @@ void UMLDeformerComponent::TickComponent(float DeltaTime, enum ELevelTick TickTy
 			SkelMeshComponent && 
 			SkelMeshComponent->GetPredictedLODLevel() == 0)
 		{
-			ModelInstance->Tick(DeltaTime);
+			TRACE_CPUPROFILER_EVENT_SCOPE(UMLDeformerComponent::TickComponent)
+			ModelInstance->Tick(DeltaTime, Weight);
 		}
 	}
 }
diff --git a/Engine/Plugins/Experimental/MLDeformer/MLDeformerFramework/Source/MLDeformerFramework/Private/MLDeformerModel.cpp b/Engine/Plugins/Experimental/MLDeformer/MLDeformerFramework/Source/MLDeformerFramework/Private/MLDeformerModel.cpp
--- a/Engine/Plugins/Experimental/MLDeformer/MLDeformerFramework/Source/MLDeformerFramework/Private/MLDeformerModel.cpp
+++ b/Engine/Plugins/Experimental/MLDeformer/MLDeformerFramework/Source/MLDeformerFramework/Private/MLDeformerModel.cpp
@@ -54,7 +54,7 @@ UMLDeformerModelInstance* UMLDeformerModel::CreateModelInstance(UMLDeformerCompo
 
 void UMLDeformerModel::Init(UMLDeformerAsset* InDeformerAsset) 
 { 
-	check(InDeformerAsset); 
+	check(InDeformerAsset);
 	DeformerAsset = InDeformerAsset; 
 	if (InputInfo == nullptr)
 	{
@@ -64,6 +64,7 @@ void UMLDeformerModel::Init(UMLDeformerAsset* InDeformerAsset)
 
 void UMLDeformerModel::Serialize(FArchive& Archive)
 {
+	TRACE_CPUPROFILER_EVENT_SCOPE(UMLDeformerModel::Serialize)
 	#if WITH_EDITOR
 		if (Archive.IsSaving() && Archive.IsPersistent())
 		{
@@ -82,6 +83,8 @@ UMLDeformerAsset* UMLDeformerModel::GetDeformerAsset() const
 
 void UMLDeformerModel::PostLoad()
 {
+	TRACE_CPUPROFILER_EVENT_SCOPE(UMLDeformerModel::PostLoad)
+
 	Super::PostLoad();
 
 	InitGPUData();
@@ -100,11 +103,18 @@ void UMLDeformerModel::PostLoad()
 
 	if (NeuralNetwork)
 	{
-		NeuralNetwork->SetDeviceType(ENeuralDeviceType::GPU, ENeuralDeviceType::CPU, ENeuralDeviceType::GPU);
-		if (NeuralNetwork->GetDeviceType() != ENeuralDeviceType::GPU || 
-			NeuralNetwork->GetOutputDeviceType() != ENeuralDeviceType::GPU)
+		// If we run the neural network on the GPU.
+		if (IsNeuralNetworkOnGPU())
+		{
+			NeuralNetwork->SetDeviceType(ENeuralDeviceType::GPU, ENeuralDeviceType::CPU, ENeuralDeviceType::GPU);
+			if (NeuralNetwork->GetDeviceType() != ENeuralDeviceType::GPU || NeuralNetwork->GetOutputDeviceType() != ENeuralDeviceType::GPU || NeuralNetwork->GetInputDeviceType() != ENeuralDeviceType::CPU)
+			{
+				UE_LOG(LogMLDeformer, Error, TEXT("Neural net in ML Deformer '%s' cannot run on the GPU, it will not be active."), *GetDeformerAsset()->GetName());
+			}
+		}
+		else // We run our neural network on the CPU.
 		{
-			UE_LOG(LogMLDeformer, Error, TEXT("Neural net in ML Deformer '%s' cannot run on the GPU, it will not be active."), *GetName());
+			NeuralNetwork->SetDeviceType(ENeuralDeviceType::CPU, ENeuralDeviceType::CPU, ENeuralDeviceType::CPU);
 		}
 	}
 }
@@ -147,6 +157,19 @@ void UMLDeformerModel::InitGPUData()
 	BeginInitResource(&VertexMapBuffer);
 }
 
+void UMLDeformerModel::FloatArrayToVector3Array(const TArray<float>& FloatArray, TArray<FVector3f>& OutVectorArray)
+{
+	check(FloatArray.Num() % 3 == 0);
+	const int32 NumVerts = FloatArray.Num() / 3;
+	OutVectorArray.Reset();
+	OutVectorArray.SetNumUninitialized(NumVerts);
+	for (int32 VertexIndex = 0; VertexIndex < NumVerts; ++VertexIndex)
+	{
+		const int32 FloatBufferOffset = VertexIndex * 3;
+		OutVectorArray[VertexIndex] = FVector3f(FloatArray[FloatBufferOffset + 0], FloatArray[FloatBufferOffset + 1], FloatArray[FloatBufferOffset + 2]);
+	}
+}
+
 #if WITH_EDITOR
 	void UMLDeformerModel::UpdateNumBaseMeshVertices()
 	{
diff --git a/Engine/Plugins/Experimental/MLDeformer/MLDeformerFramework/Source/MLDeformerFramework/Private/MLDeformerModelInstance.cpp b/Engine/Plugins/Experimental/MLDeformer/MLDeformerFramework/Source/MLDeformerFramework/Private/MLDeformerModelInstance.cpp
--- a/Engine/Plugins/Experimental/MLDeformer/MLDeformerFramework/Source/MLDeformerFramework/Private/MLDeformerModelInstance.cpp
+++ b/Engine/Plugins/Experimental/MLDeformer/MLDeformerFramework/Source/MLDeformerFramework/Private/MLDeformerModelInstance.cpp
@@ -98,6 +98,32 @@ FString UMLDeformerModelInstance::CheckCompatibility(USkeletalMeshComponent* InS
 				UE_LOG(LogMLDeformer, Error, TEXT("Deformer '%s': %s"), *(Model->GetDeformerAsset()->GetName()), *InputErrorString);
 			}
 		}
+
+		// Check if the neural network is on the right device, if not, we can't continue.
+		if (Model->IsNeuralNetworkOnGPU())
+		{
+			if (NeuralNetwork->GetDeviceType() != ENeuralDeviceType::GPU || NeuralNetwork->GetOutputDeviceType() != ENeuralDeviceType::GPU || NeuralNetwork->GetInputDeviceType() != ENeuralDeviceType::CPU)
+			{
+				const FString DeviceErrorString = "The neural network is expected to run and output on the GPU, but it isn't.";
+				ErrorText += DeviceErrorString + "\n";
+				if (LogIssues)
+				{
+					UE_LOG(LogMLDeformer, Error, TEXT("Deformer '%s': %s"), *(Model->GetDeformerAsset()->GetName()), *DeviceErrorString);
+				}
+			}
+		}
+		else
+		{
+			if (NeuralNetwork->GetDeviceType() != ENeuralDeviceType::CPU || NeuralNetwork->GetOutputDeviceType() != ENeuralDeviceType::CPU || NeuralNetwork->GetInputDeviceType() != ENeuralDeviceType::CPU)
+			{
+				const FString DeviceErrorString = "The neural network is expected to run fully on the CPU, but it isn't.";
+				ErrorText += DeviceErrorString + "\n";
+				if (LogIssues)
+				{
+					UE_LOG(LogMLDeformer, Error, TEXT("Deformer '%s': %s"), *(Model->GetDeformerAsset()->GetName()), *DeviceErrorString);
+				}
+			}
+		}
 	}
 
 	return ErrorText;
@@ -197,7 +223,7 @@ int64 UMLDeformerModelInstance::SetCurveValues(float* OutputBuffer, int64 Output
 	return Index;
 }
 
-void UMLDeformerModelInstance::SetNeuralNetworkInputValues(float* InputData, int64 NumInputFloats)
+int64 UMLDeformerModelInstance::SetNeuralNetworkInputValues(float* InputData, int64 NumInputFloats)
 {
 	check(SkeletalMeshComponent);
 
@@ -205,17 +231,30 @@ void UMLDeformerModelInstance::SetNeuralNetworkInputValues(float* InputData, int
 	int64 BufferOffset = 0;
 	BufferOffset = SetBoneTransforms(InputData, NumInputFloats, BufferOffset);
 	BufferOffset = SetCurveValues(InputData, NumInputFloats, BufferOffset);
-	check(BufferOffset == NumInputFloats);
+	return BufferOffset;
 }
 
 bool UMLDeformerModelInstance::IsValidForDataProvider() const
 {
 	const UNeuralNetwork* NeuralNetwork = Model->GetNeuralNetwork();
-	if (NeuralNetwork == nullptr || 
-		!NeuralNetwork->IsLoaded() || 
-		NeuralNetwork->GetDeviceType() != ENeuralDeviceType::GPU || 
-		NeuralNetwork->GetOutputDeviceType() != ENeuralDeviceType::GPU)
+	if (NeuralNetwork == nullptr || !NeuralNetwork->IsLoaded())
+	{
+		return false;
+	}
+
+	// We expect to run on the GPU when using a data provider for the deformer graph system (Optimus).
+	if (Model->IsNeuralNetworkOnGPU())
 	{
+		// Make sure we're actually running the network on the GPU.
+		// Inputs are expected to come from the CPU though.
+		if (NeuralNetwork->GetDeviceType() != ENeuralDeviceType::GPU || NeuralNetwork->GetOutputDeviceType() != ENeuralDeviceType::GPU || NeuralNetwork->GetInputDeviceType() != ENeuralDeviceType::CPU)
+		{
+			return false;
+		}
+	}
+	else
+	{
+		// We don't support CPU neural networks on default, when using the deformer graph, as we want the neural network output on the GPU.
 		return false;
 	}
 
@@ -224,36 +263,52 @@ bool UMLDeformerModelInstance::IsValidForDataProvider() const
 		GetNeuralNetworkInferenceHandle() != -1;
 }
 
-void UMLDeformerModelInstance::Tick(float DeltaTime)
+void UMLDeformerModelInstance::RunNeuralNetwork(float ModelWeight)
 {
+	TRACE_CPUPROFILER_EVENT_SCOPE(UMLDeformerModelInstance::RunNeuralNetwork)
+
+	UNeuralNetwork* NeuralNetwork = Model->GetNeuralNetwork();
+	if (Model->IsNeuralNetworkOnGPU())
+	{
+		// NOTE: Inputs still come from the CPU.
+		check(NeuralNetwork->GetDeviceType() == ENeuralDeviceType::GPU && NeuralNetwork->GetInputDeviceType() == ENeuralDeviceType::CPU && NeuralNetwork->GetOutputDeviceType() == ENeuralDeviceType::GPU);
+		ENQUEUE_RENDER_COMMAND(RunNeuralNetwork)
+			(
+				[NeuralNetwork, Handle = NeuralNetworkInferenceHandle](FRHICommandListImmediate& RHICmdList)
+				{
+					// Output deltas will be available on GPU for DeformerGraph via UMLDeformerDataProvider.
+					FRDGBuilder GraphBuilder(RHICmdList);
+					NeuralNetwork->Run(GraphBuilder, Handle);
+					GraphBuilder.Execute();
+				}
+		);
+	}
+	else
+	{
+		// Run on the CPU.
+		check(NeuralNetwork->GetDeviceType() == ENeuralDeviceType::CPU && NeuralNetwork->GetInputDeviceType() == ENeuralDeviceType::CPU && NeuralNetwork->GetOutputDeviceType() == ENeuralDeviceType::CPU);
+		NeuralNetwork->Run(NeuralNetworkInferenceHandle);
+	}
+}
+
+bool UMLDeformerModelInstance::SetupNeuralNetworkForFrame()
+{
+	TRACE_CPUPROFILER_EVENT_SCOPE(UMLDeformerModelInstance::SetupNeuralNetworkForFrame)
+
 	// Some safety checks.
-	if (Model == nullptr || 
-		SkeletalMeshComponent == nullptr || 
-		SkeletalMeshComponent->SkeletalMesh == nullptr || 
+	if (Model == nullptr ||
+		SkeletalMeshComponent == nullptr ||
+		SkeletalMeshComponent->SkeletalMesh == nullptr ||
 		!bIsCompatible)
 	{
-		return;
+		return false;
 	}
 
 	// Get the network and make sure it's loaded.
 	UNeuralNetwork* NeuralNetwork = Model->GetNeuralNetwork();
 	if (NeuralNetwork == nullptr || !NeuralNetwork->IsLoaded())
 	{
-		return;
-	}
-
-	// If we're not on the GPU we can't continue really.
-	// This is needed as the deformer graph system needs it on the GPU.
-	// Some platforms might not support GPU yet.
-	// Only the inputs are on the CPU.
-	check(NeuralNetwork->GetInputDeviceType() == ENeuralDeviceType::CPU);
-	if (!bAllowCPU)
-	{
-		if (NeuralNetwork->GetDeviceType() != ENeuralDeviceType::GPU ||
-			NeuralNetwork->GetOutputDeviceType() != ENeuralDeviceType::GPU)
-		{
-			return;
-		}
+		return false;
 	}
 
 	// Allocate an inference context if none has already been allocated.
@@ -262,7 +317,7 @@ void UMLDeformerModelInstance::Tick(float DeltaTime)
 		NeuralNetworkInferenceHandle = NeuralNetwork->CreateInferenceContext();
 		if (NeuralNetworkInferenceHandle == -1)
 		{
-			return;
+			return false;
 		}
 	}
 
@@ -271,27 +326,24 @@ void UMLDeformerModelInstance::Tick(float DeltaTime)
 	const int64 NumDeformerAssetInputs = Model->GetInputInfo()->CalcNumNeuralNetInputs();
 	if (NumNeuralNetInputs != NumDeformerAssetInputs)
 	{
-		return;
+		return false;
 	}
 
 	// Update and write the input values directly into the input tensor.
 	float* InputDataPointer = static_cast<float*>(NeuralNetwork->GetInputDataPointerMutableForContext(NeuralNetworkInferenceHandle));
-	SetNeuralNetworkInputValues(InputDataPointer, NumNeuralNetInputs);
+	const int64 NumFloatsWritten = SetNeuralNetworkInputValues(InputDataPointer, NumNeuralNetInputs);
+	check(NumFloatsWritten == NumNeuralNetInputs);
 
-	// Run the neural network.
-	if (NeuralNetwork->GetOutputDeviceType() == ENeuralDeviceType::GPU)
-	{
-		ENQUEUE_RENDER_COMMAND(RunNeuralNetwork)(
-			[NeuralNetwork, Handle = NeuralNetworkInferenceHandle](FRHICommandListImmediate& RHICmdList)
-		{
-				// Output deltas will be available on GPU for DeformerGraph via UMLDeformerDataProvider.
-				FRDGBuilder GraphBuilder(RHICmdList);
-				NeuralNetwork->Run(GraphBuilder, Handle);
-				GraphBuilder.Execute();
-		});
-	}
-	else
+	return true;
+}
+
+
+void UMLDeformerModelInstance::Tick(float DeltaTime, float ModelWeight)
+{
+	if (!SetupNeuralNetworkForFrame())
 	{
-		NeuralNetwork->Run();
+		return;
 	}
+
+	RunNeuralNetwork(ModelWeight);
 }
diff --git a/Engine/Plugins/Experimental/MLDeformer/MLDeformerFramework/Source/MLDeformerFramework/Private/Tests/MLDeformerTestModel.h b/Engine/Plugins/Experimental/MLDeformer/MLDeformerFramework/Source/MLDeformerFramework/Private/Tests/MLDeformerTestModel.h
--- a/Engine/Plugins/Experimental/MLDeformer/MLDeformerFramework/Source/MLDeformerFramework/Private/Tests/MLDeformerTestModel.h
+++ b/Engine/Plugins/Experimental/MLDeformer/MLDeformerFramework/Source/MLDeformerFramework/Private/Tests/MLDeformerTestModel.h
@@ -30,5 +30,5 @@ class UTestModelInstance
 	GENERATED_BODY()
 
 public:
-	void Tick(float DeltaTime) {}
+	void Tick(float DeltaTime, float ModelWeight) override { }
 };
diff --git a/Engine/Plugins/Experimental/MLDeformer/MLDeformerFramework/Source/MLDeformerFramework/Public/MLDeformerModel.h b/Engine/Plugins/Experimental/MLDeformer/MLDeformerFramework/Source/MLDeformerFramework/Public/MLDeformerModel.h
--- a/Engine/Plugins/Experimental/MLDeformer/MLDeformerFramework/Source/MLDeformerFramework/Public/MLDeformerModel.h
+++ b/Engine/Plugins/Experimental/MLDeformer/MLDeformerFramework/Source/MLDeformerFramework/Public/MLDeformerModel.h
@@ -66,6 +66,8 @@ public:
 	virtual UMLDeformerInputInfo* CreateInputInfo();
 	virtual UMLDeformerModelInstance* CreateModelInstance(UMLDeformerComponent* Component);
 	virtual FString GetDisplayName() const PURE_VIRTUAL(UMLDeformerModel::GetDisplayName, return FString(););
+	virtual void PostMLDeformerComponentInit(UMLDeformerModelInstance* ModelInstance) {}
+	virtual bool IsNeuralNetworkOnGPU() const { return true; }	// GPU neural network.
 #if WITH_EDITORONLY_DATA
 	virtual bool HasTrainingGroundTruth() const { return false; }
 	virtual void SampleGroundTruthPositions(float SampleTime, TArray<FVector3f>& OutPositions) {}
@@ -129,6 +131,7 @@ protected:
 	FMLDeformerModelOnPostEditProperty PostEditPropertyEvent;
 
 	void SetInputInfo(UMLDeformerInputInfo* Input) { InputInfo = Input; }
+	void FloatArrayToVector3Array(const TArray<float>& FloatArray, TArray<FVector3f>& OutVectorArray);
 
 public:
 	/** Cached number of skeltal mesh vertices. */
@@ -141,7 +144,7 @@ public:
 
 	/** Describes what inputs we should train the neural network on. */
 	UPROPERTY(EditAnywhere, Category = "Inputs and Output")
-	EMLDeformerTrainingInputFilter TrainingInputs = EMLDeformerTrainingInputFilter::BonesAndCurves;
+	EMLDeformerTrainingInputFilter TrainingInputs = EMLDeformerTrainingInputFilter::BonesOnly;
 
 	/** 
 	 * The information about the neural network inputs. This contains things such as bone names and curve names.
diff --git a/Engine/Plugins/Experimental/MLDeformer/MLDeformerFramework/Source/MLDeformerFramework/Public/MLDeformerModelInstance.h b/Engine/Plugins/Experimental/MLDeformer/MLDeformerFramework/Source/MLDeformerFramework/Public/MLDeformerModelInstance.h
--- a/Engine/Plugins/Experimental/MLDeformer/MLDeformerFramework/Source/MLDeformerFramework/Public/MLDeformerModelInstance.h
+++ b/Engine/Plugins/Experimental/MLDeformer/MLDeformerFramework/Source/MLDeformerFramework/Public/MLDeformerModelInstance.h
@@ -41,8 +41,20 @@ public:
 	 * Update the deformer instance.
 	 * If you use a neural network, this should provide inputs to the neural network.
 	 * @param DeltaTime The delta time since the last frame, in seconds.
+	 * @param ModelWeight The weight of the model, between 0 and 1, where 0 means it has no influence, and 1 means it is fully active.
 	 */
-	virtual void Tick(float DeltaTime);
+	virtual void Tick(float DeltaTime, float ModelWeight);
+
+	/**
+	 * Update the neural network input values directly inside its input tensor.
+	 * This will copy over the bone transforms (and possibly morph and curve weights) into this flat array.
+	 * @param InputData The buffer of floats to write the input values to.
+	 * @param NumInputFloats The number of floats of the InputData buffer. Do not write more than this number of floats to the InputData buffer.
+	 * @return Returns the new buffer element index. For example if there were 2 bones and 3 curves, and each bone takes 6 inputs and each curve takes one 
+	 * input, then after calling this base class method, it will return 15 (2*6 + 3*1). You could call this base class method and then write data after that, starting
+	 * from the offset that is returned.
+	 */
+	virtual int64 SetNeuralNetworkInputValues(float* InputData, int64 NumInputFloats);
 
 	/**
 	 * Check whether the deformer is compatible with a given skeletal mesh component.
@@ -62,6 +74,24 @@ public:
 	 */
 	virtual bool IsValidForDataProvider() const;
 
+	/**
+	 * Setup the neural network for this frame.
+	 * This has to perform compatibility checks, valid pointer checks, and call the SetNeuralNetworkInputValues.
+	 * After this method is executed, the neural network should be ready to execute.
+	 * The UMLDeformerModelInstance::Tick method internally calls both the SetupNeuralNetworkForFrame and RunNeuralNetwork methods.
+	 * @return Returns true when the setup is done correctly and the neural network is ready to be executed. Otherwise false is returned, which 
+	 * can happen when the NeuralNetwork pointer is invalid, when the model is not set, when the network is not compatible, etc.
+	 */
+	virtual bool SetupNeuralNetworkForFrame();
+
+	/**
+	 * Run the neural network.
+	 * This already assumes that compatibility checks are done, and that the network inputs are set etc.
+	 * Internally this will execute the UNeuralNetwork::Run() method, either on GPU or CPU.
+	 * @param ModelWeight The weight of the model, must be between 0 and 1.
+	 */
+	virtual void RunNeuralNetwork(float ModelWeight);
+
 	/**
 	 * Get the current component space bone transforms that we grabbed from the skeletal mesh component.
 	 * @return The cached bone transforms, in component space.
@@ -101,28 +131,7 @@ public:
 	 */
 	int32 GetNeuralNetworkInferenceHandle() const { return NeuralNetworkInferenceHandle; }
 
-	/**
-	 * Is running neural network inference on the CPU allowed for this instance?
-	 * @return Returns true if allowed, or false if only GPU is allowed.
-	 */
-	bool IsCPUInferenceAllowed() const { return bAllowCPU; }
-
-	/**
-	 * Specify whether running neural network inference on the CPU is allowed, or if it is only GPU.
-	 * This can be used by some automated tests that will need to run inference on the CPU to check the output results.
-	 * @param bAllowed Set to true when CPU inference is allowed, or false if only GPU inference is allowed.
-	 */
-	void SetCPUInferenceAllowed(bool bAllowed) { bAllowCPU = bAllowed; }
-
 protected:
-	/**
-	 * Update the neural network input values directly inside its input tensor.
-	 * This will copy over the bone transforms (and possibly morph and curve weights) into this flat array.
-	 * @param InputData The buffer of floats to write the input values to.
-	 * @param NumInputFloats The number of floats of the InputData buffer. Do not write more than this number of floats to the InputData buffer.
-	 */
-	void SetNeuralNetworkInputValues(float* InputData, int64 NumInputFloats);
-
 	/**
 	 * Set the bone transformations inside a given output buffer, starting from a given StartIndex.
 	 * @param OutputBuffer The buffer we need to write the transformations to.
@@ -170,9 +179,6 @@ protected:
 	/** Inference handle for neural network inference. */
 	int32 NeuralNetworkInferenceHandle = -1;
 
-	/** Allow running inference on the CPU? */
-	bool bAllowCPU = false;
-
 	/** Are the deformer asset and the used skeletal mesh component compatible? */
 	bool bIsCompatible = false;
 };
diff --git a/Engine/Plugins/Experimental/MLDeformer/MLDeformerFramework/Source/MLDeformerFramework/Public/MLDeformerVizSettings.h b/Engine/Plugins/Experimental/MLDeformer/MLDeformerFramework/Source/MLDeformerFramework/Public/MLDeformerVizSettings.h
--- a/Engine/Plugins/Experimental/MLDeformer/MLDeformerFramework/Source/MLDeformerFramework/Public/MLDeformerVizSettings.h
+++ b/Engine/Plugins/Experimental/MLDeformer/MLDeformerFramework/Source/MLDeformerFramework/Public/MLDeformerVizSettings.h
@@ -57,7 +57,8 @@ public:
 	bool GetDrawLabels() const { return bDrawLabels; }
 	float GetLabelScale() const { return LabelScale; }
 	EMLDeformerVizMode GetVisualizationMode() const { return VisualizationMode; }
-	int32 GetFrameNumber() const { return FrameNumber; }
+	int32 GetTrainingFrameNumber() const { return TrainingFrameNumber; }
+	int32 GetTestingFrameNumber() const { return TestingFrameNumber; }
 	float GetAnimPlaySpeed() const { return AnimPlaySpeed; }
 	const UAnimSequence* GetTestAnimSequence() const { return TestAnimSequence; }
 	UAnimSequence* GetTestAnimSequence() { return TestAnimSequence; }
@@ -95,7 +96,7 @@ public:
 
 	/** The frame number of the training data to visualize. */
 	UPROPERTY(EditAnywhere, Category = "Training Meshes", meta = (ClampMin = "0"))
-	uint32 FrameNumber = 0;
+	uint32 TrainingFrameNumber = 0;
 
 	/** Draw the text labels above each actor? */
 	UPROPERTY(EditAnywhere, Category = "Shared Settings")
@@ -113,6 +114,10 @@ public:
 	UPROPERTY(EditAnywhere, Category = "Shared Settings", meta = (ClampMin = "0", ForceUnits="cm"))
 	float MeshSpacing = 125.0f;
 
+	/** The frame number of the test data to visualize. */
+	UPROPERTY(EditAnywhere, Category = "Live Settings", meta = (ClampMin = "0"))
+	uint32 TestingFrameNumber = 0;
+
 	/** Show the heat map? This will visualize the active areas of the deformer. */
 	UPROPERTY(EditAnywhere, Category = "Live Settings")
 	bool bShowHeatMap = false;
@@ -126,7 +131,7 @@ public:
 	float HeatMapMax = 1.0f;
 
 	/** Lerp from ML deformed model to ground truth model when in heat map mode. */
-	UPROPERTY(EditAnywhere, Category = "Live Settings", meta = (EditCondition = "HeatMapMode==EMLDeformerHeatMapMode::GroundTruth", UIMin = "0.0", UIMax = "1.0", ClampMin = "0.0", ClampMax = "1.0"))
+	UPROPERTY(EditAnywhere, Category = "Live Settings", meta = (EditCondition = "HeatMapMode==EMLDeformerHeatMapMode::GroundTruth && bShowHeatMap", UIMin = "0.0", UIMax = "1.0", ClampMin = "0.0", ClampMax = "1.0"))
 	float GroundTruthLerp = 0.0f;
 
 	/** Draw the linear skinned actor? */
diff --git a/Engine/Plugins/Experimental/MLDeformer/MLDeformerFramework/Source/MLDeformerFrameworkEditor/MLDeformerFrameworkEditor.Build.cs b/Engine/Plugins/Experimental/MLDeformer/MLDeformerFramework/Source/MLDeformerFrameworkEditor/MLDeformerFrameworkEditor.Build.cs
--- a/Engine/Plugins/Experimental/MLDeformer/MLDeformerFramework/Source/MLDeformerFrameworkEditor/MLDeformerFrameworkEditor.Build.cs
+++ b/Engine/Plugins/Experimental/MLDeformer/MLDeformerFramework/Source/MLDeformerFrameworkEditor/MLDeformerFrameworkEditor.Build.cs
@@ -34,7 +34,9 @@ namespace UnrealBuildTool.Rules
 					"AnimGraph",
 					"ToolWidgets",
 					"GeometryCache",
-					"NeuralNetworkInference"
+					"NeuralNetworkInference",
+					"RenderCore",
+					"RHI"
 				}
 			);
 
diff --git a/Engine/Plugins/Experimental/MLDeformer/MLDeformerFramework/Source/MLDeformerFrameworkEditor/Private/MLDeformerEditorActor.cpp b/Engine/Plugins/Experimental/MLDeformer/MLDeformerFramework/Source/MLDeformerFrameworkEditor/Private/MLDeformerEditorActor.cpp
--- a/Engine/Plugins/Experimental/MLDeformer/MLDeformerFramework/Source/MLDeformerFrameworkEditor/Private/MLDeformerEditorActor.cpp
+++ b/Engine/Plugins/Experimental/MLDeformer/MLDeformerFramework/Source/MLDeformerFrameworkEditor/Private/MLDeformerEditorActor.cpp
@@ -1,6 +1,7 @@
 // Copyright Epic Games, Inc. All Rights Reserved.
 
 #include "MLDeformerEditorActor.h"
+#include "MLDeformerModule.h"
 #include "MLDeformerEditorStyle.h"
 #include "Animation/DebugSkelMeshComponent.h"
 #include "Components/TextRenderComponent.h"
@@ -36,6 +37,16 @@ namespace UE::MLDeformer
 		return true;
 	}
 
+	bool FMLDeformerEditorActor::HasVisualMesh() const
+	{
+		if (SkeletalMeshComponent && SkeletalMeshComponent->SkeletalMesh)
+		{
+			return true;
+		}
+
+		return false;
+	}
+
 	UTextRenderComponent* FMLDeformerEditorActor::CreateLabelComponent(AActor* InActor, FLinearColor Color, const FText& Text) const
 	{
 		const float DefaultLabelScale = FMLDeformerEditorStyle::Get().GetFloat("MLDeformer.DefaultLabelScale");
diff --git a/Engine/Plugins/Experimental/MLDeformer/MLDeformerFramework/Source/MLDeformerFrameworkEditor/Private/MLDeformerEditorMode.cpp b/Engine/Plugins/Experimental/MLDeformer/MLDeformerFramework/Source/MLDeformerFrameworkEditor/Private/MLDeformerEditorMode.cpp
--- a/Engine/Plugins/Experimental/MLDeformer/MLDeformerFramework/Source/MLDeformerFrameworkEditor/Private/MLDeformerEditorMode.cpp
+++ b/Engine/Plugins/Experimental/MLDeformer/MLDeformerFramework/Source/MLDeformerFrameworkEditor/Private/MLDeformerEditorMode.cpp
@@ -74,7 +74,7 @@ namespace UE::MLDeformer
 			return;
 		}
 
-		EditorModel->ClampCurrentFrameIndex();
+		EditorModel->ClampCurrentTrainingFrameIndex();
 		EditorModel->Tick(ViewportClient, DeltaTime);
 	}
 
diff --git a/Engine/Plugins/Experimental/MLDeformer/MLDeformerFramework/Source/MLDeformerFrameworkEditor/Private/MLDeformerEditorModel.cpp b/Engine/Plugins/Experimental/MLDeformer/MLDeformerFramework/Source/MLDeformerFrameworkEditor/Private/MLDeformerEditorModel.cpp
--- a/Engine/Plugins/Experimental/MLDeformer/MLDeformerFramework/Source/MLDeformerFrameworkEditor/Private/MLDeformerEditorModel.cpp
+++ b/Engine/Plugins/Experimental/MLDeformer/MLDeformerFramework/Source/MLDeformerFrameworkEditor/Private/MLDeformerEditorModel.cpp
@@ -8,6 +8,7 @@
 #include "MLDeformerModelRegistry.h"
 #include "MLDeformerEditorToolkit.h"
 #include "MLDeformerEditorActor.h"
+#include "MLDeformerAsset.h"
 #include "MLDeformerComponent.h"
 #include "MLDeformerInputInfo.h"
 #include "MLDeformerSampler.h"
@@ -279,6 +280,35 @@ namespace UE::MLDeformer
 
 	void FMLDeformerEditorModel::Tick(FEditorViewportClient* ViewportClient, float DeltaTime)
 	{
+		// Force the training sequence to use Step interpolation and sample raw animation data.
+		// We do this in the Tick, as this is reset to false in the engine sometimes.
+		UAnimSequence* TrainingAnimSequence = Model->GetAnimSequence();
+		if (TrainingAnimSequence)
+		{
+			TrainingAnimSequence->bUseRawDataOnly = true;
+			TrainingAnimSequence->Interpolation = EAnimInterpolationType::Step;
+		}
+
+		// Do the same for the test anim sequence.
+		UMLDeformerVizSettings* VizSettings = Model->GetVizSettings();
+		{
+			UAnimSequence* TestAnimSequence = VizSettings->GetTestAnimSequence();
+			if (TestAnimSequence)
+			{
+				TestAnimSequence->bUseRawDataOnly = true;
+
+				// Enable step interpolation when doing showing a heatmap vs ground truth.
+				if (VizSettings->HasTestGroundTruth() && VizSettings->GetShowHeatMap() && VizSettings->GetHeatMapMode() == EMLDeformerHeatMapMode::GroundTruth)
+				{
+					TestAnimSequence->Interpolation = EAnimInterpolationType::Step;
+				}
+				else
+				{
+					TestAnimSequence->Interpolation = EAnimInterpolationType::Linear;
+				}
+			}
+		}
+
 		UpdateActorTransforms();
 		UpdateLabels();
 		CheckTrainingDataFrameChanged();
@@ -357,14 +387,26 @@ namespace UE::MLDeformer
 
 	void FMLDeformerEditorModel::UpdateActorVisibility()
 	{
-		UMLDeformerVizSettings* VizSettings = Model->GetVizSettings();
+		const UMLDeformerVizSettings* VizSettings = Model->GetVizSettings();
 		const bool bShowTrainingData = (VizSettings->GetVisualizationMode() == EMLDeformerVizMode::TrainingData);
 		const bool bShowTestData = (VizSettings->GetVisualizationMode() == EMLDeformerVizMode::TestData);
 		for (FMLDeformerEditorActor* EditorActor : EditorActors)
 		{
 			if (EditorActor)
 			{
-				const bool bIsVisible = (EditorActor->IsTestActor() && bShowTestData) || (EditorActor->IsTrainingActor() && bShowTrainingData);
+				bool bIsVisible = (EditorActor->IsTestActor() && bShowTestData) || (EditorActor->IsTrainingActor() && bShowTrainingData);
+				if (EditorActor->GetTypeID() == ActorID_Test_Base)
+				{
+					bIsVisible &= VizSettings->GetDrawLinearSkinnedActor();
+				}
+				else if (EditorActor->GetTypeID() == ActorID_Test_MLDeformed)
+				{
+					bIsVisible &= VizSettings->GetDrawMLDeformedActor();
+				}
+				else if (EditorActor->GetTypeID() == ActorID_Test_GroundTruth)
+				{
+					bIsVisible &= VizSettings->GetDrawGroundTruthActor();
+				}
 				EditorActor->SetVisibility(bIsVisible);
 			}
 		}
@@ -382,20 +424,23 @@ namespace UE::MLDeformer
 
 		// Update the training base actor.
 		UDebugSkelMeshComponent* SkeletalMeshComponent = FindEditorActor(ActorID_Train_Base)->GetSkeletalMeshComponent();
-		check(SkeletalMeshComponent);
-		SkeletalMeshComponent->SetSkeletalMesh(Model->GetSkeletalMesh());
-		if (GetEditor()->GetPersonaToolkitPointer())
-		{
-			GetEditor()->GetPersonaToolkit()->GetPreviewScene()->SetPreviewMesh(Model->GetSkeletalMesh());
-		}
 		UMLDeformerVizSettings* VizSettings = Model->GetVizSettings();
 		UAnimSequence* TestAnimSequence = VizSettings->GetTestAnimSequence();
 		const float TestAnimSpeed = VizSettings->GetAnimPlaySpeed();
-		SkeletalMeshComponent->SetAnimationMode(EAnimationMode::AnimationSingleNode);
-		SkeletalMeshComponent->SetAnimation(TrainingAnimSequence);
-		SkeletalMeshComponent->SetPosition(0.0f);
-		SkeletalMeshComponent->SetPlayRate(TestAnimSpeed);
-		SkeletalMeshComponent->Play(false);
+		{
+			check(SkeletalMeshComponent);
+			SkeletalMeshComponent->SetSkeletalMesh(Model->GetSkeletalMesh());
+			if (GetEditor()->GetPersonaToolkitPointer())
+			{
+				GetEditor()->GetPersonaToolkit()->GetPreviewScene()->SetPreviewMesh(Model->GetSkeletalMesh());
+			}
+			SkeletalMeshComponent->SetAnimationMode(EAnimationMode::AnimationSingleNode);
+			const float CurrentPlayTime = SkeletalMeshComponent->GetPosition();
+			SkeletalMeshComponent->SetAnimation(TrainingAnimSequence);
+			SkeletalMeshComponent->SetPosition(CurrentPlayTime);
+			SkeletalMeshComponent->SetPlayRate(TestAnimSpeed);
+			SkeletalMeshComponent->Play(false);
+		}
 
 		// Update the test base model.
 		SkeletalMeshComponent = FindEditorActor(ActorID_Test_Base)->GetSkeletalMeshComponent();
@@ -403,8 +448,9 @@ namespace UE::MLDeformer
 		{
 			SkeletalMeshComponent->SetSkeletalMesh(Model->GetSkeletalMesh());
 			SkeletalMeshComponent->SetAnimationMode(EAnimationMode::AnimationSingleNode);
+			const float CurrentPlayTime = SkeletalMeshComponent->GetPosition();
 			SkeletalMeshComponent->SetAnimation(TestAnimSequence);
-			SkeletalMeshComponent->SetPosition(0.0f);
+			SkeletalMeshComponent->SetPosition(CurrentPlayTime);
 			SkeletalMeshComponent->SetPlayRate(TestAnimSpeed);
 			SkeletalMeshComponent->Play(true);
 		}
@@ -415,13 +461,14 @@ namespace UE::MLDeformer
 		{
 			SkeletalMeshComponent->SetSkeletalMesh(Model->GetSkeletalMesh());
 			SkeletalMeshComponent->SetAnimationMode(EAnimationMode::AnimationSingleNode);
+			const float CurrentPlayTime = SkeletalMeshComponent->GetPosition();
 			SkeletalMeshComponent->SetAnimation(TestAnimSequence);
-			SkeletalMeshComponent->SetPosition(0.0f);
+			SkeletalMeshComponent->SetPosition(CurrentPlayTime);
 			SkeletalMeshComponent->SetPlayRate(TestAnimSpeed);
 			SkeletalMeshComponent->Play(true);
 		}
 
-		bIsDataNormalized = false;
+		SetResamplingInputOutputsNeeded(true);
 	}
 
 	void FMLDeformerEditorModel::OnPostInputAssetChanged()
@@ -432,7 +479,13 @@ namespace UE::MLDeformer
 		UpdateDeformerGraph();
 		RefreshMLDeformerComponents();
 		UpdateIsReadyForTrainingState();
-		SetTrainingFrame(0);
+
+		const int32 TrainingFrame = GetTrainingFrameAtTime(CalcTrainingTimelinePosition());
+		SetTrainingFrame(TrainingFrame);
+
+		const int32 TestFrame = GetTestFrameAtTime(CalcTestTimelinePosition());
+		SetTestFrame(TestFrame);
+
 		UpdateEditorInputInfo();
 		CheckTrainingDataFrameChanged();
 	}
@@ -440,45 +493,82 @@ namespace UE::MLDeformer
 	void FMLDeformerEditorModel::OnTimeSliderScrubPositionChanged(double NewScrubTime, bool bIsScrubbing)
 	{
 		float PlayOffset = static_cast<float>(NewScrubTime);
-		const int32 TargetFrame = GetFrameAtTime(NewScrubTime);
 
 		UMLDeformerVizSettings* VizSettings = Model->GetVizSettings();
 		if (VizSettings->GetVisualizationMode() == EMLDeformerVizMode::TrainingData)
 		{
+			const int32 TargetFrame = GetTrainingFrameAtTime(NewScrubTime);
 			for (FMLDeformerEditorActor* EditorActor : EditorActors)
 			{
 				if (EditorActor && EditorActor->IsTrainingActor())
 				{
 					if (Model->HasTrainingGroundTruth())
 					{
-						PlayOffset = GetTimeAtFrame(TargetFrame);
+						PlayOffset = GetTrainingTimeAtFrame(TargetFrame);
 					}
 					EditorActor->SetPlayPosition(PlayOffset);
 				}
 			}
-			VizSettings->FrameNumber = TargetFrame;
+			VizSettings->TrainingFrameNumber = TargetFrame;
 		}
 		else if (VizSettings->GetVisualizationMode() == EMLDeformerVizMode::TestData)
 		{
+			const int32 TargetFrame = GetTestFrameAtTime(NewScrubTime);
 			for (FMLDeformerEditorActor* EditorActor : EditorActors)
 			{
 				if (EditorActor && EditorActor->IsTestActor())
 				{
 					if (Model->GetVizSettings()->HasTestGroundTruth())
 					{
-						PlayOffset = GetTimeAtFrame(TargetFrame);
+						PlayOffset = GetTestTimeAtFrame(TargetFrame);
 					}
 					EditorActor->SetPlayPosition(PlayOffset);
 				}
 			}
+			VizSettings->TestingFrameNumber = TargetFrame;
 		}
 	}
 
+	double FMLDeformerEditorModel::GetTrainingTimeAtFrame(int32 FrameNumber) const
+	{
+		return Model->GetAnimSequence() ? Model->GetAnimSequence()->GetTimeAtFrame(FrameNumber) : 0.0;
+	}
+
+	int32 FMLDeformerEditorModel::GetTrainingFrameAtTime(double TimeInSeconds) const
+	{
+		return Model->GetAnimSequence() ? Model->GetAnimSequence()->GetFrameAtTime(TimeInSeconds) : 0;
+	}
+
+	double FMLDeformerEditorModel::GetTestTimeAtFrame(int32 FrameNumber) const
+	{
+		return Model->GetVizSettings()->GetTestAnimSequence() ? Model->GetVizSettings()->GetTestAnimSequence()->GetTimeAtFrame(FrameNumber) : 0.0;
+	}
+
+	int32 FMLDeformerEditorModel::GetTestFrameAtTime(double TimeInSeconds) const
+	{
+		return Model->GetVizSettings()->GetTestAnimSequence() ? Model->GetVizSettings()->GetTestAnimSequence()->GetFrameAtTime(TimeInSeconds) : 0;
+	}
+
 	void FMLDeformerEditorModel::SetTrainingFrame(int32 FrameNumber)
 	{
 		UMLDeformerVizSettings* VizSettings = Model->GetVizSettings();
-		VizSettings->FrameNumber = FrameNumber;
-		OnTimeSliderScrubPositionChanged(GetTimeAtFrame(FrameNumber), false);
+		VizSettings->TrainingFrameNumber = FrameNumber;
+		ClampCurrentTrainingFrameIndex();
+		if (VizSettings->GetVisualizationMode() == EMLDeformerVizMode::TrainingData)
+		{
+			OnTimeSliderScrubPositionChanged(GetTrainingTimeAtFrame(FrameNumber), false);
+		}
+	}
+
+	void FMLDeformerEditorModel::SetTestFrame(int32 FrameNumber)
+	{
+		UMLDeformerVizSettings* VizSettings = Model->GetVizSettings();
+		VizSettings->TestingFrameNumber = FrameNumber;
+		ClampCurrentTestFrameIndex();
+		if (VizSettings->GetVisualizationMode() == EMLDeformerVizMode::TestData)
+		{
+			OnTimeSliderScrubPositionChanged(GetTestTimeAtFrame(FrameNumber), false);
+		}
 	}
 
 	void FMLDeformerEditorModel::HandleDefaultPropertyChanges(FPropertyChangedEvent& PropertyChangedEvent)
@@ -533,11 +623,17 @@ namespace UE::MLDeformer
 		{
 			UpdateTestAnimPlaySpeed();
 		}
-		else if (Property->GetFName() == GET_MEMBER_NAME_CHECKED(UMLDeformerVizSettings, FrameNumber))
+		else if (Property->GetFName() == GET_MEMBER_NAME_CHECKED(UMLDeformerVizSettings, TrainingFrameNumber))
 		{
-			ClampCurrentFrameIndex();
-			const int32 CurrentFrameNumber = Model->GetVizSettings()->GetFrameNumber();
-			OnTimeSliderScrubPositionChanged(GetTimeAtFrame(CurrentFrameNumber), false);
+			ClampCurrentTrainingFrameIndex();
+			const int32 CurrentFrameNumber = Model->GetVizSettings()->GetTrainingFrameNumber();
+			OnTimeSliderScrubPositionChanged(GetTrainingTimeAtFrame(CurrentFrameNumber), false);
+		}
+		else if (Property->GetFName() == GET_MEMBER_NAME_CHECKED(UMLDeformerVizSettings, TestingFrameNumber))
+		{
+			ClampCurrentTestFrameIndex();
+			const int32 CurrentFrameNumber = Model->GetVizSettings()->GetTestingFrameNumber();
+			OnTimeSliderScrubPositionChanged(GetTestTimeAtFrame(CurrentFrameNumber), false);
 		}
 		else if (Property->GetFName() == GET_MEMBER_NAME_CHECKED(UMLDeformerVizSettings, bShowHeatMap))
 		{
@@ -572,6 +668,7 @@ namespace UE::MLDeformer
 		
 		FMLDeformerEditorActor* BaseTestActor = FindEditorActor(ActorID_Test_Base);
 		const bool bMustPause = (BaseTestActor && BaseTestActor->GetSkeletalMeshComponent()) ? !BaseTestActor->GetSkeletalMeshComponent()->bPauseAnims : false;
+
 		for (FMLDeformerEditorActor* EditorActor : EditorActors)
 		{
 			if (EditorActor && EditorActor->IsTestActor())
@@ -599,21 +696,38 @@ namespace UE::MLDeformer
 		return false;
 	}
 
-	double FMLDeformerEditorModel::CalcTimelinePosition() const
+	double FMLDeformerEditorModel::CalcTrainingTimelinePosition() const
 	{
-		// Get the base editor actor, depending on the mode we're in.
-		FMLDeformerEditorActor* EditorActor = nullptr;
-		UMLDeformerVizSettings* VizSettings = Model->GetVizSettings();
-		if (VizSettings->GetVisualizationMode() == EMLDeformerVizMode::TestData)
+		FMLDeformerEditorActor* EditorActor = FindEditorActor(ActorID_Train_GroundTruth);
+		if (EditorActor && EditorActor->HasVisualMesh())
+		{
+			return EditorActor->GetPlayPosition();
+		}
+
+		EditorActor = FindEditorActor(ActorID_Train_Base);
+		if (EditorActor && EditorActor->HasVisualMesh())
 		{
-			EditorActor = FindEditorActor(ActorID_Test_Base);
+			return EditorActor->GetPlayPosition();
 		}
-		else if (VizSettings->GetVisualizationMode() == EMLDeformerVizMode::TrainingData)
+
+		return 0.0;
+	}
+
+	double FMLDeformerEditorModel::CalcTestTimelinePosition() const
+	{
+		FMLDeformerEditorActor* EditorActor = FindEditorActor(ActorID_Test_GroundTruth);
+		if (EditorActor && EditorActor->HasVisualMesh())
 		{
-			EditorActor = FindEditorActor(ActorID_Train_Base);
+			return EditorActor->GetPlayPosition();
 		}
 
-		return EditorActor ? EditorActor->GetPlayPosition() : 0.0;
+		EditorActor = FindEditorActor(ActorID_Test_Base);
+		if (EditorActor && EditorActor->HasVisualMesh())
+		{
+			return EditorActor->GetPlayPosition();
+		}
+
+		return 0.0;
 	}
 
 	void FMLDeformerEditorModel::UpdateTestAnimPlaySpeed()
@@ -629,22 +743,45 @@ namespace UE::MLDeformer
 		}
 	}
 
-	void FMLDeformerEditorModel::ClampCurrentFrameIndex()
+	void FMLDeformerEditorModel::ClampCurrentTrainingFrameIndex()
 	{
 		UMLDeformerVizSettings* VizSettings = Model->GetVizSettings();
-		if (GetNumFrames() > 0)
+		if (GetNumTrainingFrames() > 0)
 		{
-			VizSettings->FrameNumber = FMath::Min(VizSettings->FrameNumber, static_cast<uint32>(GetNumFrames() - 1));
+			VizSettings->TrainingFrameNumber = FMath::Min(VizSettings->TrainingFrameNumber, static_cast<uint32>(GetNumTrainingFrames() - 1));
 		}
 		else
 		{
-			VizSettings->FrameNumber = 0;
+			VizSettings->TrainingFrameNumber = 0;
 		}
 	}
 
+	void FMLDeformerEditorModel::ClampCurrentTestFrameIndex()
+	{
+		UMLDeformerVizSettings* VizSettings = Model->GetVizSettings();
+		if (GetNumTestFrames() > 0)
+		{
+			VizSettings->TestingFrameNumber = FMath::Min(VizSettings->TestingFrameNumber, static_cast<uint32>(GetNumTestFrames() - 1));
+		}
+		else
+		{
+			VizSettings->TestingFrameNumber = 0;
+		}
+	}
+
+	int32 FMLDeformerEditorModel::GetNumTestFrames() const
+	{
+		const UAnimSequence* Sequence = Model->GetVizSettings()->GetTestAnimSequence();
+		if (Sequence)
+		{
+			return Sequence->GetNumberOfSampledKeys();
+		}
+		return 0;
+	}
+
 	int32 FMLDeformerEditorModel::GetNumFramesForTraining() const 
 	{ 
-		return FMath::Min(GetNumFrames(), Model->GetTrainingFrameLimit());
+		return FMath::Min(GetNumTrainingFrames(), Model->GetTrainingFrameLimit());
 	}
 
 	FText FMLDeformerEditorModel::GetBaseAssetChangedErrorText() const
@@ -1073,7 +1210,7 @@ namespace UE::MLDeformer
 	void FMLDeformerEditorModel::SampleDeltas()
 	{
 		UMLDeformerVizSettings* VizSettings = Model->GetVizSettings();
-		ClampCurrentFrameIndex();
+		ClampCurrentTrainingFrameIndex();
 
 		// If we have no Persona toolkit yet, then it is not yet safe to init the sampler.
 		if (Editor->GetPersonaToolkitPointer() != nullptr)
@@ -1084,15 +1221,15 @@ namespace UE::MLDeformer
 		if (Sampler->IsInitialized())
 		{
 			Sampler->SetVertexDeltaSpace(EVertexDeltaSpace::PostSkinning);
-			Sampler->Sample(VizSettings->FrameNumber);
+			Sampler->Sample(VizSettings->TrainingFrameNumber);
 		}
 	}
 
 	void FMLDeformerEditorModel::CheckTrainingDataFrameChanged()
 	{
 		UMLDeformerVizSettings* VizSettings = Model->GetVizSettings();
-		ClampCurrentFrameIndex();
-		if (CurrentTrainingFrame != VizSettings->FrameNumber)
+		ClampCurrentTrainingFrameIndex();
+		if (CurrentTrainingFrame != VizSettings->TrainingFrameNumber)
 		{
 			OnTrainingDataFrameChanged();
 		}
@@ -1103,9 +1240,9 @@ namespace UE::MLDeformer
 		UMLDeformerVizSettings* VizSettings = Model->GetVizSettings();
 
 		// If the current frame number changed, re-sample the deltas if needed.
-		if (CurrentTrainingFrame != VizSettings->FrameNumber)
+		if (CurrentTrainingFrame != VizSettings->TrainingFrameNumber)
 		{
-			CurrentTrainingFrame = VizSettings->FrameNumber;
+			CurrentTrainingFrame = VizSettings->TrainingFrameNumber;
 			if (VizSettings->GetDrawVertexDeltas() && VizSettings->GetVisualizationMode() == EMLDeformerVizMode::TrainingData)
 			{
 				SampleDeltas();
@@ -1250,7 +1387,7 @@ namespace UE::MLDeformer
 		}
 	}
 
-	void FMLDeformerEditorModel::OnPostTraining(ETrainingResult TrainingResult)
+	void FMLDeformerEditorModel::OnPostTraining(ETrainingResult TrainingResult, bool bUsePartiallyTrainedWhenAborted)
 	{
 		for (FMLDeformerEditorActor* EditorActor : EditorActors)
 		{
@@ -1262,13 +1399,9 @@ namespace UE::MLDeformer
 			}
 		}
 
-		if (TrainingResult == ETrainingResult::Success || TrainingResult == ETrainingResult::Aborted)
-		{
-			// The InitAssets call resets the normalized flag, so set it back to true.
-			// This is safe as we finished training, which means we already normalized data.
-			// If we aborted we still have normalized the data. Only when we have AbortedCantUse then we canceled the normalization process.
-			bIsDataNormalized = true;
-		}
+		Sampler->SetVertexDeltaSpace(EVertexDeltaSpace::PostSkinning);
+		SampleDeltas();
+		Model->InitGPUData();
 	}
 
 	FMLDeformerEditorActor* FMLDeformerEditorModel::GetTimelineEditorActor() const
@@ -1287,14 +1420,25 @@ namespace UE::MLDeformer
 
 	UNeuralNetwork* FMLDeformerEditorModel::LoadNeuralNetworkFromOnnx(const FString& Filename) const
 	{
-		FString OnnxFile = FPaths::ConvertRelativePathToFull(Filename);
+		const FString OnnxFile = FPaths::ConvertRelativePathToFull(Filename);
 		if (FPaths::FileExists(OnnxFile))
 		{
 			UE_LOG(LogMLDeformer, Display, TEXT("Loading Onnx file '%s'..."), *OnnxFile);
 			UNeuralNetwork* Result = NewObject<UNeuralNetwork>(Model, UNeuralNetwork::StaticClass());		
 			if (Result->Load(OnnxFile))
 			{
-				Result->SetDeviceType(ENeuralDeviceType::GPU, ENeuralDeviceType::CPU, ENeuralDeviceType::GPU);	
+				if (Model->IsNeuralNetworkOnGPU())
+				{
+					Result->SetDeviceType(ENeuralDeviceType::GPU, ENeuralDeviceType::CPU, ENeuralDeviceType::GPU);	
+					if (Result->GetDeviceType() != ENeuralDeviceType::GPU || Result->GetOutputDeviceType() != ENeuralDeviceType::GPU || Result->GetInputDeviceType() != ENeuralDeviceType::CPU)
+					{
+						UE_LOG(LogMLDeformer, Error, TEXT("Neural net in ML Deformer '%s' cannot run on the GPU, it will not be active."), *Model->GetDeformerAsset()->GetName());
+					}
+				}
+				else
+				{
+					Result->SetDeviceType(ENeuralDeviceType::CPU, ENeuralDeviceType::CPU, ENeuralDeviceType::CPU);
+				}
 				UE_LOG(LogMLDeformer, Display, TEXT("Successfully loaded Onnx file '%s'..."), *OnnxFile);
 				return Result;
 			}
@@ -1317,7 +1461,7 @@ namespace UE::MLDeformer
 		if (!Model->HasTrainingGroundTruth() ||
 			Model->GetAnimSequence() == nullptr ||
 			Model->GetSkeletalMesh() == nullptr ||
-			GetNumFrames() == 0)
+			GetNumTrainingFrames() == 0)
 		{
 			return false;
 		}
@@ -1361,6 +1505,114 @@ namespace UE::MLDeformer
 		}
 	}
 
+	void FMLDeformerEditorModel::ZeroDeltasByThreshold(TArray<FVector3f>& Deltas, float Threshold)
+	{
+		for (int32 VertexIndex = 0; VertexIndex < Deltas.Num(); ++VertexIndex)
+		{
+			if (Deltas[VertexIndex].Length() <= Threshold)
+			{
+				Deltas[VertexIndex] = FVector3f::ZeroVector;
+			}
+		}
+	}
+
+	void FMLDeformerEditorModel::CreateEngineMorphTargets(TArray<UMorphTarget*>& OutMorphTargets, const TArray<FVector3f>& Deltas, const FString& NamePrefix, int32 LOD, float DeltaThreshold)
+	{
+		OutMorphTargets.Reset();
+
+		const int32 NumBaseMeshVerts = Model->GetNumBaseMeshVerts();
+		check(Deltas.Num() % NumBaseMeshVerts == 0);
+		const int32 NumMorphTargets = Deltas.Num() / NumBaseMeshVerts;
+		check((Deltas.Num() / NumMorphTargets) == NumBaseMeshVerts);
+		check(!Model->GetVertexMap().IsEmpty());
+
+		USkeletalMesh* SkelMesh = Model->GetSkeletalMesh();
+		FSkeletalMeshRenderData* RenderData = SkelMesh->GetResourceForRendering();
+		check(RenderData);
+		check(!RenderData->LODRenderData.IsEmpty());
+		const int32 NumRenderVertices = RenderData->LODRenderData[LOD].GetNumVertices();
+
+		// Initialize an engine morph target for each model morph target.
+		UE_LOG(LogMLDeformer, Display, TEXT("Initializing %d engine morph targets of %d vertices each"), NumMorphTargets, Deltas.Num() / NumMorphTargets);
+		for (int32 BlendShapeIndex = 0; BlendShapeIndex < NumMorphTargets; ++BlendShapeIndex)
+		{
+			const FName MorphName = *FString::Printf(TEXT("%s%.3d"), *NamePrefix, BlendShapeIndex);
+			UMorphTarget* MorphTarget = NewObject<UMorphTarget>(SkelMesh, MorphName);
+			MorphTarget->BaseSkelMesh = SkelMesh;
+			OutMorphTargets.Add(MorphTarget);
+
+			// Create a new LOD model for this morph.
+			TArray<FMorphTargetLODModel>& MorphLODs = MorphTarget->GetMorphLODModels();
+			MorphLODs.AddDefaulted();
+			FMorphTargetLODModel& MorphLODModel = MorphLODs.Last();
+
+			// Initialize the morph target LOD level.
+			MorphLODModel.Reset();
+			MorphLODModel.bGeneratedByEngine = true;
+			MorphLODModel.NumBaseMeshVerts = NumRenderVertices;
+			MorphLODModel.NumVertices = NumRenderVertices;
+
+			// Init sections.
+			const int32 NumSections = RenderData->LODRenderData[LOD].RenderSections.Num();
+			MorphLODModel.SectionIndices.AddUninitialized(NumSections);
+			for (int32 SectionIndex = 0; SectionIndex < NumSections; ++SectionIndex)
+			{
+				MorphLODModel.SectionIndices[SectionIndex] = SectionIndex;
+			}
+
+			// Init deltas for this morph target.
+			MorphLODModel.Vertices.Reserve(NumRenderVertices);
+			for (int32 VertexIndex = 0; VertexIndex < NumRenderVertices; ++VertexIndex)
+			{
+				const int32 ImportedVertexNumber = Model->VertexMap[VertexIndex];
+				if (ImportedVertexNumber != INDEX_NONE)
+				{
+					const FVector3f Delta = Deltas[ImportedVertexNumber + BlendShapeIndex * NumBaseMeshVerts];
+					if (Delta.Length() > DeltaThreshold)
+					{
+						MorphLODModel.Vertices.AddDefaulted();
+						FMorphTargetDelta& MorphTargetDelta = MorphLODModel.Vertices.Last();
+						MorphTargetDelta.PositionDelta = Delta;
+						MorphTargetDelta.SourceIdx = VertexIndex;
+						MorphTargetDelta.TangentZDelta = FVector3f::ZeroVector;
+					}
+				}
+			}
+
+			MorphLODModel.Vertices.Shrink();
+		}
+	}
+
+	void FMLDeformerEditorModel::CompressEngineMorphTargets(FMorphTargetVertexInfoBuffers& OutMorphBuffers, const TArray<UMorphTarget*>& MorphTargets, int32 LOD, float MorphErrorTolerance)
+	{
+		USkeletalMesh* SkelMesh = Model->GetSkeletalMesh();
+		FSkeletalMeshRenderData* RenderData = SkelMesh->GetResourceForRendering();
+		check(RenderData);
+		check(!RenderData->LODRenderData.IsEmpty());
+		const int32 NumRenderVertices = RenderData->LODRenderData[LOD].GetNumVertices();
+
+		// Release any existing morph buffer data.
+		if (OutMorphBuffers.IsRHIIntialized() && OutMorphBuffers.IsInitialized())
+		{
+			ReleaseResourceAndFlush(&OutMorphBuffers);
+		}
+
+		// Don't empty the array of morph target data when we init the RHI buffers, as we need them to serialize later on.
+		OutMorphBuffers = FMorphTargetVertexInfoBuffers();
+		OutMorphBuffers.SetEmptyMorphCPUDataOnInitRHI(false);
+
+		// Initialize the compressed morph target buffers.
+		OutMorphBuffers.InitMorphResources
+		(
+			GMaxRHIShaderPlatform,
+			TArray<FSkelMeshRenderSection>(),	// Empty array, as we don't need tangents, since we recalculate them in the shader.
+			MorphTargets,
+			NumRenderVertices,
+			LOD,
+			MorphErrorTolerance
+		);
+	}
+
 	FString FMLDeformerEditorModel::GetHeatMapMaterialPath() const
 	{
 		return FString(TEXT("/MLDeformerFramework/Materials/MLDeformerHeatMapMat.MLDeformerHeatMapMat"));
diff --git a/Engine/Plugins/Experimental/MLDeformer/MLDeformerFramework/Source/MLDeformerFrameworkEditor/Private/MLDeformerEditorToolkit.cpp b/Engine/Plugins/Experimental/MLDeformer/MLDeformerFramework/Source/MLDeformerFrameworkEditor/Private/MLDeformerEditorToolkit.cpp
--- a/Engine/Plugins/Experimental/MLDeformer/MLDeformerFramework/Source/MLDeformerFrameworkEditor/Private/MLDeformerEditorToolkit.cpp
+++ b/Engine/Plugins/Experimental/MLDeformer/MLDeformerFramework/Source/MLDeformerFrameworkEditor/Private/MLDeformerEditorToolkit.cpp
@@ -11,6 +11,7 @@
 #include "MLDeformerInputInfo.h"
 #include "MLDeformerEditorStyle.h"
 #include "MLDeformerVizSettings.h"
+#include "MLDeformerSampler.h"
 #include "EditorModeManager.h"
 #include "EditorViewportClient.h"
 #include "Modules/ModuleManager.h"
@@ -103,7 +104,8 @@ namespace UE::MLDeformer
 		if (ActiveModel)
 		{
 			ActiveModel->UpdateIsReadyForTrainingState();
-			ActiveModel->SetTrainingFrame(ActiveModel->GetModel()->GetVizSettings()->GetFrameNumber());
+			ActiveModel->SetTrainingFrame(ActiveModel->GetModel()->GetVizSettings()->GetTrainingFrameNumber());
+			ActiveModel->SetTestFrame(ActiveModel->GetModel()->GetVizSettings()->GetTestingFrameNumber());
 		}
 
 		if (DeformerModel)
@@ -363,15 +365,14 @@ namespace UE::MLDeformer
 						const double StartTime = FPlatformTime::Seconds();
 						const ETrainingResult TrainingResult = ActiveModel->Train();
 						const double TrainingDuration = FPlatformTime::Seconds() - StartTime;
-						const bool bMarkDirty = HandleTrainingResult(TrainingResult, TrainingDuration);
+						bool bUsePartiallyTrained = true;	// Will get modified by the HandleTrainingResult function below.
+						const bool bMarkDirty = HandleTrainingResult(TrainingResult, TrainingDuration, bUsePartiallyTrained);
+						ActiveModel->OnPostTraining(TrainingResult, bUsePartiallyTrained);
 
-						// Reinit the input assets, but restore the frame number.
-						const int32 FrameNumber = Model->GetVizSettings()->GetFrameNumber();
-						ActiveModel->OnInputAssetsChanged();
-						ActiveModel->SetTrainingFrame(FrameNumber);
+						const bool NeedsResamplingBackup = ActiveModel->GetResamplingInputOutputsNeeded();
+						ActiveModel->SetResamplingInputOutputsNeeded(NeedsResamplingBackup);
 						ActiveModel->RefreshMLDeformerComponents();
 						ActiveModel->SetHeatMapMaterialEnabled(Model->GetVizSettings()->GetShowHeatMap());
-						ActiveModel->OnPostTraining(TrainingResult);
 
 						if (bMarkDirty)
 						{
@@ -414,8 +415,10 @@ namespace UE::MLDeformer
 		ToolbarBuilder.EndSection();
 	}
 
-	bool FMLDeformerEditorToolkit::HandleTrainingResult(ETrainingResult TrainingResult, double TrainingDuration)
+	bool FMLDeformerEditorToolkit::HandleTrainingResult(ETrainingResult TrainingResult, double TrainingDuration, bool& bOutUsePartiallyTrained)
 	{
+		bOutUsePartiallyTrained = true;
+
 		FText WindowTitle(LOCTEXT("TrainingResultsWindowTitle", "Training Results"));
 		FText WindowMessage;
 
@@ -444,6 +447,7 @@ namespace UE::MLDeformer
 			// Training fully finished.
 			case ETrainingResult::Success:
 			{
+				ActiveModel->SetResamplingInputOutputsNeeded(false);
 				if (!ActiveModel->LoadTrainedNetwork())
 				{
 					GEditor->PlayEditorSound(TEXT("/Engine/EditorSounds/Notifications/CompileFailed_Cue.CompileFailed_Cue"));
@@ -483,7 +487,7 @@ namespace UE::MLDeformer
 				}
 				else
 				{
-					ActiveModel->OnTrainingAborted();
+					bOutUsePartiallyTrained = false;
 					ShowNotification(LOCTEXT("TrainingAborted", "Training aborted!"), SNotificationItem::ECompletionState::CS_None, true);
 				}
 			}
@@ -513,14 +517,6 @@ namespace UE::MLDeformer
 			}
 			break;
 
-			// Unknown failure, probably something in the python code.
-			case ETrainingResult::FailMissingPythonClass:
-			{
-				GEditor->PlayEditorSound(TEXT("/Engine/EditorSounds/Notifications/CompileFailed_Cue.CompileFailed_Cue"));
-				WindowMessage = LOCTEXT("TrainingMissingPythonClass", "Training failed!\nThe python script doesn't contain a MLDeformerTrainingModel derived class with train function inside.");
-			}
-			break;
-
 			// Unhandled error codes.
 			default: check(false);
 		}
@@ -604,18 +600,9 @@ namespace UE::MLDeformer
 		if (ActiveModel)
 		{
 			ActiveModel->SetDefaultDeformerGraphIfNeeded();
-
 			ActiveModel->CreateActors(InPersonaPreviewScene);
 			ActiveModel->UpdateActorVisibility();
-
-			UMLDeformerVizSettings* VizSettings = ActiveModel->GetModel()->GetVizSettings();
-			const int32 FrameNumber = VizSettings->GetFrameNumber();
-			ActiveModel->OnInputAssetsChanged();	// This resets FrameNumber to 0
-			VizSettings->FrameNumber = FrameNumber;
-
-			ActiveModel->SetTrainingFrame(FrameNumber);
-			ActiveModel->SampleDeltas();
-
+			ActiveModel->OnInputAssetsChanged();
 			ActiveModel->CreateHeatMapAssets();
 			ActiveModel->SetHeatMapMaterialEnabled(ActiveModel->GetModel()->GetVizSettings()->GetShowHeatMap());
 		}
@@ -655,9 +642,17 @@ namespace UE::MLDeformer
 
 				if (VizSettings->GetVisualizationMode() == EMLDeformerVizMode::TrainingData)
 				{
-					ActiveModel->ClampCurrentFrameIndex();
-					const int32 CurrentFrameNumber = ActiveModel->GetModel()->GetVizSettings()->GetFrameNumber();
+					ActiveModel->ClampCurrentTrainingFrameIndex();
+					const int32 CurrentFrameNumber = ActiveModel->GetModel()->GetVizSettings()->GetTrainingFrameNumber();
 					ActiveModel->SetTrainingFrame(CurrentFrameNumber);
+					ActiveModel->GetSampler()->SetVertexDeltaSpace(EVertexDeltaSpace::PostSkinning);
+					ActiveModel->SampleDeltas();
+				}
+				else if (VizSettings->GetVisualizationMode() == EMLDeformerVizMode::TestData)
+				{
+					ActiveModel->ClampCurrentTestFrameIndex();
+					const int32 CurrentFrameNumber = ActiveModel->GetModel()->GetVizSettings()->GetTestingFrameNumber();
+					ActiveModel->SetTestFrame(CurrentFrameNumber);
 				}
 
 				if (GetVizSettingsDetailsView())
@@ -742,7 +737,15 @@ namespace UE::MLDeformer
 	{
 		if (ActiveModel)
 		{
-			return ActiveModel->CalcTimelinePosition();
+			const UMLDeformerVizSettings* VizSettings = ActiveModel->GetModel()->GetVizSettings();
+			if (VizSettings->GetVisualizationMode() == EMLDeformerVizMode::TrainingData)
+			{
+				return ActiveModel->CalcTrainingTimelinePosition();
+			}
+			else if (VizSettings->GetVisualizationMode() == EMLDeformerVizMode::TestData)
+			{
+				return ActiveModel->CalcTestTimelinePosition();
+			}
 		}
 
 		return 0.0;
@@ -779,7 +782,7 @@ namespace UE::MLDeformer
 			SetTimeSliderRange(0.0, Duration);
 		}
 
-		ActiveModel->ClampCurrentFrameIndex();
+		ActiveModel->ClampCurrentTrainingFrameIndex();
 	}
 
 	void FMLDeformerEditorToolkit::SetTimeSliderRange(double StartTime, double EndTime)
diff --git a/Engine/Plugins/Experimental/MLDeformer/MLDeformerFramework/Source/MLDeformerFrameworkEditor/Private/MLDeformerTrainingModel.cpp b/Engine/Plugins/Experimental/MLDeformer/MLDeformerFramework/Source/MLDeformerFrameworkEditor/Private/MLDeformerTrainingModel.cpp
--- a/Engine/Plugins/Experimental/MLDeformer/MLDeformerFramework/Source/MLDeformerFrameworkEditor/Private/MLDeformerTrainingModel.cpp
+++ b/Engine/Plugins/Experimental/MLDeformer/MLDeformerFramework/Source/MLDeformerFrameworkEditor/Private/MLDeformerTrainingModel.cpp
@@ -41,14 +41,24 @@ bool UMLDeformerTrainingModel::SetCurrentSampleIndex(int32 Index)
 	return SampleFrame(Index);
 }
 
+bool UMLDeformerTrainingModel::GetNeedsResampling() const
+{
+	return EditorModel->GetResamplingInputOutputsNeeded();
+}
+
+void UMLDeformerTrainingModel::SetNeedsResampling(bool bNeedsResampling)
+{
+	EditorModel->SetResamplingInputOutputsNeeded(bNeedsResampling);
+}
+
 bool UMLDeformerTrainingModel::SampleFrame(int32 Index)
 {
 	using namespace UE::MLDeformer;
 
 	// Make sure we have a valid frame number.
-	if (Index < 0 || Index >= EditorModel->GetNumFrames())
+	if (Index < 0 || Index >= EditorModel->GetNumTrainingFrames())
 	{
-		UE_LOG(LogMLDeformer, Warning, TEXT("Sample index must range from %d to %d, but a value of %d was provided."), 0, EditorModel->GetNumFrames()-1, Index);
+		UE_LOG(LogMLDeformer, Warning, TEXT("Sample index must range from %d to %d, but a value of %d was provided."), 0, EditorModel->GetNumTrainingFrames()-1, Index);
 		return false;
 	}
 
diff --git a/Engine/Plugins/Experimental/MLDeformer/MLDeformerFramework/Source/MLDeformerFrameworkEditor/Private/MLDeformerVizSettingsDetails.cpp b/Engine/Plugins/Experimental/MLDeformer/MLDeformerFramework/Source/MLDeformerFrameworkEditor/Private/MLDeformerVizSettingsDetails.cpp
--- a/Engine/Plugins/Experimental/MLDeformer/MLDeformerFramework/Source/MLDeformerFrameworkEditor/Private/MLDeformerVizSettingsDetails.cpp
+++ b/Engine/Plugins/Experimental/MLDeformer/MLDeformerFramework/Source/MLDeformerFrameworkEditor/Private/MLDeformerVizSettingsDetails.cpp
@@ -155,6 +155,7 @@ namespace UE::MLDeformer
 		LiveSettingsCategory->SetCategoryVisibility(bShowTestData);
 		LiveSettingsCategory->AddProperty(GET_MEMBER_NAME_CHECKED(UMLDeformerVizSettings, Weight), UMLDeformerVizSettings::StaticClass());
 		LiveSettingsCategory->AddProperty(GET_MEMBER_NAME_CHECKED(UMLDeformerVizSettings, AnimPlaySpeed), UMLDeformerVizSettings::StaticClass());
+		LiveSettingsCategory->AddProperty(GET_MEMBER_NAME_CHECKED(UMLDeformerVizSettings, TestingFrameNumber), UMLDeformerVizSettings::StaticClass());
 
 		IDetailGroup& HeatMapGroup = LiveSettingsCategory->AddGroup("HeatMap", LOCTEXT("HeatMap", "Heat Map"), false, true);
 		HeatMapGroup.AddPropertyRow(DetailBuilder.GetProperty(GET_MEMBER_NAME_CHECKED(UMLDeformerVizSettings, bShowHeatMap), UMLDeformerVizSettings::StaticClass()));
@@ -162,6 +163,8 @@ namespace UE::MLDeformer
 		HeatMapGroup.AddPropertyRow(DetailBuilder.GetProperty(GET_MEMBER_NAME_CHECKED(UMLDeformerVizSettings, HeatMapMax), UMLDeformerVizSettings::StaticClass()));
 		HeatMapGroup.AddPropertyRow(DetailBuilder.GetProperty(GET_MEMBER_NAME_CHECKED(UMLDeformerVizSettings, GroundTruthLerp), UMLDeformerVizSettings::StaticClass()));
 
+		AddAdditionalSettings();
+
 		IDetailGroup& VisGroup = LiveSettingsCategory->AddGroup("Visibility", LOCTEXT("VisibilityLabel", "Visibility"), false, true);
 		VisGroup.AddPropertyRow(DetailBuilder.GetProperty(GET_MEMBER_NAME_CHECKED(UMLDeformerVizSettings, bDrawLinearSkinnedActor), UMLDeformerVizSettings::StaticClass()));
 		VisGroup.AddPropertyRow(DetailBuilder.GetProperty(GET_MEMBER_NAME_CHECKED(UMLDeformerVizSettings, bDrawMLDeformedActor), UMLDeformerVizSettings::StaticClass()));
@@ -170,7 +173,7 @@ namespace UE::MLDeformer
 
 		// Training data.
 		TrainingMeshesCategoryBuilder->SetCategoryVisibility(bShowTrainingData);
-		TrainingMeshesCategoryBuilder->AddProperty(GET_MEMBER_NAME_CHECKED(UMLDeformerVizSettings, FrameNumber), UMLDeformerVizSettings::StaticClass());
+		TrainingMeshesCategoryBuilder->AddProperty(GET_MEMBER_NAME_CHECKED(UMLDeformerVizSettings, TrainingFrameNumber), UMLDeformerVizSettings::StaticClass());
 		TrainingMeshesCategoryBuilder->AddProperty(GET_MEMBER_NAME_CHECKED(UMLDeformerVizSettings, bDrawDeltas), UMLDeformerVizSettings::StaticClass());
 		TrainingMeshesCategoryBuilder->AddProperty(GET_MEMBER_NAME_CHECKED(UMLDeformerVizSettings, bXRayDeltas), UMLDeformerVizSettings::StaticClass());
 	}
diff --git a/Engine/Plugins/Experimental/MLDeformer/MLDeformerFramework/Source/MLDeformerFrameworkEditor/Public/MLDeformerEditorActor.h b/Engine/Plugins/Experimental/MLDeformer/MLDeformerFramework/Source/MLDeformerFrameworkEditor/Public/MLDeformerEditorActor.h
--- a/Engine/Plugins/Experimental/MLDeformer/MLDeformerFramework/Source/MLDeformerFrameworkEditor/Public/MLDeformerEditorActor.h
+++ b/Engine/Plugins/Experimental/MLDeformer/MLDeformerFramework/Source/MLDeformerFrameworkEditor/Public/MLDeformerEditorActor.h
@@ -43,6 +43,7 @@ namespace UE::MLDeformer
 		virtual FBox GetBoundingBox() const;
 		virtual void Pause(bool bPaused);
 		virtual bool IsGroundTruthActor() const;
+		virtual bool HasVisualMesh() const;
 
 		AActor* GetActor() const { return Actor; }
 		int32 GetTypeID() const { return TypeID; }
diff --git a/Engine/Plugins/Experimental/MLDeformer/MLDeformerFramework/Source/MLDeformerFrameworkEditor/Public/MLDeformerEditorModel.h b/Engine/Plugins/Experimental/MLDeformer/MLDeformerFramework/Source/MLDeformerFrameworkEditor/Public/MLDeformerEditorModel.h
--- a/Engine/Plugins/Experimental/MLDeformer/MLDeformerFramework/Source/MLDeformerFrameworkEditor/Public/MLDeformerEditorModel.h
+++ b/Engine/Plugins/Experimental/MLDeformer/MLDeformerFramework/Source/MLDeformerFrameworkEditor/Public/MLDeformerEditorModel.h
@@ -23,17 +23,27 @@ class UAnimSequence;
 class UMaterial;
 class UGeometryCache;
 class UNeuralNetwork;
+class UMorphTarget;
+class FMorphTargetVertexInfoBuffers;
 
 /** Training process return codes. */
 UENUM()
 enum class ETrainingResult : uint8
 {
-	Success = 0,	/** The training successfully finished. */
-	Aborted,		/** The user has aborted the training process. */
-	AbortedCantUse,	/** The user has aborted the training process and we can't use the resulting network. */
-	FailOnData,		/** The input or output data to the network has issues, which means we cannot train. */
-	FailPythonError,/** The python script has some error (see output log). */
-	FailMissingPythonClass,	/** No training model class was implemented in Python, including a valid train function. */
+	/** The training successfully finished. */
+	Success = 0,
+
+	/** The user has aborted the training process. */
+	Aborted,
+
+	/** The user has aborted the training process and we can't use the resulting network. */
+	AbortedCantUse,
+
+	/** The input or output data to the network has issues, which means we cannot train. */
+	FailOnData,
+
+	/** The python script has some error (see output log). */
+	FailPythonError
 };
 
 namespace UE::MLDeformer
@@ -60,10 +70,8 @@ namespace UE::MLDeformer
 		virtual ~FMLDeformerEditorModel();
 
 		// Required overrides.
-		virtual int32 GetNumFrames() const PURE_VIRTUAL(FMLDeformerEditorModel::GetNumFrames, return 0;);
+		virtual int32 GetNumTrainingFrames() const PURE_VIRTUAL(FMLDeformerEditorModel::GetNumTrainingFrames, return 0;);
 		virtual ETrainingResult Train() PURE_VIRTUAL(FMLDeformerEditorModel::Train, return ETrainingResult::Success;);
-		virtual double GetTimeAtFrame(int32 FrameNumber) const PURE_VIRTUAL(FMLDeformerEditorModel::GetTimeAtFrame, return 0.0;);
-		virtual int32 GetFrameAtTime(double TimeInSeconds) const PURE_VIRTUAL(FMLDeformerEditorModel::GetFrameAtTime, return 0;);
 
 		// Optional overrides.
 		virtual void Init(const InitSettings& Settings);
@@ -72,6 +80,11 @@ namespace UE::MLDeformer
 		virtual void ClearWorld();
 		virtual FMLDeformerEditorActor* CreateEditorActor(const FMLDeformerEditorActor::FConstructSettings& Settings) const;
 		virtual FMLDeformerSampler* CreateSampler() const;
+		virtual double GetTrainingTimeAtFrame(int32 FrameNumber) const;
+		virtual int32 GetTrainingFrameAtTime(double TimeInSeconds) const;
+		virtual double GetTestTimeAtFrame(int32 FrameNumber) const;
+		virtual int32 GetTestFrameAtTime(double TimeInSeconds) const;
+		virtual int32 GetNumTestFrames() const;
 		virtual void Tick(FEditorViewportClient* ViewportClient, float DeltaTime);
 		virtual void CreateTrainingLinearSkinnedActor(const TSharedRef<IPersonaPreviewScene>& InPersonaPreviewScene);
 		virtual void CreateTestLinearSkinnedActor(UWorld* World);
@@ -87,15 +100,18 @@ namespace UE::MLDeformer
 		virtual void HandleDefaultPropertyChanges(FPropertyChangedEvent& PropertyChangedEvent);
 		virtual void OnPlayButtonPressed();
 		virtual void OnPreTraining() {}
-		virtual void OnPostTraining(ETrainingResult TrainingResult);
-		virtual void OnTrainingAborted() {}
+		virtual void OnPostTraining(ETrainingResult TrainingResult, bool bUsePartiallyTrainedWhenAborted);
+		virtual void OnTrainingAborted(bool bUsePartiallyTrainedData) {}
 		virtual bool IsPlayingAnim() const;
-		virtual double CalcTimelinePosition() const;
+		virtual double CalcTrainingTimelinePosition() const;
+		virtual double CalcTestTimelinePosition() const;
 		virtual void OnTimeSliderScrubPositionChanged(double NewScrubTime, bool bIsScrubbing);
 		virtual void UpdateTestAnimPlaySpeed();
-		virtual void ClampCurrentFrameIndex();
+		virtual void ClampCurrentTrainingFrameIndex();
+		virtual void ClampCurrentTestFrameIndex();
 		virtual int32 GetNumFramesForTraining() const;
 		virtual void SetTrainingFrame(int32 FrameNumber);
+		virtual void SetTestFrame(int32 FrameNumber);
 		virtual void Render(const FSceneView* View, FViewport* Viewport, FPrimitiveDrawInterface* PDI);
 		virtual void UpdateIsReadyForTrainingState() { bIsReadyForTraining = false; }
 		virtual FText GetOverlayText() const;
@@ -124,8 +140,8 @@ namespace UE::MLDeformer
 		FMLDeformerEditorActor* FindEditorActor(int32 ActorTypeID) const;
 		bool IsReadyForTraining() const { return bIsReadyForTraining; }
 		FMLDeformerSampler* GetSampler() const { return Sampler; }
-		void SetDataNormalized(bool IsNormalized) { bIsDataNormalized = IsNormalized; }
-		bool IsDataNormalized() const { return bIsDataNormalized; }
+		void SetResamplingInputOutputsNeeded(bool bNeeded) { bNeedToResampleInputOutputs = bNeeded; }
+		bool GetResamplingInputOutputsNeeded() const { return bNeedToResampleInputOutputs; }
 
 		FText GetBaseAssetChangedErrorText() const;
 		FText GetVertexMapChangedErrorText() const;
@@ -150,6 +166,28 @@ namespace UE::MLDeformer
 		void DeleteEditorActors();
 		bool IsEditorReadyForTrainingBasicChecks();
 
+		/** Zero all deltas with a length equal to, or smaller than the threshold value. */
+		void ZeroDeltasByThreshold(TArray<FVector3f>& Deltas, float Threshold);
+
+		/**
+		 * Generate engine morph targets from a set of deltas.
+		 * @param OutMorphTargets The output array with generated morph targets. This array will be reset, and then filled with generated morph targets.
+		 * @param Deltas The per vertex deltas for all morph targets, as one big buffer. Each morph target has 'GetNumBaseMeshVerts()' number of deltas.
+		 * @param NamePrefix The morph target name prefix. If set to "MorphTarget_" the names will be "MorphTarget_000", "MorphTarget_001", "MorphTarget_002", etc.
+		 * @param LOD The LOD index to generate the morphs for.
+		 * @param DeltaThreshold Only include deltas with a length larger than this threshold in the morph targets.
+		 */
+		void CreateEngineMorphTargets(TArray<UMorphTarget*>& OutMorphTargets, const TArray<FVector3f>& Deltas, const FString& NamePrefix = TEXT("MorphTarget_"), int32 LOD = 0, float DeltaThreshold = 0.01f);
+
+		/** 
+		 * Compress morph targets into GPU based morph buffers.
+		 * @param OutMorphBuffers The output compressed GPU based morph buffers. If this buffer is already initialized it will be released first.
+		 * @param MorphTargets The morph targets to compress into GPU friendly buffers.
+		 * @param LOD The LOD index to generate the morphs for.
+		 * @param MorphErrorTolerance The error tolerance for the delta compression, in cm. Higher values compress better but can result in artifacts.
+		 */
+		void CompressEngineMorphTargets(FMorphTargetVertexInfoBuffers& OutMorphBuffers, const TArray<UMorphTarget*>& MorphTargets, int32 LOD = 0, float MorphErrorTolerance = 0.01f);
+
 	protected:
 		/** The runtime model associated with this editor model. */
 		TObjectPtr<UMLDeformerModel> Model = nullptr;
@@ -184,8 +222,8 @@ namespace UE::MLDeformer
 		 */
 		bool bIsReadyForTraining = false;	
 
-		/** Did we already normalize the data needed for training? */
-		bool bIsDataNormalized = false;
+		/** Do we need to resample all input/output data? */
+		bool bNeedToResampleInputOutputs = true;
 	};
 
 	template<class TrainingModelClass>
@@ -197,7 +235,7 @@ namespace UE::MLDeformer
 		if (TrainingModels.IsEmpty())
 		{
 			// We didn't define a derived class in Python.
-			return ETrainingResult::FailMissingPythonClass;
+			return ETrainingResult::FailPythonError;
 		}
 
 		// Perform the training.
diff --git a/Engine/Plugins/Experimental/MLDeformer/MLDeformerFramework/Source/MLDeformerFrameworkEditor/Public/MLDeformerEditorToolkit.h b/Engine/Plugins/Experimental/MLDeformer/MLDeformerFramework/Source/MLDeformerFrameworkEditor/Public/MLDeformerEditorToolkit.h
--- a/Engine/Plugins/Experimental/MLDeformer/MLDeformerFramework/Source/MLDeformerFrameworkEditor/Public/MLDeformerEditorToolkit.h
+++ b/Engine/Plugins/Experimental/MLDeformer/MLDeformerFramework/Source/MLDeformerFrameworkEditor/Public/MLDeformerEditorToolkit.h
@@ -99,7 +99,7 @@ namespace UE::MLDeformer
 		void ShowNotification(const FText& Message, SNotificationItem::ECompletionState State, bool PlaySound) const;
 		FText GetOverlayText() const;
 		void OnSwitchedVisualizationMode();
-		bool HandleTrainingResult(ETrainingResult TrainingResult, double TrainingDuration);
+		bool HandleTrainingResult(ETrainingResult TrainingResult, double TrainingDuration, bool& bOutUsePartiallyTrained);
 		void OnModelChanged(int Index);
 		void OnVizModeChanged(EMLDeformerVizMode Mode);
 		FText GetActiveModelName() const;
@@ -113,25 +113,25 @@ namespace UE::MLDeformer
 
 	private:
 		/** The persona toolkit. */	
-		TSharedPtr<IPersonaToolkit> PersonaToolkit;
+		TSharedPtr<IPersonaToolkit> PersonaToolkit = nullptr;
 
 		/** Model details view. */
-		TSharedPtr<IDetailsView> ModelDetailsView;
+		TSharedPtr<IDetailsView> ModelDetailsView = nullptr;
 
 		/** Model viz settings details view. */
-		TSharedPtr<IDetailsView> VizSettingsDetailsView;
+		TSharedPtr<IDetailsView> VizSettingsDetailsView = nullptr;
 
 		/** The timeline slider widget. */
-		TSharedPtr<SSimpleTimeSlider> TimeSlider;
+		TSharedPtr<SSimpleTimeSlider> TimeSlider = nullptr;
 
 		/** The currently active editor model. */
-		TSharedPtr<FMLDeformerEditorModel> ActiveModel;
+		TSharedPtr<FMLDeformerEditorModel> ActiveModel = nullptr;
 
 		// Persona viewport.
-		TSharedPtr<IPersonaViewport> PersonaViewport;
+		TSharedPtr<IPersonaViewport> PersonaViewport = nullptr;
 
 		/** The ML Deformer Asset. */
-		TObjectPtr<UMLDeformerAsset> DeformerAsset;
+		TObjectPtr<UMLDeformerAsset> DeformerAsset = nullptr;
 
 		/** Has the asset editor been initialized? */
 		bool bIsInitialized = false;
diff --git a/Engine/Plugins/Experimental/MLDeformer/MLDeformerFramework/Source/MLDeformerFrameworkEditor/Public/MLDeformerTrainingModel.h b/Engine/Plugins/Experimental/MLDeformer/MLDeformerFramework/Source/MLDeformerFrameworkEditor/Public/MLDeformerTrainingModel.h
--- a/Engine/Plugins/Experimental/MLDeformer/MLDeformerFramework/Source/MLDeformerFrameworkEditor/Public/MLDeformerTrainingModel.h
+++ b/Engine/Plugins/Experimental/MLDeformer/MLDeformerFramework/Source/MLDeformerFrameworkEditor/Public/MLDeformerTrainingModel.h
@@ -48,6 +48,12 @@ public:
 	UFUNCTION(BlueprintCallable, Category = "Training Data")
 	bool SetCurrentSampleIndex(int32 Index);
 
+	UFUNCTION(BlueprintCallable, Category = "Training Data")
+	bool GetNeedsResampling() const;
+
+	UFUNCTION(BlueprintCallable, Category = "Training Data")
+	void SetNeedsResampling(bool bNeedsResampling);
+
 protected:
 	/** Sample a given frame. This updates the sample deltas, curves, and bone rotations. */
     virtual bool SampleFrame(int32 Index);
diff --git a/Engine/Plugins/Experimental/MLDeformer/MLDeformerFramework/Source/MLDeformerFrameworkEditor/Public/MLDeformerVizSettingsDetails.h b/Engine/Plugins/Experimental/MLDeformer/MLDeformerFramework/Source/MLDeformerFrameworkEditor/Public/MLDeformerVizSettingsDetails.h
--- a/Engine/Plugins/Experimental/MLDeformer/MLDeformerFramework/Source/MLDeformerFrameworkEditor/Public/MLDeformerVizSettingsDetails.h
+++ b/Engine/Plugins/Experimental/MLDeformer/MLDeformerFramework/Source/MLDeformerFrameworkEditor/Public/MLDeformerVizSettingsDetails.h
@@ -40,6 +40,7 @@ namespace UE::MLDeformer
 		virtual void AddTestSequenceErrors() {}
 		virtual void AddDeformerGraphErrors() {}
 		virtual void AddGroundTruth() {}
+		virtual void AddAdditionalSettings() {}
 
 	protected:
 		/** Associated detail layout builder. */
diff --git a/Engine/Plugins/Experimental/MLDeformer/NeuralMorphModel/NeuralMorphModel.uplugin b/Engine/Plugins/Experimental/MLDeformer/NeuralMorphModel/NeuralMorphModel.uplugin
new file mode 100644
--- /dev/null
+++ b/Engine/Plugins/Experimental/MLDeformer/NeuralMorphModel/NeuralMorphModel.uplugin
@@ -0,0 +1,51 @@
+{
+	"FileVersion": 3,
+	"Version": 1,
+	"VersionName": "0.1",
+	"FriendlyName": "ML Deformer Neural Morph Model",
+	"Description": "Neural Morph Model for the ML Deformer Framework",
+	"Category": "Animation",
+	"CreatedBy": "Epic Games, Inc.",
+	"CreatedByURL": "http://epicgames.com",
+	"DocsURL": "https://docs.unrealengine.com/5.0/en-US/using-the-machine-learning-deformer-in-unreal-engine/",
+	"MarketplaceURL": "https://www.unrealengine.com/marketplace/en-US/product/ml-deformer",
+	"SupportURL": "https://forums.unrealengine.com/",
+	"CanContainContent": true,
+	"IsBetaVersion": false,
+	"IsExperimentalVersion": true,
+	"Installed": false,
+	"Modules": [
+		{
+			"Name": "NeuralMorphModel",
+			"Type": "Runtime",
+			"LoadingPhase": "Default"
+		},
+		{
+			"Name": "NeuralMorphModelEditor",
+			"Type": "Editor",
+			"LoadingPhase": "Default"
+		}
+	],
+	"Plugins": [
+		{
+			"Name": "GeometryCache",
+			"Enabled": true
+		},
+		{
+			"Name": "NeuralNetworkInference",
+			"Enabled": true
+		},
+		{
+			"Name": "Optimus",
+			"Enabled": true
+		},
+		{
+			"Name": "PythonFoundationPackages",
+			"Enabled": true
+		},
+		{
+			"Name": "MLDeformerFramework",
+			"Enabled": true
+		}
+	]
+}
diff --git a/Engine/Plugins/Experimental/MLDeformer/NeuralMorphModel/Source/NeuralMorphModel/NeuralMorphModel.Build.cs b/Engine/Plugins/Experimental/MLDeformer/NeuralMorphModel/Source/NeuralMorphModel/NeuralMorphModel.Build.cs
new file mode 100644
--- /dev/null
+++ b/Engine/Plugins/Experimental/MLDeformer/NeuralMorphModel/Source/NeuralMorphModel/NeuralMorphModel.Build.cs
@@ -0,0 +1,58 @@
+// Copyright Epic Games, Inc. All Rights Reserved.
+
+using UnrealBuildTool;
+
+public class NeuralMorphModel : ModuleRules
+{
+	public NeuralMorphModel(ReadOnlyTargetRules Target) : base(Target)
+	{
+		PCHUsage = ModuleRules.PCHUsageMode.UseExplicitOrSharedPCHs;
+
+		PublicIncludePaths.AddRange(
+			new string[] {
+				// ... add public include paths required here ...
+			}
+		);
+				
+		
+		PrivateIncludePaths.AddRange(
+			new string[] {
+				// ... add other private include paths required here ...
+			}
+		);
+			
+		
+		PublicDependencyModuleNames.AddRange(
+			new string[]
+			{
+				"Core"
+				// ... add other public dependencies that you statically link with here ...
+			}
+		);
+			
+		
+		PrivateDependencyModuleNames.AddRange(
+			new string[]
+			{
+				"ComputeFramework",
+				"CoreUObject",
+				"Engine",
+				"GeometryCache",
+				"NeuralNetworkInference",
+				"OptimusCore",
+				"Projects",
+				"RenderCore",
+				"RHI",
+				"MLDeformerFramework"
+			}
+		);
+		
+		
+		DynamicallyLoadedModuleNames.AddRange(
+			new string[]
+			{
+				// ... add any modules that your module loads dynamically here ...
+			}
+		);
+	}
+}
diff --git a/Engine/Plugins/Experimental/MLDeformer/NeuralMorphModel/Source/NeuralMorphModel/Private/NeuralMorphModel.cpp b/Engine/Plugins/Experimental/MLDeformer/NeuralMorphModel/Source/NeuralMorphModel/Private/NeuralMorphModel.cpp
new file mode 100644
--- /dev/null
+++ b/Engine/Plugins/Experimental/MLDeformer/NeuralMorphModel/Source/NeuralMorphModel/Private/NeuralMorphModel.cpp
@@ -0,0 +1,170 @@
+// Copyright Epic Games, Inc. All Rights Reserved.
+#include "NeuralMorphModel.h"
+#include "NeuralMorphModelVizSettings.h"
+#include "NeuralMorphModelInstance.h"
+#include "MLDeformerGeomCacheHelpers.h"
+#include "MLDeformerModelInstance.h"
+#include "MLDeformerComponent.h"
+#include "UObject/Object.h"
+#include "UObject/UObjectGlobals.h"
+#include "Rendering/SkeletalMeshLODRenderData.h"
+#include "Components/SkeletalMeshComponent.h"
+
+#define LOCTEXT_NAMESPACE "NeuralMorphModel"
+
+// The morph target set ID for this model. This has to be unique for every different model.
+int32 UNeuralMorphModel::NeuralMorphsExternalMorphSetID = 0;
+
+UNeuralMorphModel::UNeuralMorphModel(const FObjectInitializer& ObjectInitializer)
+	: Super(ObjectInitializer)
+{
+#if WITH_EDITORONLY_DATA
+	VizSettings = ObjectInitializer.CreateEditorOnlyDefaultSubobject<UNeuralMorphModelVizSettings>(this, TEXT("VizSettings"));
+#endif
+
+	MorphTargetSet = MakeShared<FExternalMorphTargetSet>();
+	MorphTargetSet->Name = FName(TEXT("NeuralBlendShapes"));
+}
+
+void UNeuralMorphModel::Serialize(FArchive& Archive)
+{
+	TRACE_CPUPROFILER_EVENT_SCOPE(UNeuralMorphModel::Serialize)
+	Super::Serialize(Archive);
+
+	// Check if we have initialized our compressed morph buffers.
+	bool bHasMorphData = false;
+	if (Archive.IsSaving())
+	{
+		bHasMorphData = MorphTargetSet.IsValid() ? MorphTargetSet->MorphBuffers.IsMorphCPUDataValid() : false;
+	}
+	Archive << bHasMorphData;
+
+	// Load or save the compressed morph buffers, if they exist.
+	if (bHasMorphData)
+	{
+		check(MorphTargetSet.IsValid());
+		if (Archive.IsLoading())
+		{
+			Archive << MorphTargetSet->MorphBuffers;
+
+			// When we're in editor mode, keep the CPU data around, so we can re-initialize when needed.
+			#if WITH_EDITOR
+				MorphTargetSet->MorphBuffers.SetEmptyMorphCPUDataOnInitRHI(false);
+			#endif
+		}
+		else
+		{
+			Archive << MorphTargetSet->MorphBuffers;
+		}
+	}
+}
+
+UMLDeformerModelInstance* UNeuralMorphModel::CreateModelInstance(UMLDeformerComponent* Component)
+{
+	return NewObject<UNeuralMorphModelInstance>(Component);
+}
+
+#if WITH_EDITOR
+	void UNeuralMorphModel::UpdateNumTargetMeshVertices()
+	{
+		NumTargetMeshVerts = UE::MLDeformer::ExtractNumImportedGeomCacheVertices(GeometryCache);
+	}
+#endif // WITH_EDITOR
+
+#if WITH_EDITORONLY_DATA
+	void UNeuralMorphModel::SampleGroundTruthPositions(float SampleTime, TArray<FVector3f>& OutPositions)
+	{
+		UNeuralMorphModelVizSettings* VertexVizSettings = Cast<UNeuralMorphModelVizSettings>(VizSettings);
+		check(VertexVizSettings);
+
+		UGeometryCache* GeomCache = VertexVizSettings->GetTestGroundTruth();
+		if (GeomCache == nullptr)
+		{
+			OutPositions.Reset();
+			return;
+		}
+
+		if (MeshMappings.IsEmpty())
+		{
+			TArray<FString> FailedImportedMeshnames;
+			GenerateGeomCacheMeshMappings(SkeletalMesh, GeomCache, MeshMappings, FailedImportedMeshnames);
+		}
+
+		UE::MLDeformer::SampleGeomCachePositions(
+			0,
+			SampleTime,
+			MeshMappings,
+			SkeletalMesh,
+			GeomCache,
+			AlignmentTransform,
+			OutPositions);
+	}
+#endif
+
+void UNeuralMorphModel::PostMLDeformerComponentInit(UMLDeformerModelInstance* ModelInstance)
+{
+	TRACE_CPUPROFILER_EVENT_SCOPE(UNeuralMorphModel::PostMLDeformerComponentInit)
+
+	Super::PostMLDeformerComponentInit(ModelInstance);
+
+	// Register the external morph targets buffer to the render data of the skeletal mesh.
+	USkeletalMeshComponent* SkelMeshComponent = ModelInstance->GetSkeletalMeshComponent();
+	if (SkelMeshComponent)
+	{
+		FSkeletalMeshRenderData* RenderData = SkelMeshComponent->GetSkeletalMeshRenderData();
+		FMorphTargetVertexInfoBuffers& MorphBuffers = MorphTargetSet->MorphBuffers;
+		if (RenderData)
+		{
+			// Register the morph set. This overwrites the existing one for this model, if it already exists.
+			// Only add to LOD 0 for now.
+			const int32 LOD = 0;
+			RenderData->LODRenderData[LOD].AddExternalMorphSet(UNeuralMorphModel::NeuralMorphsExternalMorphSetID, MorphTargetSet);
+
+			// Release the render resources, but only in an editor build.
+			// The non-editor build shouldn't do this, as then it can't initialize again. The non-editor build assumes
+			// that the data doesn't change and we don't need to re-init.
+			// In the editor build we have to re-initialize the render resources as the morph targets can change after (re)training, so
+			// that is why we release them here, and intialize them again after.
+			#if WITH_EDITOR
+				BeginReleaseResource(&MorphBuffers);
+			#endif
+
+			// Reinitialize the GPU compressed buffers.
+			if (MorphBuffers.IsMorphCPUDataValid() && MorphBuffers.GetNumMorphs() > 0)
+			{
+				// In a non-editor build this will clear the CPU data.
+				// That also means it can't re-init the resources later on again.
+				BeginInitResource(&MorphBuffers);
+			}
+		}
+
+		// Update the weight information in the Skeletal Mesh.
+		SkelMeshComponent->RefreshExternalMorphTargetWeights();
+	}
+}
+
+void UNeuralMorphModel::SetMorphTargetDeltas(const TArray<float>& Deltas)
+{
+	FloatArrayToVector3Array(Deltas, MorphTargetDeltas);
+}
+
+int32 UNeuralMorphModel::GetMorphTargetDeltaStartIndex(int32 BlendShapeIndex) const
+{
+	if (MorphTargetDeltas.Num() == 0)
+	{
+		return INDEX_NONE;
+	}
+
+	return GetNumBaseMeshVerts() * BlendShapeIndex;
+}
+
+void UNeuralMorphModel::BeginDestroy()
+{
+	if (MorphTargetSet.IsValid())
+	{
+		BeginReleaseResource(&MorphTargetSet->MorphBuffers);
+	}
+	Super::BeginDestroy();
+}
+
+#undef LOCTEXT_NAMESPACE
diff --git a/Engine/Plugins/Experimental/MLDeformer/NeuralMorphModel/Source/NeuralMorphModel/Private/NeuralMorphModelInstance.cpp b/Engine/Plugins/Experimental/MLDeformer/NeuralMorphModel/Source/NeuralMorphModel/Private/NeuralMorphModelInstance.cpp
new file mode 100644
--- /dev/null
+++ b/Engine/Plugins/Experimental/MLDeformer/NeuralMorphModel/Source/NeuralMorphModel/Private/NeuralMorphModelInstance.cpp
@@ -0,0 +1,58 @@
+// Copyright Epic Games, Inc. All Rights Reserved.
+
+#include "NeuralMorphModelInstance.h"
+#include "NeuralMorphModel.h" 
+#include "NeuralNetwork.h"
+#include "Components/SkeletalMeshComponent.h"
+
+void UNeuralMorphModelInstance::RunNeuralNetwork(float ModelWeight)
+{
+	TRACE_CPUPROFILER_EVENT_SCOPE(UNeuralMorphModelInstance::RunNeuralNetwork)
+
+	// Run the neural network, which calculates its outputs, which are the weights of our morph targets.
+	const int LOD = 0;	// For now we only support LOD 0, as we can't setup an ML Deformer per LOD yet.
+
+	// Grab the weight data for this morph set.
+	// This could potentially fail if we are applying this deformer to the wrong skeletal mesh component.
+	FExternalMorphSetWeights* WeightData = SkeletalMeshComponent->GetExternalMorphWeights(LOD).MorphSets.Find(UNeuralMorphModel::NeuralMorphsExternalMorphSetID);
+	if (WeightData == nullptr)
+	{
+		return;
+	}
+
+	// If our model is active, we want to run the neural network and update the morph weights
+	// with the values that the neural net calculated for us.
+	if (ModelWeight > 0.0f)
+	{
+		// Perform the neural network inference, which updates the output tensor.
+		UMLDeformerModelInstance::RunNeuralNetwork(ModelWeight);
+
+		// Get the output tensor, read the values and use them as morph target weights inside the skeletal mesh component.
+		const UNeuralNetwork* NeuralNetwork = Model->GetNeuralNetwork();
+		const FNeuralTensor& OutputTensor = NeuralNetwork->GetOutputTensorForContext(NeuralNetworkInferenceHandle);
+		const int32 NumNetworkWeights = OutputTensor.Num();
+		const int32 NumMorphTargets = WeightData->Weights.Num();;
+
+		// If we have the expected amount of weights.
+		// +1 because we always have an extra morph target that represents the means, with fixed weight of 1.
+		if (NumMorphTargets == NumNetworkWeights + 1)	
+		{
+			// Set the first morph target, which represents the means, to a weight of 1.0, as it always needs to be fully active.
+			WeightData->Weights[0] = 1.0f * ModelWeight;
+
+			// Update all generated morph target weights with the values calculated by our neural network.
+			for (int32 MorphIndex = 0; MorphIndex < NumNetworkWeights; ++MorphIndex)
+			{
+				WeightData->Weights[MorphIndex + 1] = OutputTensor.At<float>(MorphIndex) * ModelWeight;
+			}
+		}
+	} 
+	else
+	{
+		// Reset all weights to zero.
+		for (float& Weight : WeightData->Weights)
+		{
+			Weight = 0.0f;
+		}
+	}
+}
diff --git a/Engine/Plugins/Experimental/MLDeformer/NeuralMorphModel/Source/NeuralMorphModel/Private/NeuralMorphModelModule.cpp b/Engine/Plugins/Experimental/MLDeformer/NeuralMorphModel/Source/NeuralMorphModel/Private/NeuralMorphModelModule.cpp
new file mode 100644
--- /dev/null
+++ b/Engine/Plugins/Experimental/MLDeformer/NeuralMorphModel/Source/NeuralMorphModel/Private/NeuralMorphModelModule.cpp
@@ -0,0 +1,15 @@
+// Copyright Epic Games, Inc. All Rights Reserved.
+
+#include "NeuralMorphModelModule.h"
+#include "Modules/ModuleManager.h"
+#include "Interfaces/IPluginManager.h"
+#include "Misc/Paths.h"
+#include "ShaderCore.h"
+
+NEURALMORPHMODEL_API DEFINE_LOG_CATEGORY(LogNeuralMorphModel)
+
+#define LOCTEXT_NAMESPACE "NeuralMorphModelModule"
+
+IMPLEMENT_MODULE(UE::NeuralMorphModel::FNeuralMorphModelModule, NeuralMorphModel)
+
+#undef LOCTEXT_NAMESPACE
diff --git a/Engine/Plugins/Experimental/MLDeformer/NeuralMorphModel/Source/NeuralMorphModel/Public/NeuralMorphModel.h b/Engine/Plugins/Experimental/MLDeformer/NeuralMorphModel/Source/NeuralMorphModel/Public/NeuralMorphModel.h
new file mode 100644
--- /dev/null
+++ b/Engine/Plugins/Experimental/MLDeformer/NeuralMorphModel/Source/NeuralMorphModel/Public/NeuralMorphModel.h
@@ -0,0 +1,133 @@
+// Copyright Epic Games, Inc. All Rights Reserved.
+#pragma once
+
+#include "CoreMinimal.h"
+#include "MLDeformerModel.h"
+#include "MLDeformerVizSettings.h"
+#include "MLDeformerInputInfo.h"
+#include "MLDeformerGeomCacheHelpers.h"
+#include "UObject/Object.h"
+#include "NeuralMorphModel.generated.h"
+
+class USkeletalMesh;
+class UGeometryCache;
+class UAnimSequence;
+class UNeuralNetwork;
+class UMLDeformerAsset;
+class UMLDeformerModelInstance;
+class USkeleton;
+class IPropertyHandle;
+struct FExternalMorphTargetSet;
+
+UCLASS()
+class NEURALMORPHMODEL_API UNeuralMorphModel 
+	: public UMLDeformerModel
+{
+	GENERATED_BODY()
+
+public:
+	UNeuralMorphModel(const FObjectInitializer& ObjectInitializer);
+
+	// UMLDeformerModel overrides.
+	virtual FString GetDisplayName() const override { return "Neural Morph Model"; }
+	virtual void Serialize(FArchive& Archive) override;
+	virtual void PostMLDeformerComponentInit(UMLDeformerModelInstance* ModelInstance) override;
+	virtual bool IsNeuralNetworkOnGPU() const override { return false; }	// CPU neural network.
+	virtual UMLDeformerModelInstance* CreateModelInstance(UMLDeformerComponent* Component) override;
+#if WITH_EDITORONLY_DATA
+	virtual bool HasTrainingGroundTruth() const override { return (GeometryCache.Get() != nullptr); }
+	virtual void SampleGroundTruthPositions(float SampleTime, TArray<FVector3f>& OutPositions) override;
+#endif
+#if WITH_EDITOR
+	virtual void UpdateNumTargetMeshVertices();
+#endif
+	// ~END UMLDeformerModel overrides.
+
+	// UObject overrides.
+	void BeginDestroy() override;
+	// ~END UObject overrides.
+
+#if WITH_EDITORONLY_DATA
+	const UGeometryCache* GetGeometryCache() const { return GeometryCache; }
+	UGeometryCache* GetGeometryCache() { return GeometryCache; }
+	int32 GetNumHiddenLayers() const { return NumHiddenLayers; }
+	int32 GetNumNeuronsPerLayer() const { return NumNeuronsPerLayer; }
+	int32 GetNumIterations() const { return NumIterations; }
+	int32 GetBatchSize() const { return BatchSize; }
+	float GetLearningRate() const { return LearningRate; }
+	float GetRegularizationFactor() const { return RegularizationFactor;  }
+	float GetMorphTargetDeltaThreshold() const { return MorphTargetDeltaThreshold; }
+	TArray<UE::MLDeformer::FMLDeformerGeomCacheMeshMapping>& GetGeomCacheMeshMappings() { return MeshMappings; }
+	const TArray<UE::MLDeformer::FMLDeformerGeomCacheMeshMapping>& GetGeomCacheMeshMappings() const { return MeshMappings; }
+#endif
+
+	UFUNCTION(BlueprintCallable, Category = "NeuralMorphModel")
+	void SetMorphTargetDeltas(const TArray<float>& Deltas);
+
+	const TArray<FVector3f>& GetMorphTargetDeltas() const { return MorphTargetDeltas; }
+	int32 GetMorphTargetDeltaStartIndex(int32 MorphTargetIndex) const;
+
+public:
+	/** The compressed morph target data, ready for the GPU. */
+	TSharedPtr<FExternalMorphTargetSet> MorphTargetSet;
+
+	/** The morph target set ID that is passed to FSkeletalMeshLODRenderData::AddExternalMorphBuffer(...). */
+	static int32 NeuralMorphsExternalMorphSetID;
+
+	/** 
+	 * The entire set of morph target deltas, 3 per vertex, for each morph target, as one flattened buffer. 
+	 * So the size of this buffer is: (NumVertsPerMorphTarget * 3 * NumMorphTargets).
+	 */
+	TArray<FVector3f> MorphTargetDeltas;
+
+#if WITH_EDITORONLY_DATA
+	TArray<UE::MLDeformer::FMLDeformerGeomCacheMeshMapping> MeshMappings;
+
+	/** The geometry cache that represents the complex mesh deformations. */
+	UPROPERTY(EditAnywhere, Category = "Target Mesh")
+	TObjectPtr<UGeometryCache> GeometryCache = nullptr;
+
+	/** The number of morph targets to generate per bone. */
+	UPROPERTY(EditAnywhere, BlueprintReadWrite, Category = "Training Settings", meta = (ClampMin = "1", ClampMax = "100"))
+	int32 NumMorphTargetsPerBone = 6;
+
+	/** The number of hidden layers that the neural network model will have.\nHigher numbers will slow down performance but can deal with more complex deformations. */
+	UPROPERTY(EditAnywhere, BlueprintReadWrite, Category = "Training Settings", meta = (ClampMin = "1", ClampMax = "10"))
+	int32 NumHiddenLayers = 1;
+
+	/** The number of units/neurons per hidden layer. Higher numbers will slow down performance but allow for more complex mesh deformations. */
+	UPROPERTY(EditAnywhere, BlueprintReadWrite, Category = "Training Settings", meta = (ClampMin = "1"))
+	int32 NumNeuronsPerLayer = 6;
+
+	/** The number of iterations to train the model for. */
+	UPROPERTY(EditAnywhere, BlueprintReadWrite, Category = "Training Settings", meta = (ClampMin = "1"))
+	int32 NumIterations = 2000;
+
+	/** The number of frames per batch when training the model. */
+	UPROPERTY(EditAnywhere, BlueprintReadWrite, AdvancedDisplay, Category = "Training Settings", meta = (ClampMin = "1"))
+	int32 BatchSize = 128;
+
+	/** The learning rate used during the model training. */
+	UPROPERTY(EditAnywhere, BlueprintReadWrite, AdvancedDisplay, Category = "Training Settings", meta = (ClampMin = "0.000001", ClampMax = "1.0"))
+	float LearningRate = 0.01f;
+
+	/** 
+	 * The regularization factor. Higher values can help generate more sparse morph targets, but can also lead to visual artifacts. 
+	 * A value of 0 disables the regularization, and gives the highest quality, at the cost of higher runtime memory usage.
+	 */
+	UPROPERTY(EditAnywhere, BlueprintReadWrite, AdvancedDisplay, Category = "Training Settings", meta = (ClampMin = "0.0", ClampMax = "100.0"))
+	float RegularizationFactor = 1.0f;
+
+	/** 
+	 * Morph target delta values that are smaller than or equal to this threshold will be zeroed out.
+	 * This essentially removes small deltas from morph targets, which will lower the memory usage at runtime, however when set too high it can also introduce visual artifacts.
+	 * A value of 0 will result in the highest quality morph targets, at the cost of higher runtime memory usage.
+	 */
+	UPROPERTY(EditAnywhere, BlueprintReadWrite, AdvancedDisplay, Category = "Training Settings", meta = (ClampMin = "0.0", ClampMax = "1.0"))
+	float MorphTargetDeltaThreshold = 0.0025f;
+
+	/** The morph target error tolerance. Higher values result in larger compression, but could result in visual artifacts. */
+	UPROPERTY(EditAnywhere, BlueprintReadWrite, AdvancedDisplay, Category = "Training Settings", meta = (ClampMin = "0.01", ClampMax = "10000.0"))
+	float MorphTargetErrorTolerance = 20.0f;
+#endif // WITH_EDITORONLY_DATA
+};
diff --git a/Engine/Plugins/Experimental/MLDeformer/NeuralMorphModel/Source/NeuralMorphModel/Public/NeuralMorphModelInstance.h b/Engine/Plugins/Experimental/MLDeformer/NeuralMorphModel/Source/NeuralMorphModel/Public/NeuralMorphModelInstance.h
new file mode 100644
--- /dev/null
+++ b/Engine/Plugins/Experimental/MLDeformer/NeuralMorphModel/Source/NeuralMorphModel/Public/NeuralMorphModelInstance.h
@@ -0,0 +1,16 @@
+// Copyright Epic Games, Inc. All Rights Reserved.
+#pragma once
+
+#include "MLDeformerModelInstance.h"
+#include "NeuralMorphModelInstance.generated.h"
+
+
+UCLASS()
+class NEURALMORPHMODEL_API UNeuralMorphModelInstance
+	: public UMLDeformerModelInstance
+{
+	GENERATED_BODY()
+
+public:
+	virtual void RunNeuralNetwork(float ModelWeight) override;
+};
diff --git a/Engine/Plugins/Experimental/MLDeformer/NeuralMorphModel/Source/NeuralMorphModel/Public/NeuralMorphModelModule.h b/Engine/Plugins/Experimental/MLDeformer/NeuralMorphModel/Source/NeuralMorphModel/Public/NeuralMorphModelModule.h
new file mode 100644
--- /dev/null
+++ b/Engine/Plugins/Experimental/MLDeformer/NeuralMorphModel/Source/NeuralMorphModel/Public/NeuralMorphModelModule.h
@@ -0,0 +1,16 @@
+// Copyright Epic Games, Inc. All Rights Reserved.
+
+#pragma once
+
+#include "CoreMinimal.h"
+#include "Modules/ModuleManager.h"
+
+NEURALMORPHMODEL_API DECLARE_LOG_CATEGORY_EXTERN(LogNeuralMorphModel, Log, All);
+
+namespace UE::NeuralMorphModel
+{
+	class NEURALMORPHMODEL_API FNeuralMorphModelModule
+		: public IModuleInterface
+	{
+	};
+}	// namespace NeuralMorphModel
diff --git a/Engine/Plugins/Experimental/MLDeformer/NeuralMorphModel/Source/NeuralMorphModel/Public/NeuralMorphModelVizSettings.h b/Engine/Plugins/Experimental/MLDeformer/NeuralMorphModel/Source/NeuralMorphModel/Public/NeuralMorphModelVizSettings.h
new file mode 100644
--- /dev/null
+++ b/Engine/Plugins/Experimental/MLDeformer/NeuralMorphModel/Source/NeuralMorphModel/Public/NeuralMorphModelVizSettings.h
@@ -0,0 +1,53 @@
+// Copyright Epic Games, Inc. All Rights Reserved.
+
+#pragma once
+
+#include "CoreMinimal.h"
+#include "UObject/Object.h"
+#include "MLDeformerVizSettings.h"
+#include "NeuralMorphModelVizSettings.generated.h"
+
+class UGeometryCache;
+
+/**
+ * The vizualization settings specific to this model.
+ */
+UCLASS()
+class NEURALMORPHMODEL_API UNeuralMorphModelVizSettings
+	: public UMLDeformerVizSettings
+{
+	GENERATED_BODY()
+
+public:
+#if WITH_EDITORONLY_DATA
+	virtual bool HasTestGroundTruth() const override { return (GroundTruth.Get() != nullptr); }
+	UGeometryCache* GetTestGroundTruth() const { return GroundTruth; }
+
+public:
+	/** The geometry cache that represents the ground truth of the test anim sequence. */
+	UPROPERTY(EditAnywhere, Category = "Test Assets")
+	TObjectPtr<UGeometryCache> GroundTruth = nullptr;
+
+	/** 
+	 * The morph target to visualize. The first one always being the means, so not a sparse target. 
+	 * This only can be used after you trained, in the same editor session directly after training.
+	 */
+	UPROPERTY(EditAnywhere, Category = "Live Settings", meta = (ClampMin = "0"))
+	int32 MorphTargetNumber = 0;
+
+	/**
+	 * The morph target delta threshold. This is a preview of what deltas would be included in the selected morph target
+	 * when using a delta threshold during training that is equal to this value.
+	 * This only can be used after you trained, in the same editor session directly after training.
+	 */
+	UPROPERTY(EditAnywhere, Category = "Live Settings", meta = (ClampMin = "0.001", ClampMax = "1.0"))
+	float MorphTargetDeltaThreshold = 0.01f;
+
+	/** 
+	 * Draw the morph targets as debug data? 
+	 * This only can be used after you trained, in the same editor session directly after training.
+	 */
+	UPROPERTY(EditAnywhere, Category = "Live Settings")
+	bool bDrawMorphTargets = false;
+#endif // WITH_EDITORONLY_DATA
+};
diff --git a/Engine/Plugins/Experimental/MLDeformer/NeuralMorphModel/Source/NeuralMorphModelEditor/NeuralMorphModelEditor.Build.cs b/Engine/Plugins/Experimental/MLDeformer/NeuralMorphModel/Source/NeuralMorphModelEditor/NeuralMorphModelEditor.Build.cs
new file mode 100644
--- /dev/null
+++ b/Engine/Plugins/Experimental/MLDeformer/NeuralMorphModel/Source/NeuralMorphModelEditor/NeuralMorphModelEditor.Build.cs
@@ -0,0 +1,49 @@
+// Copyright Epic Games, Inc. All Rights Reserved.
+
+using System.IO;
+
+namespace UnrealBuildTool.Rules
+{
+	public class NeuralMorphModelEditor : ModuleRules
+	{
+		public NeuralMorphModelEditor(ReadOnlyTargetRules Target) : base(Target)
+		{
+			PublicDependencyModuleNames.AddRange(
+				new string[]
+				{
+					"Core"
+				}
+			);
+
+			PrivateDependencyModuleNames.AddRange(
+				new string[]
+				{
+					"EditorFramework",
+					"UnrealEd",
+					"CoreUObject",
+					"Engine",
+					"Slate",
+					"SlateCore",
+					"EditorWidgets",
+					"EditorStyle",
+					"GeometryCache",
+					"MLDeformerFramework",
+					"MLDeformerFrameworkEditor",
+					"NeuralMorphModel",
+					"PropertyEditor",
+					"NeuralNetworkInference",
+					"ToolWidgets",
+					"ComputeFramework",
+					"RenderCore",
+					"RHI"
+				}
+			);
+
+			PrivateIncludePathModuleNames.AddRange(
+				new string[]
+				{
+				}
+			);
+		}
+	}
+}
diff --git a/Engine/Plugins/Experimental/MLDeformer/NeuralMorphModel/Source/NeuralMorphModelEditor/Private/NeuralMorphEditorModel.cpp b/Engine/Plugins/Experimental/MLDeformer/NeuralMorphModel/Source/NeuralMorphModelEditor/Private/NeuralMorphEditorModel.cpp
new file mode 100644
--- /dev/null
+++ b/Engine/Plugins/Experimental/MLDeformer/NeuralMorphModel/Source/NeuralMorphModelEditor/Private/NeuralMorphEditorModel.cpp
@@ -0,0 +1,410 @@
+// Copyright Epic Games, Inc. All Rights Reserved.
+
+#include "NeuralMorphEditorModel.h"
+#include "NeuralMorphModelModule.h"
+#include "NeuralMorphModel.h"
+#include "NeuralMorphModelVizSettings.h"
+#include "NeuralMorphTrainingModel.h"
+#include "NeuralMorphEditorModelActor.h"
+#include "MLDeformerEditorToolkit.h"
+#include "MLDeformerEditorStyle.h"
+#include "MLDeformerInputInfo.h"
+#include "MLDeformerGeomCacheSampler.h"
+#include "MLDeformerGeomCacheHelpers.h"
+#include "GeometryCache.h"
+#include "GeometryCacheComponent.h"
+#include "Animation/DebugSkelMeshComponent.h"
+#include "Animation/MorphTarget.h"
+#include "Rendering/SkeletalMeshModel.h"
+#include "Rendering/SkeletalMeshLODModel.h"
+#include "Engine/SkeletalMesh.h"
+
+#define LOCTEXT_NAMESPACE "NeuralMorphEditorModel"
+
+namespace UE::NeuralMorphModel
+{
+	using namespace UE::MLDeformer;
+
+	FNeuralMorphEditorModel::~FNeuralMorphEditorModel()
+	{
+		Model->OnPostEditChangeProperty().Unbind();
+	}
+
+	FMLDeformerEditorModel* FNeuralMorphEditorModel::MakeInstance()
+	{
+		return new FNeuralMorphEditorModel();
+	}
+
+	void FNeuralMorphEditorModel::AddReferencedObjects(FReferenceCollector& Collector)
+	{
+		Collector.AddReferencedObject(EditorInputInfo);
+	}
+
+	FMLDeformerSampler* FNeuralMorphEditorModel::CreateSampler() const 
+	{ 
+		FMLDeformerGeomCacheSampler* NewSampler = new FMLDeformerGeomCacheSampler();
+		NewSampler->OnGetGeometryCache().BindLambda([this] { return GetNeuralMorphModel()->GetGeometryCache(); });
+		return NewSampler;
+	}
+
+	void FNeuralMorphEditorModel::Init(const InitSettings& InitSettings)
+	{
+		FMLDeformerEditorModel::Init(InitSettings);
+		Model->OnPostEditChangeProperty().BindRaw(this, &FNeuralMorphEditorModel::OnPostEditChangeProperty);
+	}
+
+	void FNeuralMorphEditorModel::OnPostEditChangeProperty(FPropertyChangedEvent& PropertyChangedEvent)
+	{
+		const FProperty* Property = PropertyChangedEvent.Property;
+		if (Property == nullptr)
+		{
+			return;
+		}
+
+		HandleDefaultPropertyChanges(PropertyChangedEvent);
+
+		// When we change one of these properties below, restart animations etc.
+		if (Property->GetFName() == GET_MEMBER_NAME_CHECKED(UNeuralMorphModel, GeometryCache) || 
+			Property->GetFName() == GET_MEMBER_NAME_CHECKED(UNeuralMorphModelVizSettings, GroundTruth))
+		{
+			TriggerInputAssetChanged(true);
+		}
+		else
+		if (Property->GetFName() == GET_MEMBER_NAME_CHECKED(UNeuralMorphModel, DeltaCutoffLength) || 
+			Property->GetFName() == GET_MEMBER_NAME_CHECKED(UMLDeformerModel, AlignmentTransform))
+		{
+			if (PropertyChangedEvent.ChangeType == EPropertyChangeType::ValueSet)
+			{
+				SetResamplingInputOutputsNeeded(true);
+				SampleDeltas();
+			}
+		}
+	}
+
+	void FNeuralMorphEditorModel::OnInputAssetsChanged()
+	{
+		// Update the skeletal mesh components of the training, test base, and ml deformed actor.
+		FMLDeformerEditorModel::OnInputAssetsChanged();
+
+		UNeuralMorphModelVizSettings* VizSettings = GetNeuralMorphModelVizSettings();
+		check(VizSettings);
+		const float TestAnimSpeed = VizSettings->GetAnimPlaySpeed();
+		UAnimSequence* TestAnimSequence = VizSettings->GetTestAnimSequence();
+
+		// Update the training geometry cache.
+		UGeometryCacheComponent* GeometryCacheComponent = FindNeuralMorphModelEditorActor(ActorID_Train_GroundTruth)->GetGeometryCacheComponent();
+		check(GeometryCacheComponent);
+		GeometryCacheComponent->SetGeometryCache(GetNeuralMorphModel()->GetGeometryCache());
+		GeometryCacheComponent->SetLooping(false);
+		GeometryCacheComponent->SetManualTick(true);
+		GeometryCacheComponent->SetPlaybackSpeed(TestAnimSpeed);
+		GeometryCacheComponent->Play();
+
+		// Update the test geometry cache (ground truth) component.
+		GeometryCacheComponent = FindNeuralMorphModelEditorActor(ActorID_Test_GroundTruth)->GetGeometryCacheComponent();
+		GeometryCacheComponent->SetGeometryCache(VizSettings->GetTestGroundTruth());
+		GeometryCacheComponent->SetLooping(true);
+		GeometryCacheComponent->SetManualTick(true);
+		GeometryCacheComponent->SetPlaybackSpeed(TestAnimSpeed);
+		GeometryCacheComponent->Play();
+
+		UNeuralMorphModel* NeuralMorphModel = GetNeuralMorphModel();
+		NeuralMorphModel->MeshMappings.Reset();
+	}
+
+	void FNeuralMorphEditorModel::CreateTrainingGroundTruthActor(UWorld* World)
+	{
+		UGeometryCache* GeomCache = GetNeuralMorphModel()->GetGeometryCache();
+		const FLinearColor LabelColor = FMLDeformerEditorStyle::Get().GetColor("MLDeformer.TargetMesh.LabelColor");
+		const FLinearColor WireframeColor = FMLDeformerEditorStyle::Get().GetColor("MLDeformer.TargetMesh.WireframeColor");
+		CreateGeomCacheActor(
+			World, 
+			ActorID_Train_GroundTruth, 
+			"Train GroundTruth", 
+			GeomCache, 
+			LabelColor, 
+			WireframeColor, 
+			LOCTEXT("TrainGroundTruthActorLabelText", "Target Mesh"),
+			true);
+	}
+
+	void FNeuralMorphEditorModel::CreateTestGroundTruthActor(UWorld* World)
+	{
+		UGeometryCache* GeomCache = GetNeuralMorphModelVizSettings()->GetTestGroundTruth();
+		const FLinearColor LabelColor = FMLDeformerEditorStyle::Get().GetColor("MLDeformer.GroundTruth.LabelColor");
+		const FLinearColor WireframeColor = FMLDeformerEditorStyle::Get().GetColor("MLDeformer.GroundTruth.WireframeColor");
+		CreateGeomCacheActor(
+			World, 
+			ActorID_Test_GroundTruth, 
+			"Test GroundTruth", 
+			GeomCache, 
+			LabelColor, 
+			WireframeColor,
+			LOCTEXT("TestGroundTruthActorLabelText", "Ground Truth"),
+			false);
+	}
+
+	void FNeuralMorphEditorModel::CreateGeomCacheActor(UWorld* World, int32 ActorID, const FName& Name, UGeometryCache* GeomCache, FLinearColor LabelColor, FLinearColor WireframeColor, const FText& LabelText, bool bIsTrainingActor)
+	{
+		FActorSpawnParameters SpawnParams;
+		SpawnParams.Name = MakeUniqueObjectName(World, AActor::StaticClass(), Name);
+		AActor* Actor = World->SpawnActor<AActor>(SpawnParams);
+		Actor->SetFlags(RF_Transient);
+
+		// Create the Geometry Cache Component.
+		UGeometryCacheComponent* GeomCacheComponent = NewObject<UGeometryCacheComponent>(Actor);
+		GeomCacheComponent->SetGeometryCache(GeomCache);
+		GeomCacheComponent->RegisterComponent();
+		GeomCacheComponent->SetOverrideWireframeColor(true);
+		GeomCacheComponent->SetWireframeOverrideColor(WireframeColor);
+		GeomCacheComponent->MarkRenderStateDirty();
+		GeomCacheComponent->SetVisibility(false);
+		Actor->SetRootComponent(GeomCacheComponent);
+
+		// Create the editor actor.
+		FMLDeformerEditorActor::FConstructSettings Settings;
+		Settings.Actor = Actor;
+		Settings.TypeID = ActorID;
+		Settings.LabelColor = LabelColor;
+		Settings.LabelText = LabelText;
+		Settings.bIsTrainingActor = bIsTrainingActor;
+		FNeuralMorphEditorModelActor* EditorActor = static_cast<FNeuralMorphEditorModelActor*>(CreateEditorActor(Settings));
+		EditorActor->SetGeometryCacheComponent(GeomCacheComponent);
+		EditorActors.Add(EditorActor);
+	} 
+
+	FMLDeformerEditorActor* FNeuralMorphEditorModel::CreateEditorActor(const FMLDeformerEditorActor::FConstructSettings& Settings) const
+	{
+		return new FNeuralMorphEditorModelActor(Settings);
+	}
+
+	FNeuralMorphEditorModelActor* FNeuralMorphEditorModel::FindNeuralMorphModelEditorActor(int32 TypeID) const
+	{
+		return static_cast<FNeuralMorphEditorModelActor*>(FindEditorActor(TypeID));
+	}
+
+	UNeuralMorphModelVizSettings* FNeuralMorphEditorModel::GetNeuralMorphModelVizSettings() const
+	{
+		return Cast<UNeuralMorphModelVizSettings>(GetNeuralMorphModel()->GetVizSettings());
+	}
+
+	double FNeuralMorphEditorModel::GetTrainingTimeAtFrame(int32 FrameNumber) const
+	{
+		// Try to get the frame from the geometry cache.
+		FNeuralMorphEditorModelActor* EditorActor = static_cast<FNeuralMorphEditorModelActor*>(FindEditorActor(ActorID_Train_GroundTruth));
+		if (EditorActor && EditorActor->GetGeometryCacheComponent() && EditorActor->GetGeometryCacheComponent()->GeometryCache.Get())
+		{
+			return EditorActor->GetGeometryCacheComponent()->GetTimeAtFrame(FrameNumber);
+		}
+
+		return FMLDeformerEditorModel::GetTrainingTimeAtFrame(FrameNumber);
+	}
+
+	int32 FNeuralMorphEditorModel::GetTrainingFrameAtTime(double TimeInSeconds) const
+	{
+		FNeuralMorphEditorModelActor* EditorActor = static_cast<FNeuralMorphEditorModelActor*>(FindEditorActor(ActorID_Train_GroundTruth));
+		if (EditorActor && EditorActor->GetGeometryCacheComponent() && EditorActor->GetGeometryCacheComponent()->GeometryCache.Get())
+		{
+			return EditorActor->GetGeometryCacheComponent()->GetFrameAtTime(TimeInSeconds);
+		}
+
+		return FMLDeformerEditorModel::GetTrainingFrameAtTime(TimeInSeconds);
+	}
+
+	double FNeuralMorphEditorModel::GetTestTimeAtFrame(int32 FrameNumber) const
+	{
+		// Try to get the frame from the geometry cache.
+		FNeuralMorphEditorModelActor* EditorActor = static_cast<FNeuralMorphEditorModelActor*>(FindEditorActor(ActorID_Test_GroundTruth));
+		if (EditorActor && EditorActor->GetGeometryCacheComponent() && EditorActor->GetGeometryCacheComponent()->GeometryCache.Get())
+		{
+			return EditorActor->GetGeometryCacheComponent()->GetTimeAtFrame(FrameNumber);
+		}
+
+		return FMLDeformerEditorModel::GetTestTimeAtFrame(FrameNumber);
+	}
+
+	int32 FNeuralMorphEditorModel::GetTestFrameAtTime(double TimeInSeconds) const
+	{
+		FNeuralMorphEditorModelActor* EditorActor = static_cast<FNeuralMorphEditorModelActor*>(FindEditorActor(ActorID_Test_GroundTruth));
+		if (EditorActor && EditorActor->GetGeometryCacheComponent() && EditorActor->GetGeometryCacheComponent()->GeometryCache.Get())
+		{
+			return EditorActor->GetGeometryCacheComponent()->GetFrameAtTime(TimeInSeconds);
+		}
+
+		return FMLDeformerEditorModel::GetTestFrameAtTime(TimeInSeconds);
+	}
+
+	int32 FNeuralMorphEditorModel::GetNumTrainingFrames() const
+	{
+		const UGeometryCache* GeometryCache = GetNeuralMorphModel()->GetGeometryCache();
+		if (GeometryCache == nullptr)
+		{
+			return 0;
+		}
+		const int32 StartFrame = GeometryCache->GetStartFrame();
+		const int32 EndFrame = GeometryCache->GetEndFrame();
+		return (EndFrame - StartFrame) + 1;
+	}
+
+	void FNeuralMorphEditorModel::UpdateIsReadyForTrainingState()
+	{
+		bIsReadyForTraining = false;
+
+		// Do some basic checks first, like if there is a skeletal mesh, ground truth, anim sequence, and if there are frames.
+		if (!IsEditorReadyForTrainingBasicChecks())
+		{
+			return;
+		}
+
+		// Now make sure the assets are compatible.
+		UNeuralMorphModel* NeuralMorphModel = GetNeuralMorphModel();
+		UGeometryCache* GeomCache = NeuralMorphModel->GetGeometryCache();
+		USkeletalMesh* SkeletalMesh = NeuralMorphModel->GetSkeletalMesh();
+		if (!GetGeomCacheVertexErrorText(SkeletalMesh, GeomCache, FText(), FText()).IsEmpty() ||
+			!GetGeomCacheErrorText(SkeletalMesh, GeomCache).IsEmpty())
+		{
+			return;
+		}
+
+		// Make sure every skeletal imported mesh has some geometry track.
+		const int32 NumGeomCacheTracks = GeomCache ? GeomCache->Tracks.Num() : 0;
+		int32 NumSkelMeshes = 0;
+		if (SkeletalMesh)
+		{
+			FSkeletalMeshModel* ImportedModel = SkeletalMesh->GetImportedModel();
+			if (ImportedModel)
+			{
+				NumSkelMeshes = ImportedModel->LODModels[0].ImportedMeshInfos.Num();		
+			}
+		}
+
+		// Allow the special case where there is just one mesh and track.
+		if (NumGeomCacheTracks != 1 || NumSkelMeshes != 1)
+		{
+			if (!GetGeomCacheSampler()->GetFailedImportedMeshNames().IsEmpty())
+			{
+				return;
+			}
+		}
+
+		bIsReadyForTraining = true;
+	}
+
+	ETrainingResult FNeuralMorphEditorModel::Train()
+	{
+		return TrainModel<UNeuralMorphTrainingModel>(this);
+	}
+
+	FString FNeuralMorphEditorModel::GetTrainedNetworkOnnxFile() const
+	{
+		return FString(FPaths::ProjectIntermediateDir() + TEXT("NeuralMorphModel/NeuralMorphModel.onnx"));
+	}
+
+	FString FNeuralMorphEditorModel::GetDefaultDeformerGraphAssetPath() const
+	{
+		return FString(TEXT("/Script/OptimusCore.OptimusDeformer'/Optimus/Deformers/DG_LinearBlendSkin_Morph_Cloth_RecomputeNormals.DG_LinearBlendSkin_Morph_Cloth_RecomputeNormals'"));
+	}
+
+	FString FNeuralMorphEditorModel::GetHeatMapDeformerGraphPath() const
+	{
+		return FString(TEXT("/Script/OptimusCore.OptimusDeformer'/Optimus/Deformers/DG_LinearBlendSkin_Morph_Cloth_RecomputeNormals.DG_LinearBlendSkin_Morph_Cloth_RecomputeNormals'"));
+	}
+
+	void FNeuralMorphEditorModel::OnPreTraining()
+	{
+		// Backup the morph target deltas in case we abort training.
+		UNeuralMorphModel* NeuralModel = GetNeuralMorphModel();
+		MorphTargetDeltasBackup = NeuralModel->MorphTargetDeltas;
+	}
+
+	void FNeuralMorphEditorModel::OnPostTraining(ETrainingResult TrainingResult, bool bUsePartiallyTrainedWhenAborted)
+	{
+		// We aborted and don't want to use partially trained results, we should restore the deltas that we just overwrote after training.
+		if (TrainingResult == ETrainingResult::Aborted && !bUsePartiallyTrainedWhenAborted)
+		{
+			// Restore the blendshape backup.
+			UNeuralMorphModel* NeuralModel = GetNeuralMorphModel();
+			NeuralModel->MorphTargetDeltas = MorphTargetDeltasBackup;
+		}
+		else if (TrainingResult == ETrainingResult::Success || (TrainingResult == ETrainingResult::Aborted && bUsePartiallyTrainedWhenAborted))
+		{
+			UNeuralMorphModel* NeuralModel = GetNeuralMorphModel();
+			if (!NeuralModel->MorphTargetDeltas.IsEmpty())
+			{
+				// Set deltas with a length equal or below a given threshold to zero, for better compression.
+				TArray<FVector3f> MorphTargetDeltas = NeuralModel->MorphTargetDeltas;
+				ZeroDeltasByThreshold(MorphTargetDeltas, NeuralModel->MorphTargetDeltaThreshold);
+
+				// Build morph targets inside the engine, using the engine's compression scheme.
+				// Add one as we included the means now as extra morph target.
+				InitEngineMorphTargets(MorphTargetDeltas);
+			}
+		}
+
+		// This internally calls InitGPUData() which updates the GPU buffer with the deltas.
+		FMLDeformerEditorModel::OnPostTraining(TrainingResult, bUsePartiallyTrainedWhenAborted);
+	}
+
+	void FNeuralMorphEditorModel::InitEngineMorphTargets(const TArray<FVector3f>& Deltas)
+	{
+		const int32 LOD = 0;
+
+		// Turn the delta buffer in a set of engine morph targets.
+		UNeuralMorphModel* NeuralModel = GetNeuralMorphModel();
+		TArray<UMorphTarget*> MorphTargets;	// These will be garbage collected.
+		CreateEngineMorphTargets(MorphTargets, Deltas, TEXT("NeuralMorph_"), LOD, NeuralModel->MorphTargetDeltaThreshold);
+
+		// Now compress the morph targets to GPU friendly buffers.
+		check(NeuralModel->MorphTargetSet.IsValid());
+		FMorphTargetVertexInfoBuffers& MorphBuffers = NeuralModel->MorphTargetSet->MorphBuffers;
+		CompressEngineMorphTargets(MorphBuffers, MorphTargets, LOD, NeuralModel->MorphTargetErrorTolerance);
+	}
+
+	void FNeuralMorphEditorModel::Render(const FSceneView* View, FViewport* Viewport, FPrimitiveDrawInterface* PDI)
+	{
+		FMLDeformerEditorModel::Render(View, Viewport, PDI);
+
+		UNeuralMorphModelVizSettings* VizSettings = GetNeuralMorphModelVizSettings();
+		if (VizSettings->bDrawMorphTargets)
+		{
+			// Check if the vertex count matches.
+			UNeuralMorphModel* NeuralModel = GetNeuralMorphModel();
+			if (!NeuralModel->MorphTargetDeltas.IsEmpty())
+			{
+				const int32 NumVerts = NeuralModel->GetNumBaseMeshVerts();
+				check(NeuralModel->MorphTargetDeltas.Num() % NumVerts == 0);		
+
+				const TArray<FVector3f>& SkinnedPositions = Sampler->GetSkinnedVertexPositions();
+				check(NumVerts == SkinnedPositions.Num());
+
+				const int32 NumMorphTargets = NeuralModel->MorphTargetDeltas.Num() / NumVerts;
+
+				// Draw all deltas.
+				const FVector MeshOffset = -VizSettings->GetMeshSpacingOffsetVector();
+				const int32 BlendShapeIndex = FMath::Clamp<int32>(VizSettings->MorphTargetNumber, 0, NumMorphTargets - 1);
+				const int32 BlendShapeDeltaStartOffset = NeuralModel->GetMorphTargetDeltaStartIndex(BlendShapeIndex);
+				const FLinearColor IncludedColor(1.0f, 0.0f, 1.0f);
+				const FLinearColor ExcludedColor(0.2f, 0.2f, 0.2f);
+				const float DeltaThreshold = VizSettings->MorphTargetDeltaThreshold;
+				for (int32 VertexIndex = 0; VertexIndex < NumVerts; ++VertexIndex)
+				{
+					const FVector StartPoint = FVector(SkinnedPositions[VertexIndex].X, SkinnedPositions[VertexIndex].Y, SkinnedPositions[VertexIndex].Z) + MeshOffset;
+					const int32 DeltaArrayOffset = BlendShapeDeltaStartOffset + VertexIndex;
+					const FVector Delta = FVector(NeuralModel->MorphTargetDeltas[DeltaArrayOffset]);
+
+					if (Delta.Length() > DeltaThreshold)
+					{
+						PDI->DrawLine(StartPoint, StartPoint + Delta, IncludedColor, 10);
+						PDI->DrawPoint(StartPoint + Delta, IncludedColor, 1.0f, 10);
+					}
+					else
+					{
+						PDI->DrawPoint(StartPoint + Delta, ExcludedColor, 0.5f, 10);
+					}
+				}
+			}
+		}
+	}
+}	// namespace UE::NeuralMorphModel
+
+#undef LOCTEXT_NAMESPACE
diff --git a/Engine/Plugins/Experimental/MLDeformer/NeuralMorphModel/Source/NeuralMorphModelEditor/Private/NeuralMorphEditorModelActor.cpp b/Engine/Plugins/Experimental/MLDeformer/NeuralMorphModel/Source/NeuralMorphModelEditor/Private/NeuralMorphEditorModelActor.cpp
new file mode 100644
--- /dev/null
+++ b/Engine/Plugins/Experimental/MLDeformer/NeuralMorphModel/Source/NeuralMorphModelEditor/Private/NeuralMorphEditorModelActor.cpp
@@ -0,0 +1,131 @@
+// Copyright Epic Games, Inc. All Rights Reserved.
+
+#include "NeuralMorphEditorModelActor.h"
+#include "GeometryCacheComponent.h"
+#include "Animation/DebugSkelMeshComponent.h"
+
+namespace UE::NeuralMorphModel
+{
+	using namespace UE::MLDeformer;
+
+	FNeuralMorphEditorModelActor::FNeuralMorphEditorModelActor(const FConstructSettings& Settings)
+		: FMLDeformerEditorActor(Settings)
+	{
+	}
+
+	void FNeuralMorphEditorModelActor::SetVisibility(bool bIsVisible)
+	{
+		FMLDeformerEditorActor::SetVisibility(bIsVisible);
+
+		if (GeomCacheComponent && bIsVisible != GeomCacheComponent->IsVisible())
+		{
+			GeomCacheComponent->SetVisibility(bIsVisible, true);
+		}
+	}
+
+	bool FNeuralMorphEditorModelActor::IsVisible() const
+	{
+		if (SkeletalMeshComponent)
+		{
+			return SkeletalMeshComponent->IsVisible();
+		}
+
+		if (GeomCacheComponent)
+		{
+			return GeomCacheComponent->IsVisible();
+		}
+
+		return true;
+	}
+
+	bool FNeuralMorphEditorModelActor::HasVisualMesh() const
+	{
+		if (SkeletalMeshComponent && SkeletalMeshComponent->SkeletalMesh)
+		{
+			return true;
+		}
+
+		if (GeomCacheComponent && GeomCacheComponent->GetGeometryCache())
+		{
+			return true;
+		}
+
+		return false;
+	}
+
+	void FNeuralMorphEditorModelActor::SetPlayPosition(float TimeInSeconds, bool bAutoPause)
+	{
+		if (SkeletalMeshComponent)
+		{
+			SkeletalMeshComponent->SetPosition(TimeInSeconds);
+			if (bAutoPause)
+			{
+				SkeletalMeshComponent->bPauseAnims = true;
+			}
+		}
+
+		if (GeomCacheComponent)
+		{
+			GeomCacheComponent->SetManualTick(true);
+			GeomCacheComponent->TickAtThisTime(TimeInSeconds, false, false, false);
+		}
+	}
+
+	float FNeuralMorphEditorModelActor::GetPlayPosition() const
+	{
+		if (SkeletalMeshComponent)
+		{
+			return SkeletalMeshComponent->GetPosition();
+		}
+
+		if (GeomCacheComponent)
+		{
+			return GeomCacheComponent->GetAnimationTime();
+		}
+
+		return 0.0f;
+	}
+
+	void FNeuralMorphEditorModelActor::SetPlaySpeed(float PlaySpeed)
+	{
+		if (SkeletalMeshComponent)
+		{
+			SkeletalMeshComponent->SetPlayRate(PlaySpeed);
+		}
+
+		if (GeomCacheComponent)
+		{
+			GeomCacheComponent->SetPlaybackSpeed(PlaySpeed);
+		}
+	}
+
+	void FNeuralMorphEditorModelActor::Pause(bool bPaused)
+	{
+		if (SkeletalMeshComponent)
+		{
+			SkeletalMeshComponent->bPauseAnims = bPaused;
+		}
+
+		if (GeomCacheComponent)
+		{
+			GeomCacheComponent->SetManualTick(bPaused);
+		}
+	}
+
+	FBox FNeuralMorphEditorModelActor::GetBoundingBox() const
+	{
+		if (SkeletalMeshComponent)
+		{
+			return SkeletalMeshComponent->Bounds.GetBox();
+		}
+
+		if (GeomCacheComponent)
+		{
+			return GeomCacheComponent->Bounds.GetBox();
+		}
+
+		FBox Box;
+		Box.Init();
+		return Box;
+	}
+}	// namespace UE::NeuralMorphModel
diff --git a/Engine/Plugins/Experimental/MLDeformer/NeuralMorphModel/Source/NeuralMorphModelEditor/Private/NeuralMorphModelDetails.cpp b/Engine/Plugins/Experimental/MLDeformer/NeuralMorphModel/Source/NeuralMorphModelEditor/Private/NeuralMorphModelDetails.cpp
new file mode 100644
--- /dev/null
+++ b/Engine/Plugins/Experimental/MLDeformer/NeuralMorphModel/Source/NeuralMorphModelEditor/Private/NeuralMorphModelDetails.cpp
@@ -0,0 +1,175 @@
+// Copyright Epic Games, Inc. All Rights Reserved.
+
+#include "NeuralMorphModelDetails.h"
+#include "NeuralMorphEditorModel.h"
+#include "NeuralMorphModel.h"
+#include "MLDeformerModule.h"
+#include "MLDeformerAsset.h"
+#include "MLDeformerModel.h"
+#include "MLDeformerEditorModel.h"
+#include "MLDeformerEditorModule.h"
+#include "MLDeformerGeomCacheHelpers.h"
+#include "NeuralNetwork.h"
+#include "GeometryCache.h"
+#include "GeometryCacheTrack.h"
+#include "Animation/AnimSequence.h"
+#include "Engine/SkeletalMesh.h"
+#include "PropertyCustomizationHelpers.h"
+#include "AssetRegistry/AssetData.h"
+#include "Modules/ModuleManager.h"
+#include "DetailLayoutBuilder.h"
+#include "DetailCategoryBuilder.h"
+#include "DetailWidgetRow.h"
+#include "IDetailsView.h"
+#include "IDetailGroup.h"
+#include "Widgets/DeclarativeSyntaxSupport.h"
+#include "Widgets/Layout/SBox.h"
+#include "Widgets/Input/SButton.h"
+#include "DetailLayoutBuilder.h"
+#include "SWarningOrErrorBox.h"
+
+#define LOCTEXT_NAMESPACE "NeuralMorphModelDetails"
+
+namespace UE::NeuralMorphModel
+{
+	using namespace UE::MLDeformer;
+
+	TSharedRef<IDetailCustomization> FNeuralMorphModelDetails::MakeInstance()
+	{
+		return MakeShareable(new FNeuralMorphModelDetails());
+	}
+
+	bool FNeuralMorphModelDetails::UpdateMemberPointers(const TArray<TWeakObjectPtr<UObject>>& Objects)
+	{
+		if (!FMLDeformerModelDetails::UpdateMemberPointers(Objects))
+		{
+			return false;
+		}
+
+		NeuralMorphModel = Cast<UNeuralMorphModel>(Model);
+		check(NeuralMorphModel);
+		NeuralMorphEditorModel = static_cast<FNeuralMorphEditorModel*>(EditorModel);
+
+		return (NeuralMorphModel != nullptr && NeuralMorphEditorModel != nullptr);
+	}
+
+	void FNeuralMorphModelDetails::CustomizeDetails(IDetailLayoutBuilder& DetailBuilder)
+	{
+		// Create all the detail categories and add the properties of the base class.
+		FMLDeformerModelDetails::CustomizeDetails(DetailBuilder);
+
+		// Training settings.
+		SettingsCategoryBuilder->AddProperty(GET_MEMBER_NAME_CHECKED(UNeuralMorphModel, NumMorphTargetsPerBone));
+		SettingsCategoryBuilder->AddProperty(GET_MEMBER_NAME_CHECKED(UNeuralMorphModel, NumHiddenLayers));
+		SettingsCategoryBuilder->AddProperty(GET_MEMBER_NAME_CHECKED(UNeuralMorphModel, NumNeuronsPerLayer));
+		SettingsCategoryBuilder->AddProperty(GET_MEMBER_NAME_CHECKED(UNeuralMorphModel, NumIterations));
+		SettingsCategoryBuilder->AddProperty(GET_MEMBER_NAME_CHECKED(UNeuralMorphModel, BatchSize));
+		SettingsCategoryBuilder->AddProperty(GET_MEMBER_NAME_CHECKED(UNeuralMorphModel, LearningRate));
+		SettingsCategoryBuilder->AddProperty(GET_MEMBER_NAME_CHECKED(UNeuralMorphModel, RegularizationFactor));
+		SettingsCategoryBuilder->AddProperty(GET_MEMBER_NAME_CHECKED(UNeuralMorphModel, MorphTargetErrorTolerance));
+	}
+
+	void FNeuralMorphModelDetails::AddBaseMeshErrors()
+	{
+		// Check if the base mesh matches the target mesh vertex count.
+		FText ErrorText = GetGeomCacheVertexErrorText(Model->SkeletalMesh, NeuralMorphModel->GetGeometryCache(), FText::FromString("Base Mesh"), FText::FromString("Target Mesh"));
+		BaseMeshCategoryBuilder->AddCustomRow(FText::FromString("BaseMeshError"))
+			.Visibility(!ErrorText.IsEmpty() ? EVisibility::Visible : EVisibility::Collapsed)
+			.WholeRowContent()
+			[
+				SNew(SBox)
+				.Padding(FMargin(0.0f, 4.0f))
+				[
+					SNew(SWarningOrErrorBox)
+					.MessageStyle(EMessageStyle::Error)
+					.Message(ErrorText)
+				]
+			];
+	}
+	
+	void FNeuralMorphModelDetails::AddAnimSequenceErrors()
+	{
+		const FText WarningText = GetGeomCacheAnimSequenceErrorText(NeuralMorphModel->GetGeometryCache(), Model->GetAnimSequence());
+		BaseMeshCategoryBuilder->AddCustomRow(FText::FromString("AnimSeqWarning"))
+			.Visibility(!WarningText.IsEmpty() ? EVisibility::Visible : EVisibility::Collapsed)
+			.WholeRowContent()
+			[
+				SNew(SBox)
+				.Padding(FMargin(0.0f, 4.0f))
+				[
+					SNew(SWarningOrErrorBox)
+					.MessageStyle(EMessageStyle::Warning)
+					.Message(WarningText)
+				]
+			];
+	}
+
+	void FNeuralMorphModelDetails::AddTargetMesh()
+	{
+		TargetMeshCategoryBuilder->AddProperty(GET_MEMBER_NAME_CHECKED(UNeuralMorphModel, GeometryCache));
+
+		const FText TargetMeshErrorText = GetGeomCacheErrorText(NeuralMorphModel->GetSkeletalMesh(), NeuralMorphModel->GetGeometryCache());
+		TargetMeshCategoryBuilder->AddCustomRow(FText::FromString("TargetMeshError"))
+			.Visibility(!TargetMeshErrorText.IsEmpty() ? EVisibility::Visible : EVisibility::Collapsed)
+			.WholeRowContent()
+			[
+				SNew(SBox)
+				.Padding(FMargin(0.0f, 4.0f))
+				[
+					SNew(SWarningOrErrorBox)
+					.MessageStyle(EMessageStyle::Error)
+					.Message(TargetMeshErrorText)
+				]
+			];
+
+		const FText ChangedErrorText = EditorModel->GetTargetAssetChangedErrorText();
+		TargetMeshCategoryBuilder->AddCustomRow(FText::FromString("TargetMeshChangedError"))
+			.Visibility(!ChangedErrorText.IsEmpty() ? EVisibility::Visible : EVisibility::Collapsed)
+			.WholeRowContent()
+			[
+				SNew(SBox)
+				.Padding(FMargin(0.0f, 4.0f))
+				[
+					SNew(SWarningOrErrorBox)
+					.MessageStyle(EMessageStyle::Error)
+					.Message(ChangedErrorText)
+				]
+			];
+
+		// Show meshes that have no matching geometry track list.
+		const FText MeshMappingErrorList = GetGeomCacheMeshMappingErrorText(Model->GetSkeletalMesh(), NeuralMorphModel->GetGeometryCache());
+		FString GeomTrackNameList;
+		if (!MeshMappingErrorList.IsEmpty())
+		{
+			UGeometryCache* GeomCache = NeuralMorphModel->GetGeometryCache();
+			for (int32 Index = 0; Index < GeomCache->Tracks.Num(); ++Index)
+			{
+				GeomTrackNameList += GeomCache->Tracks[Index]->GetName();
+				if (Index < GeomCache->Tracks.Num() - 1)
+				{
+					GeomTrackNameList += TEXT("\n");
+				}
+			}
+		}
+
+		FText MeshMappingErrorFull = FText::Format(
+			LOCTEXT("MeshMappingError", "No matching GeomCache Tracks names found for meshes:\n{0}\n\nGeomCache Track List:\n{1}"), 
+			MeshMappingErrorList,
+			FText::FromString(GeomTrackNameList));
+
+		TargetMeshCategoryBuilder->AddCustomRow(FText::FromString("MeshMappingError"))
+			.Visibility(!MeshMappingErrorList.IsEmpty() ? EVisibility::Visible : EVisibility::Collapsed)
+			.WholeRowContent()
+			[
+				SNew(SBox)
+				.Padding(FMargin(0.0f, 4.0f))
+				[
+					SNew(SWarningOrErrorBox)
+					.MessageStyle(EMessageStyle::Error)
+					.Message(MeshMappingErrorFull)
+				]
+			];
+	}
+}	// namespace UE::NeuralMorphModel
+
+#undef LOCTEXT_NAMESPACE
diff --git a/Engine/Plugins/Experimental/MLDeformer/NeuralMorphModel/Source/NeuralMorphModelEditor/Private/NeuralMorphModelEditorModule.cpp b/Engine/Plugins/Experimental/MLDeformer/NeuralMorphModel/Source/NeuralMorphModelEditor/Private/NeuralMorphModelEditorModule.cpp
new file mode 100644
--- /dev/null
+++ b/Engine/Plugins/Experimental/MLDeformer/NeuralMorphModel/Source/NeuralMorphModelEditor/Private/NeuralMorphModelEditorModule.cpp
@@ -0,0 +1,56 @@
+// Copyright Epic Games, Inc. All Rights Reserved.
+
+#include "NeuralMorphModelEditorModule.h"
+#include "NeuralMorphModelVizSettingsDetails.h"
+#include "NeuralMorphModelDetails.h"
+#include "NeuralMorphEditorModel.h"
+#include "NeuralMorphModel.h"
+#include "MLDeformerEditorModule.h"
+#include "Modules/ModuleManager.h"
+#include "EditorModeRegistry.h"
+
+#define LOCTEXT_NAMESPACE "NeuralMorphModelEditorModule"
+
+IMPLEMENT_MODULE(UE::NeuralMorphModel::FNeuralMorphModelEditorModule, NeuralMorphModelEditor)
+
+namespace UE::NeuralMorphModel
+{
+	using namespace UE::MLDeformer;
+
+	void FNeuralMorphModelEditorModule::StartupModule()
+	{
+		// Register object detail customizations.
+		FPropertyEditorModule& PropertyModule = FModuleManager::LoadModuleChecked<FPropertyEditorModule>("PropertyEditor");
+		PropertyModule.RegisterCustomClassLayout("NeuralMorphModelVizSettings", FOnGetDetailCustomizationInstance::CreateStatic(&FNeuralMorphModelVizSettingsDetails::MakeInstance));
+		PropertyModule.RegisterCustomClassLayout("NeuralMorphModel", FOnGetDetailCustomizationInstance::CreateStatic(&FNeuralMorphModelDetails::MakeInstance));
+		PropertyModule.NotifyCustomizationModuleChanged();
+
+		// Register our custom ML deformer model to the model registry in the ML Deformer Framework.
+		FMLDeformerEditorModule& EditorModule = FModuleManager::LoadModuleChecked<FMLDeformerEditorModule>("MLDeformerFrameworkEditor");
+		FMLDeformerEditorModelRegistry& ModelRegistry = EditorModule.GetModelRegistry();
+		ModelRegistry.RegisterEditorModel(UNeuralMorphModel::StaticClass(), FOnGetEditorModelInstance::CreateStatic(&FNeuralMorphEditorModel::MakeInstance));
+	}
+
+	void FNeuralMorphModelEditorModule::ShutdownModule()
+	{
+		// Unregister our ML Deformer model.
+		if (FModuleManager::Get().IsModuleLoaded(TEXT("MLDeformerFrameworkEditor")))
+		{
+			FMLDeformerEditorModule& EditorModule = FModuleManager::GetModuleChecked<FMLDeformerEditorModule>("MLDeformerFrameworkEditor");
+			FMLDeformerEditorModelRegistry& ModelRegistry = EditorModule.GetModelRegistry();
+			ModelRegistry.UnregisterEditorModel(UNeuralMorphModel::StaticClass());
+		}
+
+		// Unregister object detail customizations for this model.
+		if (FModuleManager::Get().IsModuleLoaded(TEXT("PropertyEditor")))
+		{
+			FPropertyEditorModule& PropertyModule = FModuleManager::GetModuleChecked<FPropertyEditorModule>(TEXT("PropertyEditor"));
+			PropertyModule.UnregisterCustomClassLayout(TEXT("NeuralMorphModelVizSettings"));
+			PropertyModule.UnregisterCustomClassLayout(TEXT("NeuralMorphModel"));
+			PropertyModule.NotifyCustomizationModuleChanged();
+		}
+	}
+
+}	// namespace UE::NeuralMorphModel
+
+#undef LOCTEXT_NAMESPACE
diff --git a/Engine/Plugins/Experimental/MLDeformer/NeuralMorphModel/Source/NeuralMorphModelEditor/Private/NeuralMorphModelVizSettingsDetails.cpp b/Engine/Plugins/Experimental/MLDeformer/NeuralMorphModel/Source/NeuralMorphModelEditor/Private/NeuralMorphModelVizSettingsDetails.cpp
new file mode 100644
--- /dev/null
+++ b/Engine/Plugins/Experimental/MLDeformer/NeuralMorphModel/Source/NeuralMorphModelEditor/Private/NeuralMorphModelVizSettingsDetails.cpp
@@ -0,0 +1,123 @@
+// Copyright Epic Games, Inc. All Rights Reserved.
+
+#include "NeuralMorphModelVizSettingsDetails.h"
+#include "NeuralMorphModel.h"
+#include "NeuralMorphEditorModel.h"
+#include "NeuralMorphModelVizSettings.h"
+#include "MLDeformerModule.h"
+#include "MLDeformerModel.h"
+#include "MLDeformerEditorModule.h"
+#include "MLDeformerGeomCacheHelpers.h"
+#include "DetailLayoutBuilder.h"
+#include "DetailCategoryBuilder.h"
+#include "DetailWidgetRow.h"
+#include "IDetailsView.h"
+#include "IDetailChildrenBuilder.h"
+#include "IDetailGroup.h"
+#include "Modules/ModuleManager.h"
+#include "Layout/Margin.h"
+#include "Widgets/Layout/SBorder.h"
+#include "Widgets/Text/STextBlock.h"
+#include "Widgets/Layout/SBox.h"
+#include "PropertyCustomizationHelpers.h"
+#include "AssetRegistry/AssetData.h"
+#include "PropertyHandle.h"
+#include "DetailLayoutBuilder.h"
+#include "SWarningOrErrorBox.h"
+#include "GeometryCache.h"
+#include "Animation/AnimSequence.h"
+#include "Animation/MeshDeformer.h"
+#include "Engine/SkeletalMesh.h"
+
+#define LOCTEXT_NAMESPACE "NeuralMorphModelVizSettingsDetails"
+
+namespace UE::NeuralMorphModel
+{
+	using namespace UE::MLDeformer;
+
+	TSharedRef<IDetailCustomization> FNeuralMorphModelVizSettingsDetails::MakeInstance()
+	{
+		return MakeShareable(new FNeuralMorphModelVizSettingsDetails());
+	}
+
+	bool FNeuralMorphModelVizSettingsDetails::UpdateMemberPointers(const TArray<TWeakObjectPtr<UObject>>& Objects)
+	{
+		if (!FMLDeformerVizSettingsDetails::UpdateMemberPointers(Objects))
+		{
+			return false;
+		}
+
+		NeuralMorphModel = Cast<UNeuralMorphModel>(Model);
+		NeuralMorphVizSettings = Cast<UNeuralMorphModelVizSettings>(VizSettings);
+		return (NeuralMorphModel != nullptr && NeuralMorphVizSettings != nullptr);
+	}
+
+	void FNeuralMorphModelVizSettingsDetails::AddGroundTruth()
+	{
+		TestAssetsCategory->AddProperty(GET_MEMBER_NAME_CHECKED(UNeuralMorphModelVizSettings, GroundTruth));
+
+		// Show an error when the test anim sequence duration doesn't match the one of the ground truth.
+		const FText AnimErrorText = GetGeomCacheAnimSequenceErrorText(NeuralMorphVizSettings->GetTestGroundTruth(), VizSettings->GetTestAnimSequence());
+		FDetailWidgetRow& GroundTruthAnimErrorRow = TestAssetsCategory->AddCustomRow(FText::FromString("GroundTruthAnimMismatchError"))
+			.Visibility(!AnimErrorText.IsEmpty() ? EVisibility::Visible : EVisibility::Collapsed)
+			.WholeRowContent()
+			[
+				SNew(SBox)
+				.Padding(FMargin(0.0f, 4.0f))
+				[
+					SNew(SWarningOrErrorBox)
+					.MessageStyle(EMessageStyle::Warning)
+					.Message(AnimErrorText)
+				]
+			];
+
+		const FText GeomErrorText = GetGeomCacheErrorText(Model->GetSkeletalMesh(), NeuralMorphVizSettings->GetTestGroundTruth());
+		FDetailWidgetRow& GroundTruthGeomErrorRow = TestAssetsCategory->AddCustomRow(FText::FromString("GroundTruthGeomMismatchError"))
+			.Visibility(!GeomErrorText.IsEmpty() ? EVisibility::Visible : EVisibility::Collapsed)
+			.WholeRowContent()
+			[
+				SNew(SBox)
+				.Padding(FMargin(0.0f, 4.0f))
+				[
+					SNew(SWarningOrErrorBox)
+					.MessageStyle(EMessageStyle::Warning)
+					.Message(GeomErrorText)
+				]
+			];
+
+		const FText VertexErrorText = GetGeomCacheVertexErrorText(Model->GetSkeletalMesh(), NeuralMorphVizSettings->GetTestGroundTruth(), FText::FromString("Base Mesh"), FText::FromString("Ground Truth Mesh"));
+		FDetailWidgetRow& GroundTruthVertexErrorRow = TestAssetsCategory->AddCustomRow(FText::FromString("GroundTruthVertexMismatchError"))
+			.Visibility(!VertexErrorText.IsEmpty() ? EVisibility::Visible : EVisibility::Collapsed)
+			.WholeRowContent()
+			[
+				SNew(SBox)
+				.Padding(FMargin(0.0f, 4.0f))
+				[
+					SNew(SWarningOrErrorBox)
+					.MessageStyle(EMessageStyle::Warning)
+					.Message(VertexErrorText)
+				]
+			];
+	}
+
+	bool FNeuralMorphModelVizSettingsDetails::IsMorphTargetsEnabled() const
+	{
+		return NeuralMorphVizSettings->bDrawMorphTargets && !NeuralMorphModel->MorphTargetDeltas.IsEmpty();
+	}
+
+	void FNeuralMorphModelVizSettingsDetails::AddAdditionalSettings()
+	{
+		IDetailGroup& MorphsGroup = LiveSettingsCategory->AddGroup("Morph Targets", LOCTEXT("MorphTargetsLabel", "Morph Targets"), false, true);
+		MorphsGroup.AddPropertyRow(DetailLayoutBuilder->GetProperty(GET_MEMBER_NAME_CHECKED(UNeuralMorphModelVizSettings, bDrawMorphTargets)))
+			.EditCondition(!NeuralMorphModel->MorphTargetDeltas.IsEmpty(), nullptr);
+
+		MorphsGroup.AddPropertyRow(DetailLayoutBuilder->GetProperty(GET_MEMBER_NAME_CHECKED(UNeuralMorphModelVizSettings, MorphTargetNumber)))
+			.EditCondition(TAttribute<bool>(this, &FNeuralMorphModelVizSettingsDetails::IsMorphTargetsEnabled), nullptr);
+
+		MorphsGroup.AddPropertyRow(DetailLayoutBuilder->GetProperty(GET_MEMBER_NAME_CHECKED(UNeuralMorphModelVizSettings, MorphTargetDeltaThreshold)))
+			.EditCondition(TAttribute<bool>(this, &FNeuralMorphModelVizSettingsDetails::IsMorphTargetsEnabled), nullptr);
+	}
+
+}	//namespace UE::NeuralMorphModel
+
+#undef LOCTEXT_NAMESPACE
diff --git a/Engine/Plugins/Experimental/MLDeformer/NeuralMorphModel/Source/NeuralMorphModelEditor/Private/NeuralMorphTrainingModel.cpp b/Engine/Plugins/Experimental/MLDeformer/NeuralMorphModel/Source/NeuralMorphModelEditor/Private/NeuralMorphTrainingModel.cpp
new file mode 100644
--- /dev/null
+++ b/Engine/Plugins/Experimental/MLDeformer/NeuralMorphModel/Source/NeuralMorphModelEditor/Private/NeuralMorphTrainingModel.cpp
@@ -0,0 +1,22 @@
+// Copyright Epic Games, Inc. All Rights Reserved.
+
+#include "NeuralMorphTrainingModel.h"
+#include "NeuralMorphModel.h"
+#include "NeuralMorphEditorModel.h"
+#include "MLDeformerEditorModel.h"
+#include "Misc/ScopedSlowTask.h"
+#include "Math/NumericLimits.h"
+
+#define LOCTEXT_NAMESPACE "NeuralMorphTrainingModel"
+
+UNeuralMorphModel* UNeuralMorphTrainingModel::GetNeuralMorphModel() const
+{
+	return Cast<UNeuralMorphModel>(GetModel());
+}
+
+UE::NeuralMorphModel::FNeuralMorphEditorModel* UNeuralMorphTrainingModel::GetNeuralMorphEditorModel() const
+{
+	return static_cast<UE::NeuralMorphModel::FNeuralMorphEditorModel*>(EditorModel);
+}
+
+#undef LOCTEXT_NAMESPACE
diff --git a/Engine/Plugins/Experimental/MLDeformer/NeuralMorphModel/Source/NeuralMorphModelEditor/Public/NeuralMorphEditorModel.h b/Engine/Plugins/Experimental/MLDeformer/NeuralMorphModel/Source/NeuralMorphModelEditor/Public/NeuralMorphEditorModel.h
new file mode 100644
--- /dev/null
+++ b/Engine/Plugins/Experimental/MLDeformer/NeuralMorphModel/Source/NeuralMorphModelEditor/Public/NeuralMorphEditorModel.h
@@ -0,0 +1,87 @@
+// Copyright Epic Games, Inc. All Rights Reserved.
+
+#pragma once
+
+#include "CoreMinimal.h"
+#include "MLDeformerEditorModel.h"
+#include "MLDeformerGeomCacheSampler.h"
+#include "MLDeformerEditorActor.h"
+#include "NeuralMorphModel.h"
+#include "UObject/GCObject.h"
+
+class UMLDeformerAsset;
+class USkeletalMesh;
+class UGeometryCache;
+class UNeuralMorphModel;
+class UNeuralMorphModelVizSettings;
+
+namespace UE::NeuralMorphModel
+{
+	using namespace UE::MLDeformer;
+
+	class FNeuralMorphEditorModelActor;
+
+	class NEURALMORPHMODELEDITOR_API FNeuralMorphEditorModel
+		: public UE::MLDeformer::FMLDeformerEditorModel
+	{
+	public:
+		virtual ~FNeuralMorphEditorModel() override;
+
+		// We need to implement this static MakeInstance method.
+		static FMLDeformerEditorModel* MakeInstance();
+
+		// FGCObject overrides.
+		virtual void AddReferencedObjects(FReferenceCollector& Collector) override;
+		virtual FString GetReferencerName() const override { return TEXT("FNeuralMorphEditorModel"); }
+		// ~END FGCObject overrides.
+	
+		// FMLDeformerEditorModel overrides.
+		virtual void Init(const InitSettings& Settings) override;
+		virtual FMLDeformerEditorActor* CreateEditorActor(const FMLDeformerEditorActor::FConstructSettings& Settings) const override;
+		virtual FMLDeformerSampler* CreateSampler() const;
+		virtual void CreateTrainingGroundTruthActor(UWorld* World) override;
+		virtual void CreateTestGroundTruthActor(UWorld* World) override;
+		virtual void OnInputAssetsChanged() override;
+		virtual int32 GetNumTrainingFrames() const override;
+		virtual double GetTrainingTimeAtFrame(int32 FrameNumber) const override;
+		virtual int32 GetTrainingFrameAtTime(double TimeInSeconds) const override;
+		virtual double GetTestTimeAtFrame(int32 FrameNumber) const override;
+		virtual int32 GetTestFrameAtTime(double TimeInSeconds) const override;
+		virtual void UpdateIsReadyForTrainingState() override;
+		virtual ETrainingResult Train() override;
+		virtual void Render(const FSceneView* View, FViewport* Viewport, FPrimitiveDrawInterface* PDI) override;
+		virtual FString GetTrainedNetworkOnnxFile() const override;
+		virtual FString GetDefaultDeformerGraphAssetPath() const override;
+		virtual FString GetHeatMapDeformerGraphPath() const override;
+		virtual void OnPreTraining() override;
+		virtual void OnPostTraining(ETrainingResult TrainingResult, bool bUsePartiallyTrainedWhenAborted) override;
+		// ~END FMLDeformerEditorModel overrides.
+
+		// Some helpers that cast to this model's variants of some classes.
+		UNeuralMorphModel* GetNeuralMorphModel() const { return Cast<UNeuralMorphModel>(Model); }
+		UNeuralMorphModelVizSettings* GetNeuralMorphModelVizSettings() const;
+		FNeuralMorphEditorModelActor* FindNeuralMorphModelEditorActor(int32 TypeID) const;
+		FMLDeformerGeomCacheSampler* GetGeomCacheSampler() const { return static_cast<FMLDeformerGeomCacheSampler*>(Sampler); }
+	
+		void OnPostEditChangeProperty(FPropertyChangedEvent& PropertyChangedEvent);
+
+	protected:
+		void CreateGeomCacheActor(UWorld* World, int32 ActorID, const FName& Name, UGeometryCache* GeomCache, FLinearColor LabelColor, FLinearColor WireframeColor, const FText& LabelText, bool bIsTrainingActor);
+
+		/**
+		 * Initialize a set of engine morph targets and compress them to GPU friendly buffers.
+		 * These morph targets are initialized from a set of deltas. Each morph target needs to have Model->GetNumBaseVerts() number of deltas.
+		 * All deltas are concatenated in one big array. So all deltas of all vertices for the second morph target are appended to the deltas for 
+		 * the first morph target, etc. In other words, the layout is: [morph0_deltas][morph1_deltas][morph2_deltas][...].
+		 * @param Deltas The deltas for all morph targets concatenated. So the number of items in this array is a multiple of Model->GetNumBaseVerts().
+		 */
+		void InitEngineMorphTargets(const TArray<FVector3f>& Deltas);
+
+	protected:
+		/**
+		 * The entire set of morph target deltas, 3 per vertex, for each morph target, as one flattened buffer.
+		 * So the size of this buffer is: (NumVertsPerMorphTarget * NumMorphTargets).
+		 */
+		TArray<FVector3f> MorphTargetDeltasBackup;
+	};
+}	// namespace UE::NeuralMorphModel
diff --git a/Engine/Plugins/Experimental/MLDeformer/NeuralMorphModel/Source/NeuralMorphModelEditor/Public/NeuralMorphEditorModelActor.h b/Engine/Plugins/Experimental/MLDeformer/NeuralMorphModel/Source/NeuralMorphModelEditor/Public/NeuralMorphEditorModelActor.h
new file mode 100644
--- /dev/null
+++ b/Engine/Plugins/Experimental/MLDeformer/NeuralMorphModel/Source/NeuralMorphModelEditor/Public/NeuralMorphEditorModelActor.h
@@ -0,0 +1,38 @@
+// Copyright Epic Games, Inc. All Rights Reserved.
+
+#pragma once
+
+#include "MLDeformerEditorActor.h"
+
+class UMLDeformerComponent;
+class UGeometryCacheComponent;
+
+namespace UE::NeuralMorphModel
+{
+	using namespace UE::MLDeformer;
+
+	class NEURALMORPHMODELEDITOR_API FNeuralMorphEditorModelActor
+		: public FMLDeformerEditorActor
+	{
+	public:
+		FNeuralMorphEditorModelActor(const FConstructSettings& Settings);
+
+		void SetGeometryCacheComponent(UGeometryCacheComponent* Component) { GeomCacheComponent = Component; }
+		UGeometryCacheComponent* GetGeometryCacheComponent() const { return GeomCacheComponent; }
+
+		// FMLDeformerEditorActor overrides.
+		virtual void SetVisibility(bool bIsVisible) override;
+		virtual bool IsVisible() const override;
+		virtual bool HasVisualMesh() const override;
+		virtual void SetPlayPosition(float TimeInSeconds, bool bAutoPause=true) override;
+		virtual float GetPlayPosition() const override;
+		virtual void SetPlaySpeed(float PlaySpeed) override;
+		virtual void Pause(bool bPaused) override;
+		virtual FBox GetBoundingBox() const override;
+		// ~END FMLDeformerEditorActor overrides.
+
+	protected:
+		/** The geometry cache component (can be nullptr). */
+		TObjectPtr<UGeometryCacheComponent> GeomCacheComponent = nullptr;
+	};
+}	// namespace UE::NeuralMorphModel
diff --git a/Engine/Plugins/Experimental/MLDeformer/NeuralMorphModel/Source/NeuralMorphModelEditor/Public/NeuralMorphModelDetails.h b/Engine/Plugins/Experimental/MLDeformer/NeuralMorphModel/Source/NeuralMorphModelEditor/Public/NeuralMorphModelDetails.h
new file mode 100644
--- /dev/null
+++ b/Engine/Plugins/Experimental/MLDeformer/NeuralMorphModel/Source/NeuralMorphModelEditor/Public/NeuralMorphModelDetails.h
@@ -0,0 +1,45 @@
+// Copyright Epic Games, Inc. All Rights Reserved.
+
+#pragma once
+
+#include "CoreMinimal.h"
+#include "MLDeformerModelDetails.h"
+#include "IDetailCustomization.h"
+#include "AssetRegistry/AssetData.h"
+#include "Input/Reply.h"
+
+class IDetailLayoutBuilder;
+class UNeuralMorphModel;
+
+namespace UE::MLDeformer
+{
+	class FMLDeformerEditorModel;
+}
+
+namespace UE::NeuralMorphModel
+{
+	class FNeuralMorphEditorModel;
+
+	class NEURALMORPHMODELEDITOR_API FNeuralMorphModelDetails final
+		: public UE::MLDeformer::FMLDeformerModelDetails
+	{
+	public:
+		/** Makes a new instance of this detail layout class for a specific detail view requesting it. */
+		static TSharedRef<IDetailCustomization> MakeInstance();
+
+		// ILayoutDetails overrides.
+		void CustomizeDetails(IDetailLayoutBuilder& DetailBuilder) override;
+		// ~END ILayoutDetails overrides.
+
+		// FMLDeformerModelDetails overrides.
+		bool UpdateMemberPointers(const TArray<TWeakObjectPtr<UObject>>& Objects) override;
+		void AddTargetMesh() override;
+		void AddBaseMeshErrors() override;
+		void AddAnimSequenceErrors() override;
+		// ~END FMLDeformerModelDetails overrides.
+
+	protected:
+		TObjectPtr<UNeuralMorphModel> NeuralMorphModel = nullptr;
+		FNeuralMorphEditorModel* NeuralMorphEditorModel = nullptr;
+	};
+}	// namespace UE::NeuralMorphModel
diff --git a/Engine/Plugins/Experimental/MLDeformer/NeuralMorphModel/Source/NeuralMorphModelEditor/Public/NeuralMorphModelEditorModule.h b/Engine/Plugins/Experimental/MLDeformer/NeuralMorphModel/Source/NeuralMorphModelEditor/Public/NeuralMorphModelEditorModule.h
new file mode 100644
--- /dev/null
+++ b/Engine/Plugins/Experimental/MLDeformer/NeuralMorphModel/Source/NeuralMorphModelEditor/Public/NeuralMorphModelEditorModule.h
@@ -0,0 +1,19 @@
+// Copyright Epic Games, Inc. All Rights Reserved.
+
+#pragma once
+
+#include "CoreMinimal.h"
+#include "Modules/ModuleInterface.h"
+
+namespace UE::NeuralMorphModel
+{
+	class NEURALMORPHMODELEDITOR_API FNeuralMorphModelEditorModule
+		: public IModuleInterface
+	{
+	public:
+		// IModuleInterface overrides.
+		void StartupModule() override;
+		void ShutdownModule() override;
+		// ~END IModuleInterface overrides.
+	};
+}	// namespace UE::NeuralMorphModel
\ No newline at end of file
diff --git a/Engine/Plugins/Experimental/MLDeformer/NeuralMorphModel/Source/NeuralMorphModelEditor/Public/NeuralMorphModelVizSettingsDetails.h b/Engine/Plugins/Experimental/MLDeformer/NeuralMorphModel/Source/NeuralMorphModelEditor/Public/NeuralMorphModelVizSettingsDetails.h
new file mode 100644
--- /dev/null
+++ b/Engine/Plugins/Experimental/MLDeformer/NeuralMorphModel/Source/NeuralMorphModelEditor/Public/NeuralMorphModelVizSettingsDetails.h
@@ -0,0 +1,37 @@
+// Copyright Epic Games, Inc. All Rights Reserved.
+
+#pragma once
+
+#include "CoreMinimal.h"
+#include "MLDeformerVizSettingsDetails.h"
+#include "IDetailCustomization.h"
+#include "AssetRegistry/AssetData.h"
+#include "PropertyHandle.h"
+
+class IDetailLayoutBuilder;
+class USkeleton;
+class UNeuralMorphModel;
+class UNeuralMorphModelVizSettings;
+
+namespace UE::NeuralMorphModel
+{
+	class NEURALMORPHMODELEDITOR_API FNeuralMorphModelVizSettingsDetails final
+		: public UE::MLDeformer::FMLDeformerVizSettingsDetails
+	{
+	public:
+		/** Makes a new instance of this detail layout class for a specific detail view requesting it. */
+		static TSharedRef<IDetailCustomization> MakeInstance();
+
+		// FMLDeformerVizSettingsDetails overrides.
+		bool UpdateMemberPointers(const TArray<TWeakObjectPtr<UObject>>& Objects) override;
+		void AddGroundTruth() override;
+		void AddAdditionalSettings() override;
+		// ~END FMLDeformerVizSettingsDetails overrides.
+
+	protected:
+		bool IsMorphTargetsEnabled() const;
+
+		TObjectPtr<UNeuralMorphModel> NeuralMorphModel = nullptr;
+		TObjectPtr<UNeuralMorphModelVizSettings> NeuralMorphVizSettings = nullptr;
+	};
+}	// namespace UE::NeuralMorphModel
diff --git a/Engine/Plugins/Experimental/MLDeformer/NeuralMorphModel/Source/NeuralMorphModelEditor/Public/NeuralMorphTrainingModel.h b/Engine/Plugins/Experimental/MLDeformer/NeuralMorphModel/Source/NeuralMorphModelEditor/Public/NeuralMorphTrainingModel.h
new file mode 100644
--- /dev/null
+++ b/Engine/Plugins/Experimental/MLDeformer/NeuralMorphModel/Source/NeuralMorphModelEditor/Public/NeuralMorphTrainingModel.h
@@ -0,0 +1,29 @@
+// Copyright Epic Games, Inc. All Rights Reserved.
+
+#pragma once
+
+#include "CoreMinimal.h"
+#include "MLDeformerTrainingModel.h"
+#include "NeuralMorphTrainingModel.generated.h"
+
+namespace UE::NeuralMorphModel
+{
+	class FNeuralMorphEditorModel;
+}
+
+class UNeuralMorphModel;
+
+UCLASS(Blueprintable)
+class NEURALMORPHMODELEDITOR_API UNeuralMorphTrainingModel
+	: public UMLDeformerTrainingModel
+{
+	GENERATED_BODY()
+
+public:
+	/** Main training function, with implementation in python. */
+	UFUNCTION(BlueprintImplementableEvent, Category = "Training Model")
+	int32 Train() const;
+
+	UNeuralMorphModel* GetNeuralMorphModel() const;
+	UE::NeuralMorphModel::FNeuralMorphEditorModel* GetNeuralMorphEditorModel() const;
+};
diff --git a/Engine/Plugins/Experimental/MLDeformer/VertexDeltaModel/Source/VertexDeltaModelEditor/Private/VertexDeltaEditorModel.cpp b/Engine/Plugins/Experimental/MLDeformer/VertexDeltaModel/Source/VertexDeltaModelEditor/Private/VertexDeltaEditorModel.cpp
--- a/Engine/Plugins/Experimental/MLDeformer/VertexDeltaModel/Source/VertexDeltaModelEditor/Private/VertexDeltaEditorModel.cpp
+++ b/Engine/Plugins/Experimental/MLDeformer/VertexDeltaModel/Source/VertexDeltaModelEditor/Private/VertexDeltaEditorModel.cpp
@@ -73,7 +73,7 @@ namespace UE::VertexDeltaModel
 		{
 			if (PropertyChangedEvent.ChangeType == EPropertyChangeType::ValueSet)
 			{
-				bIsDataNormalized = false;
+				SetResamplingInputOutputsNeeded(true);
 				SampleDeltas();
 			}
 		}
@@ -188,27 +188,51 @@ namespace UE::VertexDeltaModel
 		return Cast<UVertexDeltaModelVizSettings>(GetVertexDeltaModel()->GetVizSettings());
 	}
 
-	double FVertexDeltaEditorModel::GetTimeAtFrame(int32 FrameNumber) const
+	double FVertexDeltaEditorModel::GetTrainingTimeAtFrame(int32 FrameNumber) const
 	{
-		const FVertexDeltaEditorModelActor* EditorActor = static_cast<FVertexDeltaEditorModelActor*>(GetTimelineEditorActor());
-		if (EditorActor && EditorActor->GetGeometryCacheComponent())
+		const FVertexDeltaEditorModelActor* EditorActor = static_cast<FVertexDeltaEditorModelActor*>(FindEditorActor(ActorID_Train_GroundTruth));
+		if (EditorActor && EditorActor->GetGeometryCacheComponent() && EditorActor->GetGeometryCacheComponent()->GeometryCache.Get())
 		{
 			return EditorActor->GetGeometryCacheComponent()->GetTimeAtFrame(FrameNumber);
 		}
-		return 0.0;
+
+		return FMLDeformerEditorModel::GetTrainingTimeAtFrame(FrameNumber);
+	}
+
+	int32 FVertexDeltaEditorModel::GetTrainingFrameAtTime(double TimeInSeconds) const
+	{
+		const FVertexDeltaEditorModelActor* EditorActor = static_cast<FVertexDeltaEditorModelActor*>(FindEditorActor(ActorID_Train_GroundTruth));
+		if (EditorActor && EditorActor->GetGeometryCacheComponent() && EditorActor->GetGeometryCacheComponent()->GeometryCache.Get())
+		{
+			return EditorActor->GetGeometryCacheComponent()->GetFrameAtTime(TimeInSeconds);
+		}
+
+		return FMLDeformerEditorModel::GetTrainingFrameAtTime(TimeInSeconds);
 	}
 
-	int32 FVertexDeltaEditorModel::GetFrameAtTime(double TimeInSeconds) const
+	double FVertexDeltaEditorModel::GetTestTimeAtFrame(int32 FrameNumber) const
 	{
-		const FVertexDeltaEditorModelActor* EditorActor = static_cast<FVertexDeltaEditorModelActor*>(GetTimelineEditorActor());
-		if (EditorActor && EditorActor->GetGeometryCacheComponent())
+		const FVertexDeltaEditorModelActor* EditorActor = static_cast<FVertexDeltaEditorModelActor*>(FindEditorActor(ActorID_Test_GroundTruth));
+		if (EditorActor && EditorActor->GetGeometryCacheComponent() && EditorActor->GetGeometryCacheComponent()->GeometryCache.Get())
+		{
+			return EditorActor->GetGeometryCacheComponent()->GetTimeAtFrame(FrameNumber);
+		}
+
+		return FMLDeformerEditorModel::GetTestTimeAtFrame(FrameNumber);
+	}
+
+	int32 FVertexDeltaEditorModel::GetTestFrameAtTime(double TimeInSeconds) const
+	{
+		const FVertexDeltaEditorModelActor* EditorActor = static_cast<FVertexDeltaEditorModelActor*>(FindEditorActor(ActorID_Test_GroundTruth));
+		if (EditorActor && EditorActor->GetGeometryCacheComponent() && EditorActor->GetGeometryCacheComponent()->GeometryCache.Get())
 		{
 			return EditorActor->GetGeometryCacheComponent()->GetFrameAtTime(TimeInSeconds);
 		}
-		return 0;
+
+		return FMLDeformerEditorModel::GetTestFrameAtTime(TimeInSeconds);
 	}
 
-	int32 FVertexDeltaEditorModel::GetNumFrames() const
+	int32 FVertexDeltaEditorModel::GetNumTrainingFrames() const
 	{
 		const UGeometryCache* GeometryCache = GetVertexDeltaModel()->GetGeometryCache();
 		if (GeometryCache == nullptr)
diff --git a/Engine/Plugins/Experimental/MLDeformer/VertexDeltaModel/Source/VertexDeltaModelEditor/Private/VertexDeltaEditorModelActor.cpp b/Engine/Plugins/Experimental/MLDeformer/VertexDeltaModel/Source/VertexDeltaModelEditor/Private/VertexDeltaEditorModelActor.cpp
--- a/Engine/Plugins/Experimental/MLDeformer/VertexDeltaModel/Source/VertexDeltaModelEditor/Private/VertexDeltaEditorModelActor.cpp
+++ b/Engine/Plugins/Experimental/MLDeformer/VertexDeltaModel/Source/VertexDeltaModelEditor/Private/VertexDeltaEditorModelActor.cpp
@@ -38,6 +38,21 @@ namespace UE::VertexDeltaModel
 		return true;
 	}
 
+	bool FVertexDeltaEditorModelActor::HasVisualMesh() const
+	{
+		if (SkeletalMeshComponent && SkeletalMeshComponent->SkeletalMesh)
+		{
+			return true;
+		}
+
+		if (GeomCacheComponent && GeomCacheComponent->GetGeometryCache())
+		{
+			return true;
+		}
+
+		return false;
+	}
+
 	void FVertexDeltaEditorModelActor::SetPlayPosition(float TimeInSeconds, bool bAutoPause)
 	{
 		if (SkeletalMeshComponent)
diff --git a/Engine/Plugins/Experimental/MLDeformer/VertexDeltaModel/Source/VertexDeltaModelEditor/Public/VertexDeltaEditorModel.h b/Engine/Plugins/Experimental/MLDeformer/VertexDeltaModel/Source/VertexDeltaModelEditor/Public/VertexDeltaEditorModel.h
--- a/Engine/Plugins/Experimental/MLDeformer/VertexDeltaModel/Source/VertexDeltaModelEditor/Public/VertexDeltaEditorModel.h
+++ b/Engine/Plugins/Experimental/MLDeformer/VertexDeltaModel/Source/VertexDeltaModelEditor/Public/VertexDeltaEditorModel.h
@@ -43,9 +43,11 @@ namespace UE::VertexDeltaModel
 		virtual void CreateTrainingGroundTruthActor(UWorld* World) override;
 		virtual void CreateTestGroundTruthActor(UWorld* World) override;
 		virtual void OnInputAssetsChanged() override;
-		virtual int32 GetNumFrames() const override;
-		virtual double GetTimeAtFrame(int32 FrameNumber) const override;
-		virtual int32 GetFrameAtTime(double TimeInSeconds) const override;
+		virtual int32 GetNumTrainingFrames() const override;
+		virtual double GetTrainingTimeAtFrame(int32 FrameNumber) const override;
+		virtual int32 GetTrainingFrameAtTime(double TimeInSeconds) const override;
+		virtual double GetTestTimeAtFrame(int32 FrameNumber) const override;
+		virtual int32 GetTestFrameAtTime(double TimeInSeconds) const override;
 		virtual void UpdateIsReadyForTrainingState() override;
 		virtual ETrainingResult Train() override;
 		FString GetTrainedNetworkOnnxFile() const override;
diff --git a/Engine/Plugins/Experimental/MLDeformer/VertexDeltaModel/Source/VertexDeltaModelEditor/Public/VertexDeltaEditorModelActor.h b/Engine/Plugins/Experimental/MLDeformer/VertexDeltaModel/Source/VertexDeltaModelEditor/Public/VertexDeltaEditorModelActor.h
--- a/Engine/Plugins/Experimental/MLDeformer/VertexDeltaModel/Source/VertexDeltaModelEditor/Public/VertexDeltaEditorModelActor.h
+++ b/Engine/Plugins/Experimental/MLDeformer/VertexDeltaModel/Source/VertexDeltaModelEditor/Public/VertexDeltaEditorModelActor.h
@@ -23,6 +23,7 @@ namespace UE::VertexDeltaModel
 		// FMLDeformerEditorActor overrides.
 		virtual void SetVisibility(bool bIsVisible) override;
 		virtual bool IsVisible() const override;
+		virtual bool HasVisualMesh() const override;
 		virtual void SetPlayPosition(float TimeInSeconds, bool bAutoPause=true) override;
 		virtual float GetPlayPosition() const override;
 		virtual void SetPlaySpeed(float PlaySpeed) override;
diff --git a/Engine/Plugins/Experimental/MLDeformer/VertexDeltaModel/VertexDeltaModel.uplugin b/Engine/Plugins/Experimental/MLDeformer/VertexDeltaModel/VertexDeltaModel.uplugin
--- a/Engine/Plugins/Experimental/MLDeformer/VertexDeltaModel/VertexDeltaModel.uplugin
+++ b/Engine/Plugins/Experimental/MLDeformer/VertexDeltaModel/VertexDeltaModel.uplugin
@@ -1,15 +1,15 @@
 {
 	"FileVersion": 3,
 	"Version": 1,
-	"VersionName": "0.2",
-	"FriendlyName": "VertexDeltaModel",
+	"VersionName": "0.21",
+	"FriendlyName": "ML Deformer Vertex Delta Model",
 	"Description": "Vertex Delta Model for the ML Deformer Framework",
 	"Category": "Animation",
 	"CreatedBy": "Epic Games, Inc.",
 	"CreatedByURL": "http://epicgames.com",
-	"DocsURL": "",
-	"MarketplaceURL": "",
-	"SupportURL": "",
+	"DocsURL": "https://docs.unrealengine.com/5.0/en-US/using-the-machine-learning-deformer-in-unreal-engine/",
+	"MarketplaceURL": "https://www.unrealengine.com/marketplace/en-US/product/ml-deformer",
+	"SupportURL": "https://forums.unrealengine.com/",
 	"CanContainContent": true,
 	"IsBetaVersion": false,
 	"IsExperimentalVersion": true,
diff --git a/Engine/Source/Runtime/Engine/Private/SkeletalRenderGPUSkin.cpp b/Engine/Source/Runtime/Engine/Private/SkeletalRenderGPUSkin.cpp
--- a/Engine/Source/Runtime/Engine/Private/SkeletalRenderGPUSkin.cpp
+++ b/Engine/Source/Runtime/Engine/Private/SkeletalRenderGPUSkin.cpp
@@ -630,7 +630,10 @@ void FSkeletalMeshObjectGPUSkin::ProcessUpdatedDynamicData(EGPUSkinCacheEntryMod
 		FGPUSkinCache::Release(SkinCacheEntry);
 	}
 
-	if (MorphVertexBuffer.bNeedsInitialClear && !(bMorph && bMorphNeedsUpdate))
+	// We need to clear the external morph buffers when the weights are zero.
+	const bool bExternalMorphsNeedClear = !DynamicData->ExternalMorphWeightData.MorphSets.IsEmpty() && !DynamicData->ExternalMorphWeightData.HasActiveMorphs();
+
+	if ((MorphVertexBuffer.bNeedsInitialClear && !(bMorph && bMorphNeedsUpdate)) || bExternalMorphsNeedClear)
 	{
 		QUICK_SCOPE_CYCLE_COUNTER(STAT_FSkeletalMeshObjectGPUSkin_ProcessUpdatedDynamicData_ClearMorphBuffer);
 		if (MorphVertexBuffer.GetUAV())
