[MLDeformer] Changes:
- Fixed a bug where duplicate label could appear when switching models.
- There used to be special handling to deal with the Persona actor, as it couldn't get deleted. This is now properly handled.
- Added a deformer graph for GPU morph based heat maps.
- Made heat maps work on the Neural Morph Model.

diff --git a/Engine/Plugins/Experimental/MLDeformer/LegacyVertexDeltaModel/Source/LegacyVertexDeltaModelEditor/Private/LegacyVertexDeltaEditorModelActor.cpp b/Engine/Plugins/Experimental/MLDeformer/LegacyVertexDeltaModel/Source/LegacyVertexDeltaModelEditor/Private/LegacyVertexDeltaEditorModelActor.cpp
--- a/Engine/Plugins/Experimental/MLDeformer/LegacyVertexDeltaModel/Source/LegacyVertexDeltaModelEditor/Private/LegacyVertexDeltaEditorModelActor.cpp
+++ b/Engine/Plugins/Experimental/MLDeformer/LegacyVertexDeltaModel/Source/LegacyVertexDeltaModelEditor/Private/LegacyVertexDeltaEditorModelActor.cpp
@@ -13,6 +13,14 @@ namespace UE::LegacyVertexDeltaModel
 	{
 	}
 
+	FLegacyVertexDeltaEditorModelActor::~FLegacyVertexDeltaEditorModelActor()
+	{
+		if (GeomCacheComponent)
+		{
+			Actor->RemoveOwnedComponent(GeomCacheComponent);
+		}
+	}
+
 	void FLegacyVertexDeltaEditorModelActor::SetVisibility(bool bIsVisible)
 	{
 		FMLDeformerEditorActor::SetVisibility(bIsVisible);
diff --git a/Engine/Plugins/Experimental/MLDeformer/LegacyVertexDeltaModel/Source/LegacyVertexDeltaModelEditor/Public/LegacyVertexDeltaEditorModelActor.h b/Engine/Plugins/Experimental/MLDeformer/LegacyVertexDeltaModel/Source/LegacyVertexDeltaModelEditor/Public/LegacyVertexDeltaEditorModelActor.h
--- a/Engine/Plugins/Experimental/MLDeformer/LegacyVertexDeltaModel/Source/LegacyVertexDeltaModelEditor/Public/LegacyVertexDeltaEditorModelActor.h
+++ b/Engine/Plugins/Experimental/MLDeformer/LegacyVertexDeltaModel/Source/LegacyVertexDeltaModelEditor/Public/LegacyVertexDeltaEditorModelActor.h
@@ -16,6 +16,7 @@ namespace UE::LegacyVertexDeltaModel
 	{
 	public:
 		FLegacyVertexDeltaEditorModelActor(const FConstructSettings& Settings);
+		virtual ~FLegacyVertexDeltaEditorModelActor() override;
 
 		void SetGeometryCacheComponent(UGeometryCacheComponent* Component) { GeomCacheComponent = Component; }
 		UGeometryCacheComponent* GetGeometryCacheComponent() const { return GeomCacheComponent; }
diff --git a/Engine/Plugins/Experimental/MLDeformer/MLDeformerFramework/Source/MLDeformerFramework/Private/MLDeformerModelInstance.cpp b/Engine/Plugins/Experimental/MLDeformer/MLDeformerFramework/Source/MLDeformerFramework/Private/MLDeformerModelInstance.cpp
--- a/Engine/Plugins/Experimental/MLDeformer/MLDeformerFramework/Source/MLDeformerFramework/Private/MLDeformerModelInstance.cpp
+++ b/Engine/Plugins/Experimental/MLDeformer/MLDeformerFramework/Source/MLDeformerFramework/Private/MLDeformerModelInstance.cpp
@@ -252,15 +252,8 @@ bool UMLDeformerModelInstance::IsValidForDataProvider() const
 			return false;
 		}
 	}
-	else
-	{
-		// We don't support CPU neural networks on default, when using the deformer graph, as we want the neural network output on the GPU.
-		return false;
-	}
 
-	return
-		Model->GetVertexMapBuffer().ShaderResourceViewRHI != nullptr &&
-		GetNeuralNetworkInferenceHandle() != -1;
+	return (Model->GetVertexMapBuffer().ShaderResourceViewRHI != nullptr) && (GetNeuralNetworkInferenceHandle() != -1);
 }
 
 void UMLDeformerModelInstance::RunNeuralNetwork(float ModelWeight)
diff --git a/Engine/Plugins/Experimental/MLDeformer/MLDeformerFramework/Source/MLDeformerFrameworkEditor/Private/MLDeformerEditorActor.cpp b/Engine/Plugins/Experimental/MLDeformer/MLDeformerFramework/Source/MLDeformerFrameworkEditor/Private/MLDeformerEditorActor.cpp
--- a/Engine/Plugins/Experimental/MLDeformer/MLDeformerFramework/Source/MLDeformerFrameworkEditor/Private/MLDeformerEditorActor.cpp
+++ b/Engine/Plugins/Experimental/MLDeformer/MLDeformerFramework/Source/MLDeformerFrameworkEditor/Private/MLDeformerEditorActor.cpp
@@ -19,6 +19,19 @@ namespace UE::MLDeformer
 		bIsTrainingActor = Settings.bIsTrainingActor;
 	}
 
+	FMLDeformerEditorActor::~FMLDeformerEditorActor()
+	{
+		if (LabelComponent)
+		{
+			Actor->RemoveOwnedComponent(LabelComponent);
+		}
+
+		if (SkeletalMeshComponent)
+		{
+			Actor->RemoveOwnedComponent(SkeletalMeshComponent);
+		}
+	}
+
 	void FMLDeformerEditorActor::SetVisibility(bool bIsVisible)
 	{
 		if (SkeletalMeshComponent && bIsVisible != SkeletalMeshComponent->IsVisible())
diff --git a/Engine/Plugins/Experimental/MLDeformer/MLDeformerFramework/Source/MLDeformerFrameworkEditor/Private/MLDeformerEditorModel.cpp b/Engine/Plugins/Experimental/MLDeformer/MLDeformerFramework/Source/MLDeformerFrameworkEditor/Private/MLDeformerEditorModel.cpp
--- a/Engine/Plugins/Experimental/MLDeformer/MLDeformerFramework/Source/MLDeformerFrameworkEditor/Private/MLDeformerEditorModel.cpp
+++ b/Engine/Plugins/Experimental/MLDeformer/MLDeformerFramework/Source/MLDeformerFrameworkEditor/Private/MLDeformerEditorModel.cpp
@@ -83,8 +83,8 @@ namespace UE::MLDeformer
 
 		// Spawn the linear skinned actor.
 		FActorSpawnParameters BaseSpawnParams;
-		BaseSpawnParams.Name = MakeUniqueObjectName(World, AAnimationEditorPreviewActor::StaticClass(), "Train Base Actor");
-		AActor* Actor = World->SpawnActor<AAnimationEditorPreviewActor>(AAnimationEditorPreviewActor::StaticClass(), FTransform::Identity, BaseSpawnParams);
+		BaseSpawnParams.Name = MakeUniqueObjectName(World, AActor::StaticClass(), "Train Base Actor");
+		AActor* Actor = World->SpawnActor<AActor>(AActor::StaticClass(), BaseSpawnParams);
 		Actor->SetFlags(RF_Transient);
 
 		// Create the preview skeletal mesh component.
@@ -94,11 +94,6 @@ namespace UE::MLDeformer
 		SkelMeshComponent->SetVisibility(false);
 		SkelMeshComponent->MarkRenderStateDirty();
 
-		// Setup an apply an anim instance to the skeletal mesh component.
-		UAnimPreviewInstance* AnimPreviewInstance = NewObject<UAnimPreviewInstance>(SkelMeshComponent, TEXT("MLDeformerAnimInstance"));
-		SkelMeshComponent->PreviewInstance = AnimPreviewInstance;
-		AnimPreviewInstance->InitializeAnimation();
-
 		// Set the skeletal mesh on the component.
 		// NOTE: This must be done AFTER setting the AnimInstance so that the correct root anim node is loaded.
 		USkeletalMesh* Mesh = Model->GetSkeletalMesh();
@@ -121,7 +116,6 @@ namespace UE::MLDeformer
 		Settings.bIsTrainingActor = true;
 		FMLDeformerEditorActor* EditorActor = CreateEditorActor(Settings);
 		EditorActor->SetSkeletalMeshComponent(SkelMeshComponent);
-		EditorActor->SetCanDestroyActor(false);	// Crash will occur when destroying the Persona actor, so disable this.
 		EditorActor->SetMeshOffsetFactor(0.0f);
 		EditorActors.Add(EditorActor);
 	}
@@ -227,28 +221,33 @@ namespace UE::MLDeformer
 	{
 		TSharedRef<IPersonaPreviewScene> PreviewScene = Editor->GetPersonaToolkit()->GetPreviewScene();
 
+		// First remove actors from the world and add them to the destroy list.
+		TArray<AActor*> ActorsToDestroy;
+		ActorsToDestroy.Reserve(EditorActors.Num());
 		UWorld* World = PreviewScene->GetWorld();
 		for (FMLDeformerEditorActor* EditorActor : EditorActors)
 		{
 			if (EditorActor)
-			{
+			{				
 				World->RemoveActor(EditorActor->GetActor(), true);
-				if (EditorActor->GetCanDestroyActor())
-				{
-					EditorActor->GetActor()->Destroy();
-				}
+				ActorsToDestroy.Add(EditorActor->GetActor());
 			}
 		}
 
+		// Now delete the editor actors, which removes components.
+		DeleteEditorActors();
+
+		PreviewScene->DeselectAll();
 		PreviewScene->SetPreviewAnimationAsset(nullptr);
 		PreviewScene->SetPreviewAnimationBlueprint(nullptr, nullptr);
 		PreviewScene->SetPreviewMesh(nullptr);
 		PreviewScene->SetPreviewMeshComponent(nullptr);
 		PreviewScene->SetActor(nullptr);
-		PreviewScene->ClearSelectedActor();
 
-		// Clear the editor actors.
-		DeleteEditorActors();
+		for (AActor* Actor : ActorsToDestroy)
+		{
+			Actor->Destroy();
+		}
 	}
 
 	FMLDeformerEditorActor* FMLDeformerEditorModel::CreateEditorActor(const FMLDeformerEditorActor::FConstructSettings& Settings) const
diff --git a/Engine/Plugins/Experimental/MLDeformer/MLDeformerFramework/Source/MLDeformerFrameworkEditor/Public/MLDeformerEditorActor.h b/Engine/Plugins/Experimental/MLDeformer/MLDeformerFramework/Source/MLDeformerFrameworkEditor/Public/MLDeformerEditorActor.h
--- a/Engine/Plugins/Experimental/MLDeformer/MLDeformerFramework/Source/MLDeformerFrameworkEditor/Public/MLDeformerEditorActor.h
+++ b/Engine/Plugins/Experimental/MLDeformer/MLDeformerFramework/Source/MLDeformerFrameworkEditor/Public/MLDeformerEditorActor.h
@@ -32,7 +32,7 @@ namespace UE::MLDeformer
 			bool bIsTrainingActor = false;
 		};
 		FMLDeformerEditorActor(const FConstructSettings& Settings);
-		virtual ~FMLDeformerEditorActor() = default;
+		virtual ~FMLDeformerEditorActor();
 
 		// Main methods you can override.
 		virtual void SetVisibility(bool bIsVisible);
@@ -56,8 +56,6 @@ namespace UE::MLDeformer
 		void SetMeshOffsetFactor(float OffsetFactor) { MeshOffsetFactor = OffsetFactor; }
 		bool IsTrainingActor() const { return bIsTrainingActor; }
 		bool IsTestActor() const { return !bIsTrainingActor; }
-		void SetCanDestroyActor(bool bCanDestroy) { bCanDestroyActor = bCanDestroy; }
-		bool GetCanDestroyActor() const { return bCanDestroyActor; }
 		float GetMeshOffsetFactor() const { return MeshOffsetFactor; }
 
 	protected:
@@ -85,11 +83,7 @@ namespace UE::MLDeformer
 		/** The position offset factor of the actor. A value of 1.0 would offset the actor with the Mesh Spacing amount, a value of 2.0 would be two times the mesh spacing offset, etc. */
 		float MeshOffsetFactor = 0.0f;
 
-		/** Can the actor we created be destroyed during cleanup? */
-		bool bCanDestroyActor = true;
-
 		/** Is this actor used for training? */
 		bool bIsTrainingActor = true;
 	};
-
 }	// namespace UE::MLDeformer
diff --git a/Engine/Plugins/Experimental/MLDeformer/NeuralMorphModel/Source/NeuralMorphModelEditor/Private/NeuralMorphEditorModel.cpp b/Engine/Plugins/Experimental/MLDeformer/NeuralMorphModel/Source/NeuralMorphModelEditor/Private/NeuralMorphEditorModel.cpp
--- a/Engine/Plugins/Experimental/MLDeformer/NeuralMorphModel/Source/NeuralMorphModelEditor/Private/NeuralMorphEditorModel.cpp
+++ b/Engine/Plugins/Experimental/MLDeformer/NeuralMorphModel/Source/NeuralMorphModelEditor/Private/NeuralMorphEditorModel.cpp
@@ -307,7 +307,7 @@ namespace UE::NeuralMorphModel
 
 	FString FNeuralMorphEditorModel::GetHeatMapDeformerGraphPath() const
 	{
-		return FString(TEXT("/Script/OptimusCore.OptimusDeformer'/Optimus/Deformers/DG_LinearBlendSkin_Morph_Cloth_RecomputeNormals.DG_LinearBlendSkin_Morph_Cloth_RecomputeNormals'"));
+		return FString(TEXT("/MLDeformerFramework/Deformers/DG_MLDeformerModel_GPUMorph_HeatMap.DG_MLDeformerModel_GPUMorph_HeatMap"));
 	}
 
 	void FNeuralMorphEditorModel::OnPreTraining()
diff --git a/Engine/Plugins/Experimental/MLDeformer/NeuralMorphModel/Source/NeuralMorphModelEditor/Private/NeuralMorphEditorModelActor.cpp b/Engine/Plugins/Experimental/MLDeformer/NeuralMorphModel/Source/NeuralMorphModelEditor/Private/NeuralMorphEditorModelActor.cpp
--- a/Engine/Plugins/Experimental/MLDeformer/NeuralMorphModel/Source/NeuralMorphModelEditor/Private/NeuralMorphEditorModelActor.cpp
+++ b/Engine/Plugins/Experimental/MLDeformer/NeuralMorphModel/Source/NeuralMorphModelEditor/Private/NeuralMorphEditorModelActor.cpp
@@ -13,6 +13,14 @@ namespace UE::NeuralMorphModel
 	{
 	}
 
+	FNeuralMorphEditorModelActor::~FNeuralMorphEditorModelActor()
+	{
+		if (GeomCacheComponent)
+		{
+			Actor->RemoveOwnedComponent(GeomCacheComponent);
+		}
+	}
+
 	void FNeuralMorphEditorModelActor::SetVisibility(bool bIsVisible)
 	{
 		FMLDeformerEditorActor::SetVisibility(bIsVisible);
diff --git a/Engine/Plugins/Experimental/MLDeformer/NeuralMorphModel/Source/NeuralMorphModelEditor/Public/NeuralMorphEditorModelActor.h b/Engine/Plugins/Experimental/MLDeformer/NeuralMorphModel/Source/NeuralMorphModelEditor/Public/NeuralMorphEditorModelActor.h
--- a/Engine/Plugins/Experimental/MLDeformer/NeuralMorphModel/Source/NeuralMorphModelEditor/Public/NeuralMorphEditorModelActor.h
+++ b/Engine/Plugins/Experimental/MLDeformer/NeuralMorphModel/Source/NeuralMorphModelEditor/Public/NeuralMorphEditorModelActor.h
@@ -16,6 +16,7 @@ namespace UE::NeuralMorphModel
 	{
 	public:
 		FNeuralMorphEditorModelActor(const FConstructSettings& Settings);
+		virtual ~FNeuralMorphEditorModelActor() override;
 
 		void SetGeometryCacheComponent(UGeometryCacheComponent* Component) { GeomCacheComponent = Component; }
 		UGeometryCacheComponent* GetGeometryCacheComponent() const { return GeomCacheComponent; }
diff --git a/Engine/Plugins/Experimental/MLDeformer/VertexDeltaModel/Source/VertexDeltaModelEditor/Private/VertexDeltaEditorModelActor.cpp b/Engine/Plugins/Experimental/MLDeformer/VertexDeltaModel/Source/VertexDeltaModelEditor/Private/VertexDeltaEditorModelActor.cpp
--- a/Engine/Plugins/Experimental/MLDeformer/VertexDeltaModel/Source/VertexDeltaModelEditor/Private/VertexDeltaEditorModelActor.cpp
+++ b/Engine/Plugins/Experimental/MLDeformer/VertexDeltaModel/Source/VertexDeltaModelEditor/Private/VertexDeltaEditorModelActor.cpp
@@ -13,6 +13,14 @@ namespace UE::VertexDeltaModel
 	{
 	}
 
+	FVertexDeltaEditorModelActor::~FVertexDeltaEditorModelActor()
+	{
+		if (GeomCacheComponent)
+		{
+			Actor->RemoveOwnedComponent(GeomCacheComponent);
+		}
+	}
+
 	void FVertexDeltaEditorModelActor::SetVisibility(bool bIsVisible)
 	{
 		FMLDeformerEditorActor::SetVisibility(bIsVisible);
diff --git a/Engine/Plugins/Experimental/MLDeformer/VertexDeltaModel/Source/VertexDeltaModelEditor/Public/VertexDeltaEditorModelActor.h b/Engine/Plugins/Experimental/MLDeformer/VertexDeltaModel/Source/VertexDeltaModelEditor/Public/VertexDeltaEditorModelActor.h
--- a/Engine/Plugins/Experimental/MLDeformer/VertexDeltaModel/Source/VertexDeltaModelEditor/Public/VertexDeltaEditorModelActor.h
+++ b/Engine/Plugins/Experimental/MLDeformer/VertexDeltaModel/Source/VertexDeltaModelEditor/Public/VertexDeltaEditorModelActor.h
@@ -16,6 +16,7 @@ namespace UE::VertexDeltaModel
 	{
 	public:
 		FVertexDeltaEditorModelActor(const FConstructSettings& Settings);
+		virtual ~FVertexDeltaEditorModelActor() override;
 
 		void SetGeometryCacheComponent(UGeometryCacheComponent* Component) { GeomCacheComponent = Component; }
 		UGeometryCacheComponent* GetGeometryCacheComponent() const { return GeomCacheComponent; }
