[MLDeformer]
- Introduced a concept of model priorities. The one model with the highest priority will be used as default when creating a new ML Deformer Asset.
- The model list on top of the UI is now sorted by name.
- Documented the FMLDeformerModelRegistry class API.
- Debug drawing of morph targets now renders the unskinned vertices.
- Some small refactors.

diff --git a/Engine/Plugins/Experimental/MLDeformer/LegacyVertexDeltaModel/Source/LegacyVertexDeltaModelEditor/Private/LegacyVertexDeltaModelEditorModule.cpp b/Engine/Plugins/Experimental/MLDeformer/LegacyVertexDeltaModel/Source/LegacyVertexDeltaModelEditor/Private/LegacyVertexDeltaModelEditorModule.cpp
--- a/Engine/Plugins/Experimental/MLDeformer/LegacyVertexDeltaModel/Source/LegacyVertexDeltaModelEditor/Private/LegacyVertexDeltaModelEditorModule.cpp
+++ b/Engine/Plugins/Experimental/MLDeformer/LegacyVertexDeltaModel/Source/LegacyVertexDeltaModelEditor/Private/LegacyVertexDeltaModelEditorModule.cpp
@@ -28,7 +28,7 @@ namespace UE::LegacyVertexDeltaModel
 		// Register our custom ML deformer model to the model registry in the ML Deformer Framework.
 		FMLDeformerEditorModule& EditorModule = FModuleManager::LoadModuleChecked<FMLDeformerEditorModule>("MLDeformerFrameworkEditor");
 		FMLDeformerEditorModelRegistry& ModelRegistry = EditorModule.GetModelRegistry();
-		ModelRegistry.RegisterEditorModel(ULegacyVertexDeltaModel::StaticClass(), FOnGetEditorModelInstance::CreateStatic(&FLegacyVertexDeltaEditorModel::MakeInstance));
+		ModelRegistry.RegisterEditorModel(ULegacyVertexDeltaModel::StaticClass(), FOnGetEditorModelInstance::CreateStatic(&FLegacyVertexDeltaEditorModel::MakeInstance), /*ModelPriority*/0);
 	}
 
 	void FLegacyVertexDeltaModelEditorModule::ShutdownModule()
diff --git a/Engine/Plugins/Experimental/MLDeformer/MLDeformerFramework/Source/MLDeformerFrameworkEditor/Private/MLDeformerEditorModel.cpp b/Engine/Plugins/Experimental/MLDeformer/MLDeformerFramework/Source/MLDeformerFrameworkEditor/Private/MLDeformerEditorModel.cpp
--- a/Engine/Plugins/Experimental/MLDeformer/MLDeformerFramework/Source/MLDeformerFrameworkEditor/Private/MLDeformerEditorModel.cpp
+++ b/Engine/Plugins/Experimental/MLDeformer/MLDeformerFramework/Source/MLDeformerFrameworkEditor/Private/MLDeformerEditorModel.cpp
@@ -351,7 +351,7 @@ namespace UE::MLDeformer
 
 				LabelComponent->SetRelativeLocation(ActorLocation + FVector(0.0f, 0.0f, VizSettings->GetLabelHeight()) - AlignmentOffset);
 				LabelComponent->SetRelativeRotation(FQuat(FVector(0.0f, 0.0f, 1.0f), FMath::DegreesToRadians(90.0f)));
-				LabelComponent->SetRelativeScale3D(FVector(VizSettings->GetLabelScale() * 0.5f));
+				LabelComponent->SetRelativeScale3D(FVector(VizSettings->GetLabelScale() * 0.3f));
 
 				// Update visibility.
 				const bool bLabelIsVisible = (bDrawTrainingActors && EditorActor->IsTrainingActor()) || (bDrawTestActors && EditorActor->IsTestActor());
@@ -1612,6 +1612,41 @@ namespace UE::MLDeformer
 		);
 	}
 
+	void FMLDeformerEditorModel::DrawMorphTarget(FPrimitiveDrawInterface* PDI, const TArray<FVector3f>& MorphDeltas, float DeltaThreshold, int32 MorphTargetIndex, const FVector& DrawOffset)
+	{
+		UMLDeformerVizSettings* VizSettings = Model->GetVizSettings();
+		if (!MorphDeltas.IsEmpty())
+		{
+			const int32 NumVerts = Model->GetNumBaseMeshVerts();
+			check(MorphDeltas.Num() % NumVerts == 0);
+
+			const TArray<FVector3f>& SkinnedPositions = Sampler->GetUnskinnedVertexPositions();
+			check(NumVerts == SkinnedPositions.Num());
+
+			// Draw all deltas.			
+			const int32 NumMorphTargets = MorphDeltas.Num() / NumVerts;
+			const int32 FinalMorphTargetIndex = FMath::Clamp<int32>(MorphTargetIndex, 0, NumMorphTargets - 1);
+			const FLinearColor IncludedColor(1.0f, 0.0f, 1.0f);
+			const FLinearColor ExcludedColor(0.1f, 0.1f, 0.1f);
+			for (int32 VertexIndex = 0; VertexIndex < NumVerts; ++VertexIndex)
+			{
+				const FVector StartPoint = FVector(SkinnedPositions[VertexIndex].X, SkinnedPositions[VertexIndex].Y, SkinnedPositions[VertexIndex].Z) + DrawOffset;
+				const int32 DeltaArrayOffset = NumVerts * FinalMorphTargetIndex + VertexIndex;
+				const FVector Delta(MorphDeltas[DeltaArrayOffset]);
+
+				if (Delta.Length() >= DeltaThreshold)
+				{
+					PDI->DrawPoint(StartPoint, IncludedColor, 1.0f, 0);
+					PDI->DrawLine(StartPoint, StartPoint + Delta, IncludedColor, 0);
+				}
+				else
+				{
+					PDI->DrawPoint(StartPoint + Delta, ExcludedColor, 0.75f, 0);
+				}
+			}
+		}
+	}
+
 	FString FMLDeformerEditorModel::GetHeatMapMaterialPath() const
 	{
 		return FString(TEXT("/MLDeformerFramework/Materials/MLDeformerHeatMapMat.MLDeformerHeatMapMat"));
diff --git a/Engine/Plugins/Experimental/MLDeformer/MLDeformerFramework/Source/MLDeformerFrameworkEditor/Private/MLDeformerEditorToolkit.cpp b/Engine/Plugins/Experimental/MLDeformer/MLDeformerFramework/Source/MLDeformerFrameworkEditor/Private/MLDeformerEditorToolkit.cpp
--- a/Engine/Plugins/Experimental/MLDeformer/MLDeformerFramework/Source/MLDeformerFrameworkEditor/Private/MLDeformerEditorToolkit.cpp
+++ b/Engine/Plugins/Experimental/MLDeformer/MLDeformerFramework/Source/MLDeformerFrameworkEditor/Private/MLDeformerEditorToolkit.cpp
@@ -52,8 +52,9 @@ namespace UE::MLDeformer
 		if (DeformerModel == nullptr)
 		{
 			if (ModelRegistry.GetNumRegisteredModels() > 0)
-			{	
-				OnModelChanged(0);
+			{
+				const int32 HighestPriorityIndex = ModelRegistry.GetHighestPriorityModelIndex();
+				OnModelChanged(HighestPriorityIndex);
 				DeformerModel = DeformerAsset->GetModel();
 			}
 		}
@@ -188,18 +189,41 @@ namespace UE::MLDeformer
 		FMenuBuilder MenuBuilder(bShouldCloseWindowAfterMenuSelection, InCommandList);
 
 		FMLDeformerEditorModule& EditorModule = FModuleManager::GetModuleChecked<FMLDeformerEditorModule>("MLDeformerFrameworkEditor");
-		const FMLDeformerEditorModelRegistry& ModelRegistry = EditorModule.GetModelRegistry();	
+		const FMLDeformerEditorModelRegistry& ModelRegistry = EditorModule.GetModelRegistry();
 
-		// Iterate over all registered models and add them to the list.
 		TArray<UClass*> ModelTypes;
 		ModelRegistry.GetRegisteredModels().GenerateKeyArray(ModelTypes);
-		for (int ModelIndex = 0; ModelIndex < ModelTypes.Num(); ++ModelIndex)
+
+		// Generate a list of model names and sort them.
+		TArray<FString> ModelNames;
+		ModelNames.Reserve(ModelTypes.Num());
+		for (int32 ModelIndex = 0; ModelIndex < ModelTypes.Num(); ++ModelIndex)
 		{
 			const UClass* Model = ModelTypes[ModelIndex];
 			const UMLDeformerModel* DefaultModel = Cast<UMLDeformerModel>(Model->GetDefaultObject());
 			check(DefaultModel);
-			FUIAction ItemAction(FExecuteAction::CreateSP(this, &FMLDeformerEditorToolkit::OnModelChanged, ModelIndex));
-			MenuBuilder.AddMenuEntry(FText::FromString(DefaultModel->GetDisplayName()), TAttribute<FText>(), FSlateIcon(), ItemAction);
+			ModelNames.Add(DefaultModel->GetDisplayName());
+		}
+		ModelNames.Sort();
+
+		// Iterate over all registered models and add them to the list.
+		for (int32 ModelIndex = 0; ModelIndex < ModelNames.Num(); ++ModelIndex)
+		{
+			// Find the model index as in the registry, so before sorting.
+			int32 UnsortedModelIndex = -1;
+			for (int32 Index = 0; Index < ModelTypes.Num(); ++Index)
+			{
+				const UClass* Model = ModelTypes[ModelIndex];
+				const UMLDeformerModel* DefaultModel = Cast<UMLDeformerModel>(Model->GetDefaultObject());
+				if (DefaultModel->GetDisplayName() == ModelNames[Index])
+				{
+					UnsortedModelIndex = Index;
+					break;
+				}
+			}
+
+			FUIAction ItemAction(FExecuteAction::CreateSP(this, &FMLDeformerEditorToolkit::OnModelChanged, UnsortedModelIndex));
+			MenuBuilder.AddMenuEntry(FText::FromString(ModelNames[ModelIndex]), TAttribute<FText>(), FSlateIcon(), ItemAction);
 		}
 
 		return MenuBuilder.MakeWidget();
diff --git a/Engine/Plugins/Experimental/MLDeformer/MLDeformerFramework/Source/MLDeformerFrameworkEditor/Private/MLDeformerModelRegistry.cpp b/Engine/Plugins/Experimental/MLDeformer/MLDeformerFramework/Source/MLDeformerFrameworkEditor/Private/MLDeformerModelRegistry.cpp
--- a/Engine/Plugins/Experimental/MLDeformer/MLDeformerFramework/Source/MLDeformerFrameworkEditor/Private/MLDeformerModelRegistry.cpp
+++ b/Engine/Plugins/Experimental/MLDeformer/MLDeformerFramework/Source/MLDeformerFrameworkEditor/Private/MLDeformerModelRegistry.cpp
@@ -12,9 +12,11 @@ namespace UE::MLDeformer
 		InstanceMap.Empty();
 	}
 
-	void FMLDeformerEditorModelRegistry::RegisterEditorModel(UClass* ModelType, FOnGetEditorModelInstance Delegate)
+	void FMLDeformerEditorModelRegistry::RegisterEditorModel(UClass* ModelType, FOnGetEditorModelInstance Delegate, int32 ModelPriority)
 	{
 		Map.FindOrAdd(ModelType, Delegate);
+		ModelPriorities.FindOrAdd(ModelType, ModelPriority);
+		UpdateHighestPriorityModel();
 	}
 
 	void FMLDeformerEditorModelRegistry::UnregisterEditorModel(const UClass* ModelType)
@@ -23,6 +25,78 @@ namespace UE::MLDeformer
 		{
 			Map.Remove(ModelType);
 		}
+
+		if (ModelPriorities.Contains(ModelType))
+		{
+			ModelPriorities.Remove(ModelType);
+			UpdateHighestPriorityModel();
+		}
+	}
+
+	void FMLDeformerEditorModelRegistry::UpdateHighestPriorityModel()
+	{
+		int32 HighestPriority = MIN_int32;
+		for (auto& Item : ModelPriorities)
+		{
+			if (Item.Value >= HighestPriority)
+			{
+				HighestPriority = Item.Value;
+				HighestPriorityModelType = Item.Key;
+			}
+		}
+	}
+
+	int32 FMLDeformerEditorModelRegistry::GetPriorityForModel(UClass* ModelType) const
+	{
+		const int32* Priority = ModelPriorities.Find(ModelType);
+		if (Priority)
+		{
+			return *Priority;
+		}
+
+		return MIN_int32;
+	}
+
+	int32 FMLDeformerEditorModelRegistry::GetHighestPriorityModelIndex() const
+	{
+		if (HighestPriorityModelType)
+		{
+			int32 Index = 0;
+			for (auto& Item : Map)
+			{
+				if (Item.Key == HighestPriorityModelType)
+				{
+					return Index;
+				}
+				Index++;
+			}
+		}
+		return -1;
+	}
+
+	int32 FMLDeformerEditorModelRegistry::GetNumRegisteredModels() const
+	{ 
+		return Map.Num();
+	}
+
+	int32 FMLDeformerEditorModelRegistry::GetNumInstancedModels() const
+	{ 
+		return InstanceMap.Num();
+	}
+
+	UClass* FMLDeformerEditorModelRegistry::GetHighestPriorityModel() const
+	{ 
+		return HighestPriorityModelType;
+	}
+
+	const TMap<UClass*, FOnGetEditorModelInstance>& FMLDeformerEditorModelRegistry::GetRegisteredModels() const
+	{ 
+		return Map;
+	}
+
+	const TMap<UMLDeformerModel*, FMLDeformerEditorModel*>& FMLDeformerEditorModelRegistry::GetModelInstances() const
+	{ 
+		return InstanceMap;
 	}
 
 	FMLDeformerEditorModel* FMLDeformerEditorModelRegistry::CreateEditorModel(UMLDeformerModel* Model)
diff --git a/Engine/Plugins/Experimental/MLDeformer/MLDeformerFramework/Source/MLDeformerFrameworkEditor/Private/MLDeformerSampler.cpp b/Engine/Plugins/Experimental/MLDeformer/MLDeformerFramework/Source/MLDeformerFrameworkEditor/Private/MLDeformerSampler.cpp
--- a/Engine/Plugins/Experimental/MLDeformer/MLDeformerFramework/Source/MLDeformerFrameworkEditor/Private/MLDeformerSampler.cpp
+++ b/Engine/Plugins/Experimental/MLDeformer/MLDeformerFramework/Source/MLDeformerFrameworkEditor/Private/MLDeformerSampler.cpp
@@ -55,6 +55,9 @@ namespace UE::MLDeformer
 		BoneMatrices.Reset();
 		BoneRotations.Reset();
 		CurveValues.Reset();
+
+		const int32 LODIndex = 0;
+		ExtractUnskinnedPositions(LODIndex, UnskinnedVertexPositions);
 	}
 
 	void FMLDeformerSampler::UpdateSkeletalMeshComponent()
@@ -146,6 +149,40 @@ namespace UE::MLDeformer
 		return InvSkinningTransform.Inverse();
 	}
 
+	void FMLDeformerSampler::ExtractUnskinnedPositions(int32 LODIndex, TArray<FVector3f>& OutPositions) const
+	{
+		OutPositions.Reset();
+
+		if (SkeletalMeshComponent == nullptr)
+		{
+			return;
+		}
+
+		USkeletalMesh* Mesh = SkeletalMeshComponent->GetSkeletalMesh();
+		if (Mesh == nullptr)
+		{
+			return;
+		}
+
+		FSkeletalMeshLODRenderData& SkelMeshLODData = Mesh->GetResourceForRendering()->LODRenderData[LODIndex];
+		const FPositionVertexBuffer& RenderPositions = SkelMeshLODData.StaticVertexBuffers.PositionVertexBuffer;
+		const int32 NumRenderVerts = RenderPositions.GetNumVertices();
+		const FSkeletalMeshModel* SkeletalMeshModel = Mesh->GetImportedModel();
+		const TArray<int32>& ImportedVertexNumbers = SkeletalMeshModel->LODModels[LODIndex].MeshToImportVertexMap;
+
+		// Get the originally imported vertex numbers from the DCC.
+		if (ImportedVertexNumbers.Num() > 0)
+		{
+			// Store the vertex positions for the original imported vertices (8 vertices for a cube).
+			OutPositions.AddZeroed(NumImportedVertices);
+			for (int32 Index = 0; Index < NumRenderVerts; ++Index)
+			{
+				const int32 ImportedVertex = ImportedVertexNumbers[Index];
+				OutPositions[ImportedVertex] = RenderPositions.VertexPosition(Index);
+			}
+		}
+	}
+
 	void FMLDeformerSampler::ExtractSkinnedPositions(int32 LODIndex, TArray<FMatrix44f>& InBoneMatrices, TArray<FVector3f>& TempPositions, TArray<FVector3f>& OutPositions) const
 	{
 		OutPositions.Reset();
diff --git a/Engine/Plugins/Experimental/MLDeformer/MLDeformerFramework/Source/MLDeformerFrameworkEditor/Public/MLDeformerEditorModel.h b/Engine/Plugins/Experimental/MLDeformer/MLDeformerFramework/Source/MLDeformerFrameworkEditor/Public/MLDeformerEditorModel.h
--- a/Engine/Plugins/Experimental/MLDeformer/MLDeformerFramework/Source/MLDeformerFrameworkEditor/Public/MLDeformerEditorModel.h
+++ b/Engine/Plugins/Experimental/MLDeformer/MLDeformerFramework/Source/MLDeformerFrameworkEditor/Public/MLDeformerEditorModel.h
@@ -162,6 +162,18 @@ namespace UE::MLDeformer
 		int32 GetCurrentTrainingFrame() const { return CurrentTrainingFrame; }
 		void CheckTrainingDataFrameChanged();
 
+		/**
+		 * Debug draw specific morph targets using lines and points.
+		 * This can show the user what deltas are included in which morph target.
+		 * @param PDI A pointer to the draw interface.
+		 * @param MorphDeltas A buffer of deltas for ALL morph targets. The size of the buffer must be a multiple of Model->GetBaseNumVerts().
+		 *        So the layout of this buffer is [Morph0_Deltas][Morph1_Deltas][Morph2_Deltas] etc.
+		 * @param DeltaThreshold Deltas with a length  larger or equal to the given threshold value will be colored differently than the ones smaller than this threshold.
+		 * @param MorphTargetIndex The morph target number to visualize.
+		 * @param DrawOffset An offset to perform the debug draw at.
+		 */
+		void DrawMorphTarget(FPrimitiveDrawInterface* PDI, const TArray<FVector3f>& MorphDeltas, float DeltaThreshold, int32 MorphTargetIndex, const FVector& DrawOffset);
+
 	protected:
 		void DeleteEditorActors();
 		bool IsEditorReadyForTrainingBasicChecks();
diff --git a/Engine/Plugins/Experimental/MLDeformer/MLDeformerFramework/Source/MLDeformerFrameworkEditor/Public/MLDeformerModelRegistry.h b/Engine/Plugins/Experimental/MLDeformer/MLDeformerFramework/Source/MLDeformerFrameworkEditor/Public/MLDeformerModelRegistry.h
--- a/Engine/Plugins/Experimental/MLDeformer/MLDeformerFramework/Source/MLDeformerFrameworkEditor/Public/MLDeformerModelRegistry.h
+++ b/Engine/Plugins/Experimental/MLDeformer/MLDeformerFramework/Source/MLDeformerFrameworkEditor/Public/MLDeformerModelRegistry.h
@@ -12,25 +12,121 @@ namespace UE::MLDeformer
 
 	DECLARE_DELEGATE_RetVal(FMLDeformerEditorModel*, FOnGetEditorModelInstance);
 
+	/**
+	 * The editor model registry, which keeps track of different types of runtime ML Deformer models and their related editor models.
+	 * When you create a new ML Deformer model type, you register it to this registry using the RegisterEditorModel method.
+	 * You typically do this in your editor module's StartupModule method. And you unregister it again in the editor module's ShutdownModule method.
+	 * Models can have priorities. The one model with the highest priority level will be the model that is created by default when creating a new ML Deformer asset.
+	 */
 	class MLDEFORMERFRAMEWORKEDITOR_API FMLDeformerEditorModelRegistry
 	{
+		friend class FMLDeformerEditorToolkit;
+		friend class FMLDeformerEditorModel;
+
 	public:
 		~FMLDeformerEditorModelRegistry();
 
-		void RegisterEditorModel(UClass* ModelType, FOnGetEditorModelInstance Delegate);
+		/**
+		 * Register a new model. This is needed to make the model appear inside the UI of the editor as selectable model.
+		 * Call this method in your editor module's StartupModule call.
+		 * @param ModelType The type of the runtime model, for example UNeuralMorphModel::StaticClass().
+		 * @param Delegate The function that creates the editor model that's related to this editor model, 
+		 *        for example FOnGetEditorModelInstance::CreateStatic(&FNeuralMorphEditorModel::MakeInstance). This MakeInstance function must return the new editor model associated with this model.
+		 * @param ModelPriority The priority level of this model. The model with the highest priority is created by default when creating a new ML Deformer Asset.
+		 *        If there is a model with the same priority level already, and it happens to be the highest priority model, then this new model will take over as priority model.
+		 */
+		void RegisterEditorModel(UClass* ModelType, FOnGetEditorModelInstance Delegate, int32 ModelPriority = 0);
+
+		/**
+		 * Unregister the editor model.
+		 * Call this method in your editor module's ShutdownModule call.
+		 * @param ModelType The model type to unregister, for example UNeuralMorphModel::StaticClass().
+		 */
 		void UnregisterEditorModel(const UClass* ModelType);
-		void RemoveEditorModelInstance(FMLDeformerEditorModel* EditorModel);
 
-		FMLDeformerEditorModel* CreateEditorModel(UMLDeformerModel* Model);
+		/**
+		 * Get the editor model related to a given runtime model. 
+		 * @param Model The runtime model to get the editor model object for.
+		 * @return A pointer to the editor model object for this specific runtime model instance, or nullptr when not found.
+		 */
 		FMLDeformerEditorModel* GetEditorModel(UMLDeformerModel* Model);
 
-		int GetNumRegisteredModels() const { return Map.Num(); }
-		int GetNumInstancedModels() const { return InstanceMap.Num(); }
-		const TMap<UClass*, FOnGetEditorModelInstance>& GetRegisteredModels() const { return Map; }
-		const TMap<UMLDeformerModel*, FMLDeformerEditorModel*>& GetModelInstances() const { return InstanceMap; }
+		/**
+		 * Get the number of registered models.
+		 * @return The number of models that have been registered through RegisterEditorModel.
+		 */
+		int32 GetNumRegisteredModels() const;
+
+		/**
+		 * Get the number of instanced models.
+		 * @return The number of instanced models.
+		 */
+		int32 GetNumInstancedModels() const;
+
+		/**
+		 * Get the highest priority runtime model class.
+		 * The model with the highest priority is created by default when creating a new ML Deformer Asset.
+		 * @return Returns the model class of the highest priority model, or nullptr when there are none.
+		 */
+		UClass* GetHighestPriorityModel() const;
+
+		/**
+		 * Get the priority level for a given runtime model type.
+		 * @param ModelType The runtime model type, for example UNeuralMorphModel::StaticClass().
+		 * @return The model priority level, or if it cannot be found it will return MIN_int32.
+		 */
+		int32 GetPriorityForModel(UClass* ModelType) const;
 
 	private:
+		/**
+		 * Get the read-only map of registered models.
+		 * @return The map of registered models that have been registered through RegisterEditorModel.
+		 */
+		const TMap<UClass*, FOnGetEditorModelInstance>& GetRegisteredModels() const;
+
+		/**
+		 * Get the read-only map that maps instanced models with their editor models.
+		 * @return The map that provides the editor model for a specific runtime model.
+		 */
+		const TMap<UMLDeformerModel*, FMLDeformerEditorModel*>& GetModelInstances() const;
+
+		/**
+		 * Get the highest priority model index into the map.
+		 * It is the 'index'-th item in the array that you can iterate by doing something like: for (auto& Item : Map).
+		 * @return The highest priority model index.
+		 */
+		int32 GetHighestPriorityModelIndex() const;
+
+		/**
+		 * Update the highest priority model by looking at the priority map.
+		 * This updates the value returned by GetHighestPriorityModel().
+		 */
+		void UpdateHighestPriorityModel();
+
+		/**
+		 * Remove a given editor model instance.
+		 * @param EditorModel A pointer to the editor model to remove.
+		 */
+		void RemoveEditorModelInstance(FMLDeformerEditorModel* EditorModel);
+
+		/**
+		 * Create a new editor model for a given runtime model.
+		 * @param Model The runtime model object to create an editor model for.
+		 * @return A pointer to the newly created editor model.
+		 */
+		FMLDeformerEditorModel* CreateEditorModel(UMLDeformerModel* Model);
+
+	private:
+		/** A map that maps runtime model types with a function that creates its editor model. */
 		TMap<UClass*, FOnGetEditorModelInstance> Map;
+
+		/** A map of runtime model instances, and the related editor models. */
 		TMap<UMLDeformerModel*, FMLDeformerEditorModel*> InstanceMap;
+
+		/** The map of model type priority levels. */
+		TMap<UClass*, int32> ModelPriorities;
+
+		/** The registered model type that has the highest priority level. */
+		UClass* HighestPriorityModelType = nullptr;
 	};
 }	// namespace UE::MLDeformer
diff --git a/Engine/Plugins/Experimental/MLDeformer/MLDeformerFramework/Source/MLDeformerFrameworkEditor/Public/MLDeformerSampler.h b/Engine/Plugins/Experimental/MLDeformer/MLDeformerFramework/Source/MLDeformerFrameworkEditor/Public/MLDeformerSampler.h
--- a/Engine/Plugins/Experimental/MLDeformer/MLDeformerFramework/Source/MLDeformerFrameworkEditor/Public/MLDeformerSampler.h
+++ b/Engine/Plugins/Experimental/MLDeformer/MLDeformerFramework/Source/MLDeformerFrameworkEditor/Public/MLDeformerSampler.h
@@ -43,6 +43,7 @@ namespace UE::MLDeformer
 		virtual float GetTimeAtFrame(int32 InAnimFrameIndex) const { return 0.0f; }
 
 		const TArray<FVector3f>& GetSkinnedVertexPositions() const { return SkinnedVertexPositions; }
+		const TArray<FVector3f>& GetUnskinnedVertexPositions() const { return UnskinnedVertexPositions; }
 		const TArray<float>& GetVertexDeltas() const { return VertexDeltas; }
 		const TArray<float>& GetBoneRotations() const { return BoneRotations; }
 		const TArray<float>& GetCurveValues() const { return CurveValues; }
@@ -57,6 +58,7 @@ namespace UE::MLDeformer
 		virtual void CreateActors();
 		virtual void RegisterTargetComponents() {}
 		void ExtractSkinnedPositions(int32 LODIndex, TArray<FMatrix44f>& InBoneMatrices, TArray<FVector3f>& TempPositions, TArray<FVector3f>& OutPositions) const;
+		void ExtractUnskinnedPositions(int32 LODIndex, TArray<FVector3f>& OutPositions) const;
 		FMatrix44f CalcInverseSkinningTransform(int32 VertexIndex, const FSkeletalMeshLODRenderData& SkelMeshLODData, const FSkinWeightVertexBuffer& SkinWeightBuffer) const;
 		AActor* CreateNewActor(UWorld* InWorld, const FName& Name) const;
 
@@ -84,6 +86,9 @@ namespace UE::MLDeformer
 		/** The skinned vertex positions. */
 		TArray<FVector3f> SkinnedVertexPositions;
 
+		/** The unskinned vertex positions. */
+		TArray<FVector3f> UnskinnedVertexPositions;
+
 		/** A temp array to store vertex positions in. */
 		TArray<FVector3f> TempVertexPositions;
 
diff --git a/Engine/Plugins/Experimental/MLDeformer/NeuralMorphModel/Source/NeuralMorphModelEditor/Private/NeuralMorphEditorModel.cpp b/Engine/Plugins/Experimental/MLDeformer/NeuralMorphModel/Source/NeuralMorphModelEditor/Private/NeuralMorphEditorModel.cpp
--- a/Engine/Plugins/Experimental/MLDeformer/NeuralMorphModel/Source/NeuralMorphModelEditor/Private/NeuralMorphEditorModel.cpp
+++ b/Engine/Plugins/Experimental/MLDeformer/NeuralMorphModel/Source/NeuralMorphModelEditor/Private/NeuralMorphEditorModel.cpp
@@ -364,45 +364,12 @@ namespace UE::NeuralMorphModel
 	{
 		FMLDeformerEditorModel::Render(View, Viewport, PDI);
 
+		// Debug draw the selected morph target.
 		UNeuralMorphModelVizSettings* VizSettings = GetNeuralMorphModelVizSettings();
 		if (VizSettings->bDrawMorphTargets)
 		{
-			// Check if the vertex count matches.
-			UNeuralMorphModel* NeuralModel = GetNeuralMorphModel();
-			if (!NeuralModel->MorphTargetDeltas.IsEmpty())
-			{
-				const int32 NumVerts = NeuralModel->GetNumBaseMeshVerts();
-				check(NeuralModel->MorphTargetDeltas.Num() % NumVerts == 0);		
-
-				const TArray<FVector3f>& SkinnedPositions = Sampler->GetSkinnedVertexPositions();
-				check(NumVerts == SkinnedPositions.Num());
-
-				const int32 NumMorphTargets = NeuralModel->MorphTargetDeltas.Num() / NumVerts;
-
-				// Draw all deltas.
-				const FVector MeshOffset = -VizSettings->GetMeshSpacingOffsetVector();
-				const int32 BlendShapeIndex = FMath::Clamp<int32>(VizSettings->MorphTargetNumber, 0, NumMorphTargets - 1);
-				const int32 BlendShapeDeltaStartOffset = NeuralModel->GetMorphTargetDeltaStartIndex(BlendShapeIndex);
-				const FLinearColor IncludedColor(1.0f, 0.0f, 1.0f);
-				const FLinearColor ExcludedColor(0.2f, 0.2f, 0.2f);
-				const float DeltaThreshold = VizSettings->MorphTargetDeltaThreshold;
-				for (int32 VertexIndex = 0; VertexIndex < NumVerts; ++VertexIndex)
-				{
-					const FVector StartPoint = FVector(SkinnedPositions[VertexIndex].X, SkinnedPositions[VertexIndex].Y, SkinnedPositions[VertexIndex].Z) + MeshOffset;
-					const int32 DeltaArrayOffset = BlendShapeDeltaStartOffset + VertexIndex;
-					const FVector Delta = FVector(NeuralModel->MorphTargetDeltas[DeltaArrayOffset]);
-
-					if (Delta.Length() > DeltaThreshold)
-					{
-						PDI->DrawLine(StartPoint, StartPoint + Delta, IncludedColor, 10);
-						PDI->DrawPoint(StartPoint + Delta, IncludedColor, 1.0f, 10);
-					}
-					else
-					{
-						PDI->DrawPoint(StartPoint + Delta, ExcludedColor, 0.5f, 10);
-					}
-				}
-			}
+			const FVector DrawOffset = -VizSettings->GetMeshSpacingOffsetVector();
+			DrawMorphTarget(PDI, GetNeuralMorphModel()->MorphTargetDeltas, VizSettings->MorphTargetDeltaThreshold, VizSettings->MorphTargetNumber, DrawOffset);
 		}
 	}
 }	// namespace UE::NeuralMorphModel
diff --git a/Engine/Plugins/Experimental/MLDeformer/NeuralMorphModel/Source/NeuralMorphModelEditor/Private/NeuralMorphModelEditorModule.cpp b/Engine/Plugins/Experimental/MLDeformer/NeuralMorphModel/Source/NeuralMorphModelEditor/Private/NeuralMorphModelEditorModule.cpp
--- a/Engine/Plugins/Experimental/MLDeformer/NeuralMorphModel/Source/NeuralMorphModelEditor/Private/NeuralMorphModelEditorModule.cpp
+++ b/Engine/Plugins/Experimental/MLDeformer/NeuralMorphModel/Source/NeuralMorphModelEditor/Private/NeuralMorphModelEditorModule.cpp
@@ -28,7 +28,7 @@ namespace UE::NeuralMorphModel
 		// Register our custom ML deformer model to the model registry in the ML Deformer Framework.
 		FMLDeformerEditorModule& EditorModule = FModuleManager::LoadModuleChecked<FMLDeformerEditorModule>("MLDeformerFrameworkEditor");
 		FMLDeformerEditorModelRegistry& ModelRegistry = EditorModule.GetModelRegistry();
-		ModelRegistry.RegisterEditorModel(UNeuralMorphModel::StaticClass(), FOnGetEditorModelInstance::CreateStatic(&FNeuralMorphEditorModel::MakeInstance));
+		ModelRegistry.RegisterEditorModel(UNeuralMorphModel::StaticClass(), FOnGetEditorModelInstance::CreateStatic(&FNeuralMorphEditorModel::MakeInstance), /*ModelPriority*/100);
 	}
 
 	void FNeuralMorphModelEditorModule::ShutdownModule()
diff --git a/Engine/Plugins/Experimental/MLDeformer/VertexDeltaModel/Source/VertexDeltaModelEditor/Private/VertexDeltaModelEditorModule.cpp b/Engine/Plugins/Experimental/MLDeformer/VertexDeltaModel/Source/VertexDeltaModelEditor/Private/VertexDeltaModelEditorModule.cpp
--- a/Engine/Plugins/Experimental/MLDeformer/VertexDeltaModel/Source/VertexDeltaModelEditor/Private/VertexDeltaModelEditorModule.cpp
+++ b/Engine/Plugins/Experimental/MLDeformer/VertexDeltaModel/Source/VertexDeltaModelEditor/Private/VertexDeltaModelEditorModule.cpp
@@ -28,7 +28,7 @@ namespace UE::VertexDeltaModel
 		// Register our custom ML deformer model to the model registry in the ML Deformer Framework.
 		FMLDeformerEditorModule& EditorModule = FModuleManager::LoadModuleChecked<FMLDeformerEditorModule>("MLDeformerFrameworkEditor");
 		FMLDeformerEditorModelRegistry& ModelRegistry = EditorModule.GetModelRegistry();
-		ModelRegistry.RegisterEditorModel(UVertexDeltaModel::StaticClass(), FOnGetEditorModelInstance::CreateStatic(&FVertexDeltaEditorModel::MakeInstance));
+		ModelRegistry.RegisterEditorModel(UVertexDeltaModel::StaticClass(), FOnGetEditorModelInstance::CreateStatic(&FVertexDeltaEditorModel::MakeInstance), /*ModelPriority*/50);
 	}
 
 	void FVertexDeltaModelEditorModule::ShutdownModule()
