[MLDeformer]
- Allow for Skeletal Meshes with more meshes than geometry tracks. The geometry cache now defines what parts of the mesh will be deformed by ML Deformer. For example your Fbx with SkelMesh could contain teeth, eyes, etc, while the geometry cache file does not contain those meshes.
- Fixed some issue where a render resource could be deleted while still being initialized, due to multithreading.
- Fixed a crash that could occur when switching ML models.
- Tooltip tweak.
- Few refactors / error reporting changes.
- Moved some properties to the advanced settings.
- Tweaked the morph target compression error tolerance min/max.

diff --git a/Engine/Plugins/Experimental/MLDeformer/LegacyVertexDeltaModel/Source/LegacyVertexDeltaModel/Private/LegacyVertexDeltaModel.cpp b/Engine/Plugins/Experimental/MLDeformer/LegacyVertexDeltaModel/Source/LegacyVertexDeltaModel/Private/LegacyVertexDeltaModel.cpp
--- a/Engine/Plugins/Experimental/MLDeformer/LegacyVertexDeltaModel/Source/LegacyVertexDeltaModel/Private/LegacyVertexDeltaModel.cpp
+++ b/Engine/Plugins/Experimental/MLDeformer/LegacyVertexDeltaModel/Source/LegacyVertexDeltaModel/Private/LegacyVertexDeltaModel.cpp
@@ -37,7 +37,8 @@ ULegacyVertexDeltaModel::ULegacyVertexDeltaModel(const FObjectInitializer& Objec
 		if (MeshMappings.IsEmpty())
 		{
 			TArray<FString> FailedImportedMeshnames;
-			GenerateGeomCacheMeshMappings(SkeletalMesh, GeomCache, MeshMappings, FailedImportedMeshnames);
+			TArray<FString> VertexMisMatchNames;
+			GenerateGeomCacheMeshMappings(SkeletalMesh, GeomCache, MeshMappings, FailedImportedMeshnames, VertexMisMatchNames);
 		}
 
 		UE::MLDeformer::SampleGeomCachePositions(
diff --git a/Engine/Plugins/Experimental/MLDeformer/LegacyVertexDeltaModel/Source/LegacyVertexDeltaModelEditor/Private/LegacyVertexDeltaEditorModel.cpp b/Engine/Plugins/Experimental/MLDeformer/LegacyVertexDeltaModel/Source/LegacyVertexDeltaModelEditor/Private/LegacyVertexDeltaEditorModel.cpp
--- a/Engine/Plugins/Experimental/MLDeformer/LegacyVertexDeltaModel/Source/LegacyVertexDeltaModelEditor/Private/LegacyVertexDeltaEditorModel.cpp
+++ b/Engine/Plugins/Experimental/MLDeformer/LegacyVertexDeltaModel/Source/LegacyVertexDeltaModelEditor/Private/LegacyVertexDeltaEditorModel.cpp
@@ -234,8 +234,7 @@ namespace UE::LegacyVertexDeltaModel
 		ULegacyVertexDeltaModel* VertexDeltaModel = GetVertexDeltaModel();
 		UGeometryCache* GeomCache = VertexDeltaModel->GetGeometryCache();
 		USkeletalMesh* SkeletalMesh = VertexDeltaModel->GetSkeletalMesh();
-		if (!GetGeomCacheVertexErrorText(SkeletalMesh, GeomCache, FText(), FText()).IsEmpty() ||
-			!GetGeomCacheErrorText(SkeletalMesh, GeomCache).IsEmpty())
+		if (!GetGeomCacheErrorText(SkeletalMesh, GeomCache).IsEmpty())
 		{
 			return;
 		}
diff --git a/Engine/Plugins/Experimental/MLDeformer/LegacyVertexDeltaModel/Source/LegacyVertexDeltaModelEditor/Private/LegacyVertexDeltaModelDetails.cpp b/Engine/Plugins/Experimental/MLDeformer/LegacyVertexDeltaModel/Source/LegacyVertexDeltaModelEditor/Private/LegacyVertexDeltaModelDetails.cpp
--- a/Engine/Plugins/Experimental/MLDeformer/LegacyVertexDeltaModel/Source/LegacyVertexDeltaModelEditor/Private/LegacyVertexDeltaModelDetails.cpp
+++ b/Engine/Plugins/Experimental/MLDeformer/LegacyVertexDeltaModel/Source/LegacyVertexDeltaModelEditor/Private/LegacyVertexDeltaModelDetails.cpp
@@ -81,24 +81,6 @@ namespace UE::LegacyVertexDeltaModel
 		SettingsCategoryBuilder->AddProperty(GET_MEMBER_NAME_CHECKED(ULegacyVertexDeltaModel, MaxCacheSizeGB));
 	}
 
-	void FLegacyVertexDeltaModelDetails::AddBaseMeshErrors()
-	{
-		// Check if the base mesh matches the target mesh vertex count.
-		FText ErrorText = GetGeomCacheVertexErrorText(Model->SkeletalMesh, VertexModel->GetGeometryCache(), FText::FromString("Base Mesh"), FText::FromString("Target Mesh"));
-		BaseMeshCategoryBuilder->AddCustomRow(FText::FromString("BaseMeshError"))
-			.Visibility(!ErrorText.IsEmpty() ? EVisibility::Visible : EVisibility::Collapsed)
-			.WholeRowContent()
-			[
-				SNew(SBox)
-				.Padding(FMargin(0.0f, 4.0f))
-				[
-					SNew(SWarningOrErrorBox)
-					.MessageStyle(EMessageStyle::Error)
-					.Message(ErrorText)
-				]
-			];
-	}
-	
 	void FLegacyVertexDeltaModelDetails::AddAnimSequenceErrors()
 	{
 		const FText WarningText = GetGeomCacheAnimSequenceErrorText(VertexModel->GetGeometryCache(), Model->GetAnimSequence());
@@ -148,39 +130,7 @@ namespace UE::LegacyVertexDeltaModel
 				]
 			];
 
-		// Show meshes that have no matching geometry track list.
-		const FText MeshMappingErrorList = GetGeomCacheMeshMappingErrorText(Model->GetSkeletalMesh(), VertexModel->GetGeometryCache());
-		FString GeomTrackNameList;
-		if (!MeshMappingErrorList.IsEmpty())
-		{
-			UGeometryCache* GeomCache = VertexModel->GetGeometryCache();
-			for (int32 Index = 0; Index < GeomCache->Tracks.Num(); ++Index)
-			{
-				GeomTrackNameList += GeomCache->Tracks[Index]->GetName();
-				if (Index < GeomCache->Tracks.Num() - 1)
-				{
-					GeomTrackNameList += TEXT("\n");
-				}
-			}
-		}
-
-		FText MeshMappingErrorFull = FText::Format(
-			LOCTEXT("MeshMappingError", "No matching GeomCache Tracks names found for meshes:\n{0}\n\nGeomCache Track List:\n{1}"), 
-			MeshMappingErrorList,
-			FText::FromString(GeomTrackNameList));
-
-		TargetMeshCategoryBuilder->AddCustomRow(FText::FromString("MeshMappingError"))
-			.Visibility(!MeshMappingErrorList.IsEmpty() ? EVisibility::Visible : EVisibility::Collapsed)
-			.WholeRowContent()
-			[
-				SNew(SBox)
-				.Padding(FMargin(0.0f, 4.0f))
-				[
-					SNew(SWarningOrErrorBox)
-					.MessageStyle(EMessageStyle::Error)
-					.Message(MeshMappingErrorFull)
-				]
-			];
+		AddGeomCacheMeshMappingWarnings(TargetMeshCategoryBuilder, Model->GetSkeletalMesh(), VertexModel->GetGeometryCache());
 	}
 }	// namespace UE::LegacyVertexDeltaModel
 
diff --git a/Engine/Plugins/Experimental/MLDeformer/LegacyVertexDeltaModel/Source/LegacyVertexDeltaModelEditor/Private/LegacyVertexDeltaModelVizSettingsDetails.cpp b/Engine/Plugins/Experimental/MLDeformer/LegacyVertexDeltaModel/Source/LegacyVertexDeltaModelEditor/Private/LegacyVertexDeltaModelVizSettingsDetails.cpp
--- a/Engine/Plugins/Experimental/MLDeformer/LegacyVertexDeltaModel/Source/LegacyVertexDeltaModelEditor/Private/LegacyVertexDeltaModelVizSettingsDetails.cpp
+++ b/Engine/Plugins/Experimental/MLDeformer/LegacyVertexDeltaModel/Source/LegacyVertexDeltaModelEditor/Private/LegacyVertexDeltaModelVizSettingsDetails.cpp
@@ -84,20 +84,6 @@ namespace UE::LegacyVertexDeltaModel
 					.Message(GeomErrorText)
 				]
 			];
-
-		const FText VertexErrorText = GetGeomCacheVertexErrorText(Model->GetSkeletalMesh(), VertexDeltaVizSettings->GetTestGroundTruth(), FText::FromString("Base Mesh"), FText::FromString("Ground Truth Mesh"));
-		FDetailWidgetRow& GroundTruthVertexErrorRow = TestAssetsCategory->AddCustomRow(FText::FromString("GroundTruthVertexMismatchError"))
-			.Visibility(!VertexErrorText.IsEmpty() ? EVisibility::Visible : EVisibility::Collapsed)
-			.WholeRowContent()
-			[
-				SNew(SBox)
-				.Padding(FMargin(0.0f, 4.0f))
-				[
-					SNew(SWarningOrErrorBox)
-					.MessageStyle(EMessageStyle::Warning)
-					.Message(VertexErrorText)
-				]
-			];
 	}
 }	//namespace UE::LegacyVertexDeltaModel
 
diff --git a/Engine/Plugins/Experimental/MLDeformer/LegacyVertexDeltaModel/Source/LegacyVertexDeltaModelEditor/Public/LegacyVertexDeltaModelDetails.h b/Engine/Plugins/Experimental/MLDeformer/LegacyVertexDeltaModel/Source/LegacyVertexDeltaModelEditor/Public/LegacyVertexDeltaModelDetails.h
--- a/Engine/Plugins/Experimental/MLDeformer/LegacyVertexDeltaModel/Source/LegacyVertexDeltaModelEditor/Public/LegacyVertexDeltaModelDetails.h
+++ b/Engine/Plugins/Experimental/MLDeformer/LegacyVertexDeltaModel/Source/LegacyVertexDeltaModelEditor/Public/LegacyVertexDeltaModelDetails.h
@@ -34,7 +34,6 @@ namespace UE::LegacyVertexDeltaModel
 		// FMLDeformerModelDetails overrides.
 		bool UpdateMemberPointers(const TArray<TWeakObjectPtr<UObject>>& Objects) override;
 		void AddTargetMesh() override;
-		void AddBaseMeshErrors() override;
 		void AddAnimSequenceErrors() override;
 		// ~END FMLDeformerModelDetails overrides.
 
diff --git a/Engine/Plugins/Experimental/MLDeformer/MLDeformerFramework/Source/MLDeformerFramework/Private/MLDeformerGeomCacheHelpers.cpp b/Engine/Plugins/Experimental/MLDeformer/MLDeformerFramework/Source/MLDeformerFramework/Private/MLDeformerGeomCacheHelpers.cpp
--- a/Engine/Plugins/Experimental/MLDeformer/MLDeformerFramework/Source/MLDeformerFramework/Private/MLDeformerGeomCacheHelpers.cpp
+++ b/Engine/Plugins/Experimental/MLDeformer/MLDeformerFramework/Source/MLDeformerFramework/Private/MLDeformerGeomCacheHelpers.cpp
@@ -70,50 +70,48 @@ namespace UE::MLDeformer
 		return Result;
 	}
 
-	FText GetGeomCacheVertexErrorText(USkeletalMesh* InSkelMesh, UGeometryCache* InGeomCache, const FText& SkelName, const FText& GeomCacheName)
-	{
-		FText Result;
-
-		if (InSkelMesh && InGeomCache)
-		{
-			const int32 SkelVertCount = UMLDeformerModel::ExtractNumImportedSkinnedVertices(InSkelMesh);
-			const int32 GeomCacheVertCount = ExtractNumImportedGeomCacheVertices(InGeomCache);
-			const bool bHasGeomCacheError = !GetGeomCacheErrorText(InSkelMesh, InGeomCache).IsEmpty();
-			if (SkelVertCount != GeomCacheVertCount && !bHasGeomCacheError)
-			{
-				Result = FText::Format(
-					LOCTEXT("MeshVertexNumVertsMismatch", "Vertex count of {0} doesn't match with {1}!\n\n{2} has {3} verts, while {4} has {5} verts."),
-					SkelName,
-					GeomCacheName,
-					SkelName,
-					SkelVertCount,
-					GeomCacheName,
-					GeomCacheVertCount);
-			}
-		}
-
-		return Result;
-	}
-
 	FText GetGeomCacheMeshMappingErrorText(USkeletalMesh* InSkelMesh, UGeometryCache* InGeomCache)
 	{
 		FText Result;
-
 		if (InGeomCache && InSkelMesh)
 		{
 			// Check for failed mesh mappings.
 			TArray<FMLDeformerGeomCacheMeshMapping> MeshMappings;
 			TArray<FString> FailedNames;
-			GenerateGeomCacheMeshMappings(InSkelMesh, InGeomCache, MeshMappings, FailedNames);
+			TArray<FString> VertexMisMatchNames;
+			GenerateGeomCacheMeshMappings(InSkelMesh, InGeomCache, MeshMappings, FailedNames, VertexMisMatchNames);
 
 			// List all mesh names that have issues.
 			FString ErrorString;
-			for (int32 Index = 0; Index < FailedNames.Num(); ++Index)
+			if (!FailedNames.IsEmpty())
+			{
+				const FText ErrorText = LOCTEXT("FailedMappingNames", "No SkelMesh meshes found for GeomCache tracks:\n\n");
+				ErrorString += ErrorText.ToString();
+				for (int32 Index = 0; Index < FailedNames.Num(); ++Index)
+				{
+					ErrorString += "\t\t" + FailedNames[Index];
+					if (Index < FailedNames.Num() - 1)
+					{
+						ErrorString += TEXT("\n");
+					}
+				}
+			}
+
+			if (VertexMisMatchNames.Num() > 0)
 			{
-				ErrorString += FailedNames[Index];
-				if (Index < FailedNames.Num() - 1)
+				if (!FailedNames.IsEmpty())
+				{
+					ErrorString += "\n\n";
+				}
+				const FText ErrorText = LOCTEXT("FailedMappingVertexMisMatch", "Mismatching vertex counts for:\n\n");
+				ErrorString += ErrorText.ToString();
+				for (int32 Index = 0; Index < VertexMisMatchNames.Num(); ++Index)
 				{
-					ErrorString += TEXT("\n");
+					ErrorString += "\t\t" + VertexMisMatchNames[Index];
+					if (Index < VertexMisMatchNames.Num() - 1)
+					{
+						ErrorString += TEXT("\n");
+					}
 				}
 			}
 
@@ -130,10 +128,11 @@ namespace UE::MLDeformer
 		return (TrackName.Find(MeshName) == 0);
 	}
 
-	void GenerateGeomCacheMeshMappings(USkeletalMesh* SkelMesh, UGeometryCache* GeomCache, TArray<FMLDeformerGeomCacheMeshMapping>& OutMeshMappings, TArray<FString>& OutFailedImportedMeshNames)
+	void GenerateGeomCacheMeshMappings(USkeletalMesh* SkelMesh, UGeometryCache* GeomCache, TArray<FMLDeformerGeomCacheMeshMapping>& OutMeshMappings, TArray<FString>& OutFailedImportedMeshNames, TArray<FString>& OutVertexMisMatchNames)
 	{
 		OutMeshMappings.Empty();
 		OutFailedImportedMeshNames.Empty();
+		OutVertexMisMatchNames.Empty();
 		if (SkelMesh == nullptr || GeomCache == nullptr)
 		{
 			return;
@@ -152,21 +151,22 @@ namespace UE::MLDeformer
 		// For all meshes in the skeletal mesh.
 		const float SampleTime = 0.0f;
 		FString SkelMeshName;
-		for (int32 SkelMeshIndex = 0; SkelMeshIndex < SkelMeshInfos.Num(); ++SkelMeshIndex)
+
+		const bool bIsSoloMesh = (GeomCache->Tracks.Num() == 1 && SkelMeshInfos.Num() == 1);	// Do we just have one mesh and one track?
+		for (int32 TrackIndex = 0; TrackIndex < GeomCache->Tracks.Num(); ++TrackIndex)
 		{
-			const FSkelMeshImportedMeshInfo& MeshInfo = SkelMeshInfos[SkelMeshIndex];
-			SkelMeshName = MeshInfo.Name.ToString();
+			// Check if this is a candidate based on the mesh and track name.
+			UGeometryCacheTrack* Track = GeomCache->Tracks[TrackIndex];
 
-			// Find the matching one in the geom cache.
 			bool bFoundMatch = false;
-			for (int32 TrackIndex = 0; TrackIndex < GeomCache->Tracks.Num(); ++TrackIndex)
+			for (int32 SkelMeshIndex = 0; SkelMeshIndex < SkelMeshInfos.Num(); ++SkelMeshIndex)
 			{
-				// Check if this is a candidate based on the mesh and track name.
-				UGeometryCacheTrack* Track = GeomCache->Tracks[TrackIndex];
-				const bool bIsSoloMesh = (GeomCache->Tracks.Num() == 1 && SkelMeshInfos.Num() == 1);	// Do we just have one mesh and one track?
-				if (Track && 
+				const FSkelMeshImportedMeshInfo& MeshInfo = SkelMeshInfos[SkelMeshIndex];
+				SkelMeshName = MeshInfo.Name.ToString();
+
+				if (Track &&
 					(IsPotentialMatch(Track->GetName(), SkelMeshName) || bIsSoloMesh))
-				{	
+				{
 					// Extract the geom cache mesh data.
 					FGeometryCacheMeshData GeomCacheMeshData;
 					if (!Track->GetMeshDataAtTime(SampleTime, GeomCacheMeshData))
@@ -192,6 +192,7 @@ namespace UE::MLDeformer
 					const int32 NumSkelMeshVerts = MeshInfo.NumVertices;
 					if (NumSkelMeshVerts != NumGeomMeshVerts)
 					{
+						OutVertexMisMatchNames.Add(Track->GetName());
 						continue;
 					}
 
@@ -211,7 +212,7 @@ namespace UE::MLDeformer
 						// However they all share the same vertex position, so we can just find the first hit, as we only need the position later on.
 						const int32 GeomCacheVertexIndex = GeomCacheMeshData.ImportedVertexNumbers.Find(VertexIndex);
 						Mapping.SkelMeshToTrackVertexMap[VertexIndex] = GeomCacheVertexIndex;
-					
+
 						// Map the source asset vertex number to a render vertex. This is the first duplicate of that vertex.
 						const int32 RenderVertexIndex = ImportedModel->LODModels[0].MeshToImportVertexMap.Find(MeshInfo.StartImportedVertex + VertexIndex);
 						Mapping.ImportedVertexToRenderVertexMap[VertexIndex] = RenderVertexIndex;
@@ -221,14 +222,14 @@ namespace UE::MLDeformer
 					bFoundMatch = true;
 					break;
 				} // If the track name matches the skeletal meshes internal mesh name.
-			} // For all tracks.
+			} // For all meshes in the Skeletal Mesh.
 
 			if (!bFoundMatch)
 			{
-				OutFailedImportedMeshNames.Add(SkelMeshName);
-				UE_LOG(LogMLDeformer, Warning, TEXT("Imported mesh '%s' cannot be matched with a geometry cache track."), *SkelMeshName);
+				OutFailedImportedMeshNames.Add(Track->GetName());
+				UE_LOG(LogMLDeformer, Warning, TEXT("Geometry cache '%s' cannot be matched with a mesh inside the Skeletal Mesh."), *Track->GetName());
 			}
-		} // For all meshes inside the skeletal mesh.
+		} // For all tracks.
 	}
 
 	void SampleGeomCachePositions(
diff --git a/Engine/Plugins/Experimental/MLDeformer/MLDeformerFramework/Source/MLDeformerFramework/Public/MLDeformerGeomCacheHelpers.h b/Engine/Plugins/Experimental/MLDeformer/MLDeformerFramework/Source/MLDeformerFramework/Public/MLDeformerGeomCacheHelpers.h
--- a/Engine/Plugins/Experimental/MLDeformer/MLDeformerFramework/Source/MLDeformerFramework/Public/MLDeformerGeomCacheHelpers.h
+++ b/Engine/Plugins/Experimental/MLDeformer/MLDeformerFramework/Source/MLDeformerFramework/Public/MLDeformerGeomCacheHelpers.h
@@ -24,12 +24,11 @@ namespace UE::MLDeformer
 	// Error checks.
 	MLDEFORMERFRAMEWORK_API FText GetGeomCacheErrorText(USkeletalMesh* InSkeletalMesh, UGeometryCache* InGeomCache);
 	MLDEFORMERFRAMEWORK_API FText GetGeomCacheAnimSequenceErrorText(UGeometryCache* InGeomCache, UAnimSequence* InAnimSequence);
-	MLDEFORMERFRAMEWORK_API FText GetGeomCacheVertexErrorText(USkeletalMesh* InSkelMesh, UGeometryCache* InGeomCache, const FText& SkelName, const FText& GeomCacheName);
 	MLDEFORMERFRAMEWORK_API FText GetGeomCacheMeshMappingErrorText(USkeletalMesh* InSkelMesh, UGeometryCache* InGeomCache);
 
 	// Geom cache operations.
 	MLDEFORMERFRAMEWORK_API int32 ExtractNumImportedGeomCacheVertices(UGeometryCache* GeometryCache);
-	MLDEFORMERFRAMEWORK_API void GenerateGeomCacheMeshMappings(USkeletalMesh* SkelMesh, UGeometryCache* GeomCache, TArray<FMLDeformerGeomCacheMeshMapping>& OutMeshMappings, TArray<FString>& OutFailedImportedMeshNames);
+	MLDEFORMERFRAMEWORK_API void GenerateGeomCacheMeshMappings(USkeletalMesh* SkelMesh, UGeometryCache* GeomCache, TArray<FMLDeformerGeomCacheMeshMapping>& OutMeshMappings, TArray<FString>& OutFailedImportedMeshNames, TArray<FString>& OutVertexMisMatchNames);
 	MLDEFORMERFRAMEWORK_API void SampleGeomCachePositions(
 		int32 InLODIndex,
 		float InSampleTime,
diff --git a/Engine/Plugins/Experimental/MLDeformer/MLDeformerFramework/Source/MLDeformerFrameworkEditor/Private/MLDeformerEditorToolkit.cpp b/Engine/Plugins/Experimental/MLDeformer/MLDeformerFramework/Source/MLDeformerFrameworkEditor/Private/MLDeformerEditorToolkit.cpp
--- a/Engine/Plugins/Experimental/MLDeformer/MLDeformerFramework/Source/MLDeformerFrameworkEditor/Private/MLDeformerEditorToolkit.cpp
+++ b/Engine/Plugins/Experimental/MLDeformer/MLDeformerFramework/Source/MLDeformerFrameworkEditor/Private/MLDeformerEditorToolkit.cpp
@@ -414,7 +414,7 @@ namespace UE::MLDeformer
 				),
 				NAME_None,
 				LOCTEXT("TrainModel", "Train Model"),
-				LOCTEXT("TrainModelTooltip", "Train Model using Pytorch"),
+				LOCTEXT("TrainModelTooltip", "Train the machine learning model.\nIf this button is disabled you first need to configure your inputs in the details panel."),
 				FSlateIcon(),
 				EUserInterfaceActionType::ToggleButton
 			);
diff --git a/Engine/Plugins/Experimental/MLDeformer/MLDeformerFramework/Source/MLDeformerFrameworkEditor/Private/MLDeformerGeomCacheSampler.cpp b/Engine/Plugins/Experimental/MLDeformer/MLDeformerFramework/Source/MLDeformerFrameworkEditor/Private/MLDeformerGeomCacheSampler.cpp
--- a/Engine/Plugins/Experimental/MLDeformer/MLDeformerFramework/Source/MLDeformerFrameworkEditor/Private/MLDeformerGeomCacheSampler.cpp
+++ b/Engine/Plugins/Experimental/MLDeformer/MLDeformerFramework/Source/MLDeformerFrameworkEditor/Private/MLDeformerGeomCacheSampler.cpp
@@ -44,7 +44,8 @@ namespace UE::MLDeformer
 				Model->GetSkeletalMesh(), 
 				GeomCache, 
 				MeshMappings, 
-				FailedImportedMeshNames);
+				FailedImportedMeshNames,
+				VertexCountMisMatchNames);
 
 			GeomCacheMeshDatas.Reset(MeshMappings.Num());
 			GeomCacheMeshDatas.AddDefaulted(MeshMappings.Num());
diff --git a/Engine/Plugins/Experimental/MLDeformer/MLDeformerFramework/Source/MLDeformerFrameworkEditor/Private/MLDeformerModelDetails.cpp b/Engine/Plugins/Experimental/MLDeformer/MLDeformerFramework/Source/MLDeformerFrameworkEditor/Private/MLDeformerModelDetails.cpp
--- a/Engine/Plugins/Experimental/MLDeformer/MLDeformerFramework/Source/MLDeformerFrameworkEditor/Private/MLDeformerModelDetails.cpp
+++ b/Engine/Plugins/Experimental/MLDeformer/MLDeformerFramework/Source/MLDeformerFrameworkEditor/Private/MLDeformerModelDetails.cpp
@@ -6,6 +6,7 @@
 #include "MLDeformerModel.h"
 #include "MLDeformerEditorModel.h"
 #include "MLDeformerEditorModule.h"
+#include "MLDeformerGeomCacheHelpers.h"
 #include "NeuralNetwork.h"
 #include "GeometryCache.h"
 #include "GeometryCacheTrack.h"
@@ -273,6 +274,23 @@ namespace UE::MLDeformer
 		DetailLayoutBuilder->ForceRefreshDetails();
 		return FReply::Handled();
 	}
+
+	void FMLDeformerModelDetails::AddGeomCacheMeshMappingWarnings(IDetailCategoryBuilder* InTargetMeshCategoryBuilder, USkeletalMesh* SkeletalMesh, UGeometryCache* GeometryCache)
+	{
+		const FText MeshMappingError = GetGeomCacheMeshMappingErrorText(SkeletalMesh, GeometryCache);
+		InTargetMeshCategoryBuilder->AddCustomRow(FText::FromString("MeshMappingWarning"))
+			.Visibility(!MeshMappingError.IsEmpty() ? EVisibility::Visible : EVisibility::Collapsed)
+			.WholeRowContent()
+			[
+				SNew(SBox)
+				.Padding(FMargin(0.0f, 4.0f))
+				[
+					SNew(SWarningOrErrorBox)
+					.MessageStyle(EMessageStyle::Warning)
+					.Message(MeshMappingError)
+				]
+			];
+	}
 }	// namespace UE::MLDeformer
 
 #undef LOCTEXT_NAMESPACE
diff --git a/Engine/Plugins/Experimental/MLDeformer/MLDeformerFramework/Source/MLDeformerFrameworkEditor/Public/MLDeformerGeomCacheSampler.h b/Engine/Plugins/Experimental/MLDeformer/MLDeformerFramework/Source/MLDeformerFrameworkEditor/Public/MLDeformerGeomCacheSampler.h
--- a/Engine/Plugins/Experimental/MLDeformer/MLDeformerFramework/Source/MLDeformerFrameworkEditor/Public/MLDeformerGeomCacheSampler.h
+++ b/Engine/Plugins/Experimental/MLDeformer/MLDeformerFramework/Source/MLDeformerFrameworkEditor/Public/MLDeformerGeomCacheSampler.h
@@ -48,9 +48,12 @@ namespace UE::MLDeformer
 		/** The geometry cache mesh data reusable buffers. One for each MeshMapping.*/
 		TArray<FGeometryCacheMeshData> GeomCacheMeshDatas;
 
-		/** Imported mesh names in the skeletal mesh for which no geom cache track could be found. */
+		/** Geom cache track names for which no mesh can be found inside the skeletal mesh. */
 		TArray<FString> FailedImportedMeshNames; 
 
+		/** Imported mesh names in the skeletal mesh for which the geometry track had a different vertex count. */
+		TArray<FString> VertexCountMisMatchNames;
+
 		/** The function that grabs the geometry cache. */
 		FMLDeformerGetGeomCacheEvent GetGeometryCacheEvent;
 	};
diff --git a/Engine/Plugins/Experimental/MLDeformer/MLDeformerFramework/Source/MLDeformerFrameworkEditor/Public/MLDeformerModelDetails.h b/Engine/Plugins/Experimental/MLDeformer/MLDeformerFramework/Source/MLDeformerFrameworkEditor/Public/MLDeformerModelDetails.h
--- a/Engine/Plugins/Experimental/MLDeformer/MLDeformerFramework/Source/MLDeformerFrameworkEditor/Public/MLDeformerModelDetails.h
+++ b/Engine/Plugins/Experimental/MLDeformer/MLDeformerFramework/Source/MLDeformerFrameworkEditor/Public/MLDeformerModelDetails.h
@@ -11,6 +11,8 @@ class IDetailLayoutBuilder;
 class USkeleton;
 class IDetailCategoryBuilder;
 class UMLDeformerModel;
+class UGeometryCache;
+class USkeletalMesh;
 
 namespace UE::MLDeformer
 {
@@ -35,10 +37,10 @@ namespace UE::MLDeformer
 
 	protected:
 		bool FilterAnimSequences(const FAssetData& AssetData, USkeleton* Skeleton);
+		void AddGeomCacheMeshMappingWarnings(IDetailCategoryBuilder* InTargetMeshCategoryBuilder, USkeletalMesh* SkeletalMesh, UGeometryCache* GeometryCache);
 		FReply OnFilterAnimatedBonesOnly() const;
 		FReply OnFilterAnimatedCurvesOnly() const;
 
-
 	protected:
 		/** Associated detail layout builder. */
 		IDetailLayoutBuilder* DetailLayoutBuilder = nullptr;
diff --git a/Engine/Plugins/Experimental/MLDeformer/NeuralMorphModel/Source/NeuralMorphModel/Private/NeuralMorphModel.cpp b/Engine/Plugins/Experimental/MLDeformer/NeuralMorphModel/Source/NeuralMorphModel/Private/NeuralMorphModel.cpp
--- a/Engine/Plugins/Experimental/MLDeformer/NeuralMorphModel/Source/NeuralMorphModel/Private/NeuralMorphModel.cpp
+++ b/Engine/Plugins/Experimental/MLDeformer/NeuralMorphModel/Source/NeuralMorphModel/Private/NeuralMorphModel.cpp
@@ -77,8 +77,9 @@ UMLDeformerModelInstance* UNeuralMorphModel::CreateModelInstance(UMLDeformerComp
 
 		if (MeshMappings.IsEmpty())
 		{
-			TArray<FString> FailedImportedMeshnames;
-			GenerateGeomCacheMeshMappings(SkeletalMesh, GeomCache, MeshMappings, FailedImportedMeshnames);
+			TArray<FString> FailedImportedMeshNames;
+			TArray<FString> VertexMisMatchNames;
+			GenerateGeomCacheMeshMappings(SkeletalMesh, GeomCache, MeshMappings, FailedImportedMeshNames, VertexMisMatchNames);
 		}
 
 		UE::MLDeformer::SampleGeomCachePositions(
@@ -163,7 +164,10 @@ void UNeuralMorphModel::BeginDestroy()
 {
 	if (MorphTargetSet.IsValid())
 	{
-		BeginReleaseResource(&MorphTargetSet->MorphBuffers);
+		// Release and flush, waiting for the release to have completed, 
+		// If we don't do this we can get an error that we destroy a render resource that is still initialized,
+		// as the release happens in another thread.
+		ReleaseResourceAndFlush(&MorphTargetSet->MorphBuffers);
 		MorphTargetSet.Reset();
 	}
 	Super::BeginDestroy();
diff --git a/Engine/Plugins/Experimental/MLDeformer/NeuralMorphModel/Source/NeuralMorphModel/Public/NeuralMorphModel.h b/Engine/Plugins/Experimental/MLDeformer/NeuralMorphModel/Source/NeuralMorphModel/Public/NeuralMorphModel.h
--- a/Engine/Plugins/Experimental/MLDeformer/NeuralMorphModel/Source/NeuralMorphModel/Public/NeuralMorphModel.h
+++ b/Engine/Plugins/Experimental/MLDeformer/NeuralMorphModel/Source/NeuralMorphModel/Public/NeuralMorphModel.h
@@ -91,18 +91,18 @@ public:
 	UPROPERTY(EditAnywhere, BlueprintReadWrite, Category = "Training Settings", meta = (ClampMin = "1", ClampMax = "100"))
 	int32 NumMorphTargetsPerBone = 6;
 
+	/** The number of iterations to train the model for. */
+	UPROPERTY(EditAnywhere, BlueprintReadWrite, Category = "Training Settings", meta = (ClampMin = "1"))
+	int32 NumIterations = 2000;
+
 	/** The number of hidden layers that the neural network model will have.\nHigher numbers will slow down performance but can deal with more complex deformations. */
-	UPROPERTY(EditAnywhere, BlueprintReadWrite, Category = "Training Settings", meta = (ClampMin = "1", ClampMax = "10"))
+	UPROPERTY(EditAnywhere, BlueprintReadWrite, AdvancedDisplay, Category = "Training Settings", meta = (ClampMin = "1", ClampMax = "10"))
 	int32 NumHiddenLayers = 1;
 
 	/** The number of units/neurons per hidden layer. Higher numbers will slow down performance but allow for more complex mesh deformations. */
-	UPROPERTY(EditAnywhere, BlueprintReadWrite, Category = "Training Settings", meta = (ClampMin = "1"))
+	UPROPERTY(EditAnywhere, BlueprintReadWrite, AdvancedDisplay, Category = "Training Settings", meta = (ClampMin = "1"))
 	int32 NumNeuronsPerLayer = 6;
 
-	/** The number of iterations to train the model for. */
-	UPROPERTY(EditAnywhere, BlueprintReadWrite, Category = "Training Settings", meta = (ClampMin = "1"))
-	int32 NumIterations = 2000;
-
 	/** The number of frames per batch when training the model. */
 	UPROPERTY(EditAnywhere, BlueprintReadWrite, AdvancedDisplay, Category = "Training Settings", meta = (ClampMin = "1"))
 	int32 BatchSize = 128;
@@ -127,7 +127,7 @@ public:
 	float MorphTargetDeltaThreshold = 0.0025f;
 
 	/** The morph target error tolerance. Higher values result in larger compression, but could result in visual artifacts. */
-	UPROPERTY(EditAnywhere, BlueprintReadWrite, AdvancedDisplay, Category = "Training Settings", meta = (ClampMin = "0.01", ClampMax = "10000.0"))
-	float MorphTargetErrorTolerance = 20.0f;
+	UPROPERTY(EditAnywhere, BlueprintReadWrite, AdvancedDisplay, Category = "Training Settings", meta = (ClampMin = "1", ClampMax = "1000"))
+	float MorphTargetErrorTolerance = 50.0f;
 #endif // WITH_EDITORONLY_DATA
 };
diff --git a/Engine/Plugins/Experimental/MLDeformer/NeuralMorphModel/Source/NeuralMorphModelEditor/Private/NeuralMorphEditorModel.cpp b/Engine/Plugins/Experimental/MLDeformer/NeuralMorphModel/Source/NeuralMorphModelEditor/Private/NeuralMorphEditorModel.cpp
--- a/Engine/Plugins/Experimental/MLDeformer/NeuralMorphModel/Source/NeuralMorphModelEditor/Private/NeuralMorphEditorModel.cpp
+++ b/Engine/Plugins/Experimental/MLDeformer/NeuralMorphModel/Source/NeuralMorphModelEditor/Private/NeuralMorphEditorModel.cpp
@@ -260,8 +260,7 @@ namespace UE::NeuralMorphModel
 		UNeuralMorphModel* NeuralMorphModel = GetNeuralMorphModel();
 		UGeometryCache* GeomCache = NeuralMorphModel->GetGeometryCache();
 		USkeletalMesh* SkeletalMesh = NeuralMorphModel->GetSkeletalMesh();
-		if (!GetGeomCacheVertexErrorText(SkeletalMesh, GeomCache, FText(), FText()).IsEmpty() ||
-			!GetGeomCacheErrorText(SkeletalMesh, GeomCache).IsEmpty())
+		if (!GetGeomCacheErrorText(SkeletalMesh, GeomCache).IsEmpty())
 		{
 			return;
 		}
diff --git a/Engine/Plugins/Experimental/MLDeformer/NeuralMorphModel/Source/NeuralMorphModelEditor/Private/NeuralMorphModelDetails.cpp b/Engine/Plugins/Experimental/MLDeformer/NeuralMorphModel/Source/NeuralMorphModelEditor/Private/NeuralMorphModelDetails.cpp
--- a/Engine/Plugins/Experimental/MLDeformer/NeuralMorphModel/Source/NeuralMorphModelEditor/Private/NeuralMorphModelDetails.cpp
+++ b/Engine/Plugins/Experimental/MLDeformer/NeuralMorphModel/Source/NeuralMorphModelEditor/Private/NeuralMorphModelDetails.cpp
@@ -68,24 +68,6 @@ namespace UE::NeuralMorphModel
 		SettingsCategoryBuilder->AddProperty(GET_MEMBER_NAME_CHECKED(UNeuralMorphModel, RegularizationFactor));
 		SettingsCategoryBuilder->AddProperty(GET_MEMBER_NAME_CHECKED(UNeuralMorphModel, MorphTargetErrorTolerance));
 	}
-
-	void FNeuralMorphModelDetails::AddBaseMeshErrors()
-	{
-		// Check if the base mesh matches the target mesh vertex count.
-		FText ErrorText = GetGeomCacheVertexErrorText(Model->SkeletalMesh, NeuralMorphModel->GetGeometryCache(), FText::FromString("Base Mesh"), FText::FromString("Target Mesh"));
-		BaseMeshCategoryBuilder->AddCustomRow(FText::FromString("BaseMeshError"))
-			.Visibility(!ErrorText.IsEmpty() ? EVisibility::Visible : EVisibility::Collapsed)
-			.WholeRowContent()
-			[
-				SNew(SBox)
-				.Padding(FMargin(0.0f, 4.0f))
-				[
-					SNew(SWarningOrErrorBox)
-					.MessageStyle(EMessageStyle::Error)
-					.Message(ErrorText)
-				]
-			];
-	}
 	
 	void FNeuralMorphModelDetails::AddAnimSequenceErrors()
 	{
@@ -136,39 +118,7 @@ namespace UE::NeuralMorphModel
 				]
 			];
 
-		// Show meshes that have no matching geometry track list.
-		const FText MeshMappingErrorList = GetGeomCacheMeshMappingErrorText(Model->GetSkeletalMesh(), NeuralMorphModel->GetGeometryCache());
-		FString GeomTrackNameList;
-		if (!MeshMappingErrorList.IsEmpty())
-		{
-			UGeometryCache* GeomCache = NeuralMorphModel->GetGeometryCache();
-			for (int32 Index = 0; Index < GeomCache->Tracks.Num(); ++Index)
-			{
-				GeomTrackNameList += GeomCache->Tracks[Index]->GetName();
-				if (Index < GeomCache->Tracks.Num() - 1)
-				{
-					GeomTrackNameList += TEXT("\n");
-				}
-			}
-		}
-
-		FText MeshMappingErrorFull = FText::Format(
-			LOCTEXT("MeshMappingError", "No matching GeomCache Tracks names found for meshes:\n{0}\n\nGeomCache Track List:\n{1}"), 
-			MeshMappingErrorList,
-			FText::FromString(GeomTrackNameList));
-
-		TargetMeshCategoryBuilder->AddCustomRow(FText::FromString("MeshMappingError"))
-			.Visibility(!MeshMappingErrorList.IsEmpty() ? EVisibility::Visible : EVisibility::Collapsed)
-			.WholeRowContent()
-			[
-				SNew(SBox)
-				.Padding(FMargin(0.0f, 4.0f))
-				[
-					SNew(SWarningOrErrorBox)
-					.MessageStyle(EMessageStyle::Error)
-					.Message(MeshMappingErrorFull)
-				]
-			];
+		AddGeomCacheMeshMappingWarnings(TargetMeshCategoryBuilder, Model->GetSkeletalMesh(), NeuralMorphModel->GetGeometryCache());
 	}
 }	// namespace UE::NeuralMorphModel
 
diff --git a/Engine/Plugins/Experimental/MLDeformer/NeuralMorphModel/Source/NeuralMorphModelEditor/Private/NeuralMorphModelVizSettingsDetails.cpp b/Engine/Plugins/Experimental/MLDeformer/NeuralMorphModel/Source/NeuralMorphModelEditor/Private/NeuralMorphModelVizSettingsDetails.cpp
--- a/Engine/Plugins/Experimental/MLDeformer/NeuralMorphModel/Source/NeuralMorphModelEditor/Private/NeuralMorphModelVizSettingsDetails.cpp
+++ b/Engine/Plugins/Experimental/MLDeformer/NeuralMorphModel/Source/NeuralMorphModelEditor/Private/NeuralMorphModelVizSettingsDetails.cpp
@@ -84,20 +84,6 @@ namespace UE::NeuralMorphModel
 					.Message(GeomErrorText)
 				]
 			];
-
-		const FText VertexErrorText = GetGeomCacheVertexErrorText(Model->GetSkeletalMesh(), NeuralMorphVizSettings->GetTestGroundTruth(), FText::FromString("Base Mesh"), FText::FromString("Ground Truth Mesh"));
-		FDetailWidgetRow& GroundTruthVertexErrorRow = TestAssetsCategory->AddCustomRow(FText::FromString("GroundTruthVertexMismatchError"))
-			.Visibility(!VertexErrorText.IsEmpty() ? EVisibility::Visible : EVisibility::Collapsed)
-			.WholeRowContent()
-			[
-				SNew(SBox)
-				.Padding(FMargin(0.0f, 4.0f))
-				[
-					SNew(SWarningOrErrorBox)
-					.MessageStyle(EMessageStyle::Warning)
-					.Message(VertexErrorText)
-				]
-			];
 	}
 
 	bool FNeuralMorphModelVizSettingsDetails::IsMorphTargetsEnabled() const
diff --git a/Engine/Plugins/Experimental/MLDeformer/NeuralMorphModel/Source/NeuralMorphModelEditor/Public/NeuralMorphModelDetails.h b/Engine/Plugins/Experimental/MLDeformer/NeuralMorphModel/Source/NeuralMorphModelEditor/Public/NeuralMorphModelDetails.h
--- a/Engine/Plugins/Experimental/MLDeformer/NeuralMorphModel/Source/NeuralMorphModelEditor/Public/NeuralMorphModelDetails.h
+++ b/Engine/Plugins/Experimental/MLDeformer/NeuralMorphModel/Source/NeuralMorphModelEditor/Public/NeuralMorphModelDetails.h
@@ -34,7 +34,6 @@ namespace UE::NeuralMorphModel
 		// FMLDeformerModelDetails overrides.
 		bool UpdateMemberPointers(const TArray<TWeakObjectPtr<UObject>>& Objects) override;
 		void AddTargetMesh() override;
-		void AddBaseMeshErrors() override;
 		void AddAnimSequenceErrors() override;
 		// ~END FMLDeformerModelDetails overrides.
 
diff --git a/Engine/Plugins/Experimental/MLDeformer/VertexDeltaModel/Source/VertexDeltaModel/Private/VertexDeltaModel.cpp b/Engine/Plugins/Experimental/MLDeformer/VertexDeltaModel/Source/VertexDeltaModel/Private/VertexDeltaModel.cpp
--- a/Engine/Plugins/Experimental/MLDeformer/VertexDeltaModel/Source/VertexDeltaModel/Private/VertexDeltaModel.cpp
+++ b/Engine/Plugins/Experimental/MLDeformer/VertexDeltaModel/Source/VertexDeltaModel/Private/VertexDeltaModel.cpp
@@ -38,7 +38,8 @@ UVertexDeltaModel::UVertexDeltaModel(const FObjectInitializer& ObjectInitializer
 		if (MeshMappings.IsEmpty())
 		{
 			TArray<FString> FailedImportedMeshnames;
-			GenerateGeomCacheMeshMappings(SkeletalMesh, GeomCache, MeshMappings, FailedImportedMeshnames);
+			TArray<FString> VertexMisMatchNames;
+			GenerateGeomCacheMeshMappings(SkeletalMesh, GeomCache, MeshMappings, FailedImportedMeshnames, VertexMisMatchNames);
 		}
 
 		UE::MLDeformer::SampleGeomCachePositions(
diff --git a/Engine/Plugins/Experimental/MLDeformer/VertexDeltaModel/Source/VertexDeltaModelEditor/Private/VertexDeltaEditorModel.cpp b/Engine/Plugins/Experimental/MLDeformer/VertexDeltaModel/Source/VertexDeltaModelEditor/Private/VertexDeltaEditorModel.cpp
--- a/Engine/Plugins/Experimental/MLDeformer/VertexDeltaModel/Source/VertexDeltaModelEditor/Private/VertexDeltaEditorModel.cpp
+++ b/Engine/Plugins/Experimental/MLDeformer/VertexDeltaModel/Source/VertexDeltaModelEditor/Private/VertexDeltaEditorModel.cpp
@@ -258,8 +258,7 @@ namespace UE::VertexDeltaModel
 		UVertexDeltaModel* VertexDeltaModel = GetVertexDeltaModel();
 		UGeometryCache* GeomCache = VertexDeltaModel->GetGeometryCache();
 		USkeletalMesh* SkeletalMesh = VertexDeltaModel->GetSkeletalMesh();
-		if (!GetGeomCacheVertexErrorText(SkeletalMesh, GeomCache, FText(), FText()).IsEmpty() ||
-			!GetGeomCacheErrorText(SkeletalMesh, GeomCache).IsEmpty())
+		if (!GetGeomCacheErrorText(SkeletalMesh, GeomCache).IsEmpty())
 		{
 			return;
 		}
diff --git a/Engine/Plugins/Experimental/MLDeformer/VertexDeltaModel/Source/VertexDeltaModelEditor/Private/VertexDeltaModelDetails.cpp b/Engine/Plugins/Experimental/MLDeformer/VertexDeltaModel/Source/VertexDeltaModelEditor/Private/VertexDeltaModelDetails.cpp
--- a/Engine/Plugins/Experimental/MLDeformer/VertexDeltaModel/Source/VertexDeltaModelEditor/Private/VertexDeltaModelDetails.cpp
+++ b/Engine/Plugins/Experimental/MLDeformer/VertexDeltaModel/Source/VertexDeltaModelEditor/Private/VertexDeltaModelDetails.cpp
@@ -65,24 +65,6 @@ namespace UE::VertexDeltaModel
 		SettingsCategoryBuilder->AddProperty(GET_MEMBER_NAME_CHECKED(UVertexDeltaModel, BatchSize));
 		SettingsCategoryBuilder->AddProperty(GET_MEMBER_NAME_CHECKED(UVertexDeltaModel, LearningRate));
 	}
-
-	void FVertexDeltaModelDetails::AddBaseMeshErrors()
-	{
-		// Check if the base mesh matches the target mesh vertex count.
-		FText ErrorText = GetGeomCacheVertexErrorText(Model->SkeletalMesh, VertexModel->GetGeometryCache(), FText::FromString("Base Mesh"), FText::FromString("Target Mesh"));
-		BaseMeshCategoryBuilder->AddCustomRow(FText::FromString("BaseMeshError"))
-			.Visibility(!ErrorText.IsEmpty() ? EVisibility::Visible : EVisibility::Collapsed)
-			.WholeRowContent()
-			[
-				SNew(SBox)
-				.Padding(FMargin(0.0f, 4.0f))
-				[
-					SNew(SWarningOrErrorBox)
-					.MessageStyle(EMessageStyle::Error)
-					.Message(ErrorText)
-				]
-			];
-	}
 	
 	void FVertexDeltaModelDetails::AddAnimSequenceErrors()
 	{
@@ -133,39 +115,7 @@ namespace UE::VertexDeltaModel
 				]
 			];
 
-		// Show meshes that have no matching geometry track list.
-		const FText MeshMappingErrorList = GetGeomCacheMeshMappingErrorText(Model->GetSkeletalMesh(), VertexModel->GetGeometryCache());
-		FString GeomTrackNameList;
-		if (!MeshMappingErrorList.IsEmpty())
-		{
-			UGeometryCache* GeomCache = VertexModel->GetGeometryCache();
-			for (int32 Index = 0; Index < GeomCache->Tracks.Num(); ++Index)
-			{
-				GeomTrackNameList += GeomCache->Tracks[Index]->GetName();
-				if (Index < GeomCache->Tracks.Num() - 1)
-				{
-					GeomTrackNameList += TEXT("\n");
-				}
-			}
-		}
-
-		FText MeshMappingErrorFull = FText::Format(
-			LOCTEXT("MeshMappingError", "No matching GeomCache Tracks names found for meshes:\n{0}\n\nGeomCache Track List:\n{1}"), 
-			MeshMappingErrorList,
-			FText::FromString(GeomTrackNameList));
-
-		TargetMeshCategoryBuilder->AddCustomRow(FText::FromString("MeshMappingError"))
-			.Visibility(!MeshMappingErrorList.IsEmpty() ? EVisibility::Visible : EVisibility::Collapsed)
-			.WholeRowContent()
-			[
-				SNew(SBox)
-				.Padding(FMargin(0.0f, 4.0f))
-				[
-					SNew(SWarningOrErrorBox)
-					.MessageStyle(EMessageStyle::Error)
-					.Message(MeshMappingErrorFull)
-				]
-			];
+		AddGeomCacheMeshMappingWarnings(TargetMeshCategoryBuilder, Model->GetSkeletalMesh(), VertexModel->GetGeometryCache());
 	}
 }	// namespace UE::VertexDeltaModel
 
diff --git a/Engine/Plugins/Experimental/MLDeformer/VertexDeltaModel/Source/VertexDeltaModelEditor/Private/VertexDeltaModelVizSettingsDetails.cpp b/Engine/Plugins/Experimental/MLDeformer/VertexDeltaModel/Source/VertexDeltaModelEditor/Private/VertexDeltaModelVizSettingsDetails.cpp
--- a/Engine/Plugins/Experimental/MLDeformer/VertexDeltaModel/Source/VertexDeltaModelEditor/Private/VertexDeltaModelVizSettingsDetails.cpp
+++ b/Engine/Plugins/Experimental/MLDeformer/VertexDeltaModel/Source/VertexDeltaModelEditor/Private/VertexDeltaModelVizSettingsDetails.cpp
@@ -84,20 +84,6 @@ namespace UE::VertexDeltaModel
 					.Message(GeomErrorText)
 				]
 			];
-
-		const FText VertexErrorText = GetGeomCacheVertexErrorText(Model->GetSkeletalMesh(), VertexDeltaVizSettings->GetTestGroundTruth(), FText::FromString("Base Mesh"), FText::FromString("Ground Truth Mesh"));
-		FDetailWidgetRow& GroundTruthVertexErrorRow = TestAssetsCategory->AddCustomRow(FText::FromString("GroundTruthVertexMismatchError"))
-			.Visibility(!VertexErrorText.IsEmpty() ? EVisibility::Visible : EVisibility::Collapsed)
-			.WholeRowContent()
-			[
-				SNew(SBox)
-				.Padding(FMargin(0.0f, 4.0f))
-				[
-					SNew(SWarningOrErrorBox)
-					.MessageStyle(EMessageStyle::Warning)
-					.Message(VertexErrorText)
-				]
-			];
 	}
 }	//namespace UE::VertexDeltaModel
 
diff --git a/Engine/Plugins/Experimental/MLDeformer/VertexDeltaModel/Source/VertexDeltaModelEditor/Public/VertexDeltaModelDetails.h b/Engine/Plugins/Experimental/MLDeformer/VertexDeltaModel/Source/VertexDeltaModelEditor/Public/VertexDeltaModelDetails.h
--- a/Engine/Plugins/Experimental/MLDeformer/VertexDeltaModel/Source/VertexDeltaModelEditor/Public/VertexDeltaModelDetails.h
+++ b/Engine/Plugins/Experimental/MLDeformer/VertexDeltaModel/Source/VertexDeltaModelEditor/Public/VertexDeltaModelDetails.h
@@ -34,7 +34,6 @@ namespace UE::VertexDeltaModel
 		// FMLDeformerModelDetails overrides.
 		bool UpdateMemberPointers(const TArray<TWeakObjectPtr<UObject>>& Objects) override;
 		void AddTargetMesh() override;
-		void AddBaseMeshErrors() override;
 		void AddAnimSequenceErrors() override;
 		// ~END FMLDeformerModelDetails overrides.
 
diff --git a/Engine/Source/Runtime/Engine/Private/Components/SkinnedMeshComponent.cpp b/Engine/Source/Runtime/Engine/Private/Components/SkinnedMeshComponent.cpp
--- a/Engine/Source/Runtime/Engine/Private/Components/SkinnedMeshComponent.cpp
+++ b/Engine/Source/Runtime/Engine/Private/Components/SkinnedMeshComponent.cpp
@@ -1906,8 +1906,8 @@ USkinnedAsset* USkinnedMeshComponent::GetSkinnedAsset() const
 void USkinnedMeshComponent::SetMeshDeformer(UMeshDeformer* InMeshDeformer)
 {
 	MeshDeformer = InMeshDeformer;
-	MeshDeformerInstanceSettings = MeshDeformer->CreateSettingsInstance(this);
-	MeshDeformerInstance = (MeshDeformer != nullptr) ? MeshDeformer->CreateInstance(this, MeshDeformerInstanceSettings) : nullptr;
+	MeshDeformerInstanceSettings = MeshDeformer ? MeshDeformer->CreateSettingsInstance(this) : nullptr;
+	MeshDeformerInstance = MeshDeformer ? MeshDeformer->CreateInstance(this, MeshDeformerInstanceSettings) : nullptr;
 	MarkRenderDynamicDataDirty();
 }
 
