Added a new frame based animation timeline for the ML Deformer

diff --git a/Engine/Plugins/Experimental/MLDeformer/MLDeformerFramework/Source/MLDeformerFrameworkEditor/MLDeformerFrameworkEditor.Build.cs b/Engine/Plugins/Experimental/MLDeformer/MLDeformerFramework/Source/MLDeformerFrameworkEditor/MLDeformerFrameworkEditor.Build.cs
--- a/Engine/Plugins/Experimental/MLDeformer/MLDeformerFramework/Source/MLDeformerFrameworkEditor/MLDeformerFrameworkEditor.Build.cs
+++ b/Engine/Plugins/Experimental/MLDeformer/MLDeformerFramework/Source/MLDeformerFrameworkEditor/MLDeformerFrameworkEditor.Build.cs
@@ -11,7 +11,8 @@ namespace UnrealBuildTool.Rules
 			PublicDependencyModuleNames.AddRange(
 				new string[]
 				{
-					"Core"
+					"Core",
+					"SequencerWidgets"
 				}
 			);
 
@@ -33,7 +34,10 @@ namespace UnrealBuildTool.Rules
 					"AnimationEditMode",
 					"AnimGraph",
 					"ToolWidgets",
+					"EditorWidgets",
 					"GeometryCache",
+					"SequencerWidgets",
+					"TimeManagement",
 					"NeuralNetworkInference",
 					"RenderCore",
 					"RHI"
diff --git a/Engine/Plugins/Experimental/MLDeformer/MLDeformerFramework/Source/MLDeformerFrameworkEditor/Private/MLDeformerEditorActor.cpp b/Engine/Plugins/Experimental/MLDeformer/MLDeformerFramework/Source/MLDeformerFrameworkEditor/Private/MLDeformerEditorActor.cpp
--- a/Engine/Plugins/Experimental/MLDeformer/MLDeformerFramework/Source/MLDeformerFrameworkEditor/Private/MLDeformerEditorActor.cpp
+++ b/Engine/Plugins/Experimental/MLDeformer/MLDeformerFramework/Source/MLDeformerFrameworkEditor/Private/MLDeformerEditorActor.cpp
@@ -106,6 +106,15 @@ namespace UE::MLDeformer
 		}
 	}
 
+	float FMLDeformerEditorActor::GetPlaySpeed()
+	{
+		if (SkeletalMeshComponent)
+		{
+			return SkeletalMeshComponent->GetPlayRate();
+		}
+		return 0.0f;
+	}
+
 	void FMLDeformerEditorActor::Pause(bool bPaused)
 	{
 		if (SkeletalMeshComponent)
diff --git a/Engine/Plugins/Experimental/MLDeformer/MLDeformerFramework/Source/MLDeformerFrameworkEditor/Private/MLDeformerEditorModel.cpp b/Engine/Plugins/Experimental/MLDeformer/MLDeformerFramework/Source/MLDeformerFrameworkEditor/Private/MLDeformerEditorModel.cpp
--- a/Engine/Plugins/Experimental/MLDeformer/MLDeformerFramework/Source/MLDeformerFrameworkEditor/Private/MLDeformerEditorModel.cpp
+++ b/Engine/Plugins/Experimental/MLDeformer/MLDeformerFramework/Source/MLDeformerFrameworkEditor/Private/MLDeformerEditorModel.cpp
@@ -59,6 +59,10 @@ namespace UE::MLDeformer
 		Sampler = CreateSampler();
 		check(Sampler);
 		Sampler->Init(this);
+
+		ViewRange = TRange<double>(0.0, 100.0);
+		WorkingRange = TRange<double>(0.0, 100.0);
+		PlaybackRange = TRange<double>(0.0, 100.0);
 	}
 
 	void FMLDeformerEditorModel::UpdateEditorInputInfo()
@@ -438,7 +442,7 @@ namespace UE::MLDeformer
 			SkeletalMeshComponent->SetAnimation(TrainingAnimSequence);
 			SkeletalMeshComponent->SetPosition(CurrentPlayTime);
 			SkeletalMeshComponent->SetPlayRate(TestAnimSpeed);
-			SkeletalMeshComponent->Play(false);
+			SkeletalMeshComponent->Play(true);
 		}
 
 		// Update the test base model.
@@ -657,21 +661,36 @@ namespace UE::MLDeformer
 		}
 	}
 
-	void FMLDeformerEditorModel::OnPlayButtonPressed()
+	FMLDeformerEditorActor* FMLDeformerEditorModel::GetVisualizationModeBaseActor() const
 	{
 		UMLDeformerVizSettings* VizSettings = Model->GetVizSettings();
-		if (VizSettings->GetVisualizationMode() != EMLDeformerVizMode::TestData)
+		const bool bTestData = VizSettings->GetVisualizationMode() == EMLDeformerVizMode::TestData;
+		FMLDeformerEditorActor* BaseActor = nullptr;
+		if (bTestData)
 		{
-			return;
+			BaseActor = FindEditorActor(ActorID_Test_Base);
+		}
+		else
+		{
+			BaseActor = FindEditorActor(ActorID_Train_Base);
 		}
-		
-		FMLDeformerEditorActor* BaseTestActor = FindEditorActor(ActorID_Test_Base);
-		const bool bMustPause = (BaseTestActor && BaseTestActor->GetSkeletalMeshComponent()) ? !BaseTestActor->GetSkeletalMeshComponent()->bPauseAnims : false;
+		return BaseActor;
+	}
+
+	void FMLDeformerEditorModel::OnPlayPressed()
+	{
+		UMLDeformerVizSettings* VizSettings = Model->GetVizSettings();
+		const bool bTestData = VizSettings->GetVisualizationMode() == EMLDeformerVizMode::TestData;
+		FMLDeformerEditorActor* BaseActor = GetVisualizationModeBaseActor();
+
+		const bool bMustPause = (BaseActor && BaseActor->GetSkeletalMeshComponent()) ? !BaseActor->GetSkeletalMeshComponent()->bPauseAnims : false;
 
 		for (FMLDeformerEditorActor* EditorActor : EditorActors)
 		{
-			if (EditorActor && EditorActor->IsTestActor())
+			if (EditorActor && ((EditorActor->IsTestActor() && bTestData) || (EditorActor->IsTrainingActor() && !bTestData)))
 			{
+
+				EditorActor->SetPlaySpeed(VizSettings->AnimPlaySpeed);
 				EditorActor->Pause(bMustPause);
 			}
 		}
@@ -679,22 +698,31 @@ namespace UE::MLDeformer
 
 	bool FMLDeformerEditorModel::IsPlayingAnim() const
 	{
-		UMLDeformerVizSettings* VizSettings = Model->GetVizSettings();
-		if (VizSettings->GetVisualizationMode() == EMLDeformerVizMode::TestData)
+		FMLDeformerEditorActor* BaseActor = GetVisualizationModeBaseActor();
+
+		if (BaseActor)
 		{
-			FMLDeformerEditorActor* EditorActor = FindEditorActor(ActorID_Test_Base);
-			if (EditorActor)
+			UDebugSkelMeshComponent* SkeletalMeshComponent = BaseActor->GetSkeletalMeshComponent();
+			if (SkeletalMeshComponent)
 			{
-				UDebugSkelMeshComponent* SkeletalMeshComponent = EditorActor->GetSkeletalMeshComponent();
-				if (SkeletalMeshComponent)
-				{
-					return !SkeletalMeshComponent->bPauseAnims;
-				}
+				return !SkeletalMeshComponent->bPauseAnims;
 			}
 		}
 		return false;
 	}
 
+	bool FMLDeformerEditorModel::IsPlayingForward() const
+	{
+		FMLDeformerEditorActor* BaseActor = GetVisualizationModeBaseActor();
+
+		if (BaseActor)
+		{
+			return BaseActor->GetPlaySpeed() > 0.0f;
+		}
+		return false;
+	}
+
+
 	double FMLDeformerEditorModel::CalcTrainingTimelinePosition() const
 	{
 		FMLDeformerEditorActor* EditorActor = FindEditorActor(ActorID_Train_GroundTruth);
@@ -1666,6 +1694,178 @@ namespace UE::MLDeformer
 	{
 		return FString(FPaths::ProjectIntermediateDir() + TEXT("MLDeformerNetwork.onnx"));
 	}
+
+	const UAnimSequence* FMLDeformerEditorModel::GetAnimSequence() const
+	{
+		if (GetModel())
+		{
+			const UMLDeformerVizSettings* VizSettings = GetModel()->GetVizSettings();
+			if (VizSettings->GetVisualizationMode() == EMLDeformerVizMode::TrainingData)
+			{
+				return GetModel()->GetAnimSequence();
+			}
+			else if (VizSettings->GetVisualizationMode() == EMLDeformerVizMode::TestData)
+			{
+				return GetModel()->GetVizSettings()->GetTestAnimSequence();
+			}
+		}
+		return nullptr;
+	}
+
+	void FMLDeformerEditorModel::UpdateRanges()
+	{
+		ViewRange = TRange<double>(0.0, 100.0);
+		WorkingRange = TRange<double>(0.0, 100.0);
+		PlaybackRange = TRange<double>(0.0, 100.0);
+		const UAnimSequence* AnimSeq = GetAnimSequence();
+		const double Duration = AnimSeq ? AnimSeq->GetPlayLength() : 0.0;
+		ViewRange = TRange<double>(0.0, Duration);
+		WorkingRange = ViewRange;
+		PlaybackRange = WorkingRange;
+		ClampCurrentTrainingFrameIndex();
+	}
+
+	void FMLDeformerEditorModel::HandleModelChanged()
+	{
+		UpdateRanges();
+	}
+
+	void FMLDeformerEditorModel::HandleVizModeChanged(EMLDeformerVizMode Mode)
+	{
+		UpdateRanges();
+	}
+
+	/** Get the current view range */
+	TRange<double> FMLDeformerEditorModel::GetViewRange() const
+	{
+		return ViewRange;
+	}
+
+	void FMLDeformerEditorModel::SetViewRange(TRange<double> InRange)
+	{
+		ViewRange = InRange;
+
+		if (WorkingRange.HasLowerBound() && WorkingRange.HasUpperBound())
+		{
+			WorkingRange = TRange<double>::Hull(WorkingRange, ViewRange);
+		}
+		else
+		{
+			WorkingRange = ViewRange;
+		}
+	}
+
+	TRange<double> FMLDeformerEditorModel::GetWorkingRange() const
+	{
+		return WorkingRange;
+	}
+
+	TRange<FFrameNumber> FMLDeformerEditorModel::GetPlaybackRange() const
+	{
+		const int32 Resolution = GetTickResolution();
+		return TRange<FFrameNumber>(FFrameNumber(FMath::RoundToInt32(PlaybackRange.GetLowerBoundValue() * (double)Resolution)), FFrameNumber(FMath::RoundToInt32(PlaybackRange.GetUpperBoundValue() * (double)Resolution)));
+
+	}
+
+	int32 FMLDeformerEditorModel::GetTicksPerFrame() const
+	{
+		// Default to millisecond resolution
+		const int32 TimelineScrubSnapValue = 1000;
+		return TimelineScrubSnapValue;
+	}
+
+	FFrameNumber FMLDeformerEditorModel::GetTickResScrubPosition() const
+	{
+		double Resolution = GetTickResolution();
+		const UMLDeformerVizSettings* VizSettings = GetModel()->GetVizSettings();
+		double TrainingFrameNumber = 0.0;
+		if (VizSettings->GetVisualizationMode() == EMLDeformerVizMode::TrainingData)
+		{
+			TrainingFrameNumber = GetTrainingFrameAtTime(CalcTimelinePosition());
+		}
+		else if (VizSettings->GetVisualizationMode() == EMLDeformerVizMode::TestData)
+		{
+			TrainingFrameNumber = GetTestFrameAtTime(CalcTimelinePosition());
+		}
+		FFrameNumber FrameNumber = (int32)(TrainingFrameNumber * Resolution / (double)GetFrameRate());
+		return FrameNumber;
+	}
+
+	float FMLDeformerEditorModel::GetScrubTime() const
+	{
+		return (float)ScrubPosition.AsDecimal() / (float)GetFrameRate();
+	}
+
+	void FMLDeformerEditorModel::SetScrubPosition(FFrameTime NewScrubPostion)
+	{
+		double FrameNumber = NewScrubPostion.AsDecimal() * (double)GetFrameRate() / (double)GetTickResolution();
+		ScrubPosition.FrameNumber = (int32)FrameNumber;
+	
+		SetTrainingFrame(ScrubPosition.FrameNumber.Value);
+		SetTestFrame(ScrubPosition.FrameNumber.Value);
+	}
+
+	void FMLDeformerEditorModel::SetScrubPosition(FFrameNumber NewScrubPostion)
+	{
+		double Resolution = GetTickResolution();
+		double ClampedValue = FMath::Clamp((double)NewScrubPostion.Value, PlaybackRange.GetLowerBoundValue() * Resolution, PlaybackRange.GetUpperBoundValue() * Resolution);
+		int32 FrameValue = ClampedValue * (double)GetFrameRate() / Resolution;
+		SetTrainingFrame(FrameValue);
+		SetTestFrame(FrameValue);
+	}
+
+	void  FMLDeformerEditorModel::SetDisplayFrames(bool bInDisplayFrames)
+	{
+		bDisplayFrames = bInDisplayFrames;
+	}
+
+	bool  FMLDeformerEditorModel::IsDisplayingFrames() const
+	{
+		return bDisplayFrames;
+	}
+
+	void FMLDeformerEditorModel::HandleViewRangeChanged(TRange<double> InRange)
+	{
+		SetViewRange(InRange);
+	}
+
+	void FMLDeformerEditorModel::HandleWorkingRangeChanged(TRange<double> InRange)
+	{
+		WorkingRange = InRange;
+	}
+
+	double FMLDeformerEditorModel::CalcTimelinePosition() const
+	{
+		const UMLDeformerVizSettings* VizSettings = GetModel()->GetVizSettings();
+		if (VizSettings->GetVisualizationMode() == EMLDeformerVizMode::TrainingData)
+		{
+			return CalcTrainingTimelinePosition();
+		}
+		else if (VizSettings->GetVisualizationMode() == EMLDeformerVizMode::TestData)
+		{
+			return CalcTestTimelinePosition();
+		}
+		return 0.0;
+	}
+
+	double FMLDeformerEditorModel::GetFrameRate() const
+	{
+		const UAnimSequence* AnimSequence = GetAnimSequence();
+		if (AnimSequence)
+		{
+			return AnimSequence->GetSamplingFrameRate().AsDecimal();
+		}
+		else
+		{
+			return 30.0;
+		}
+	}
+
+	int32 FMLDeformerEditorModel::GetTickResolution() const
+	{
+		return FMath::RoundToInt32((double)GetTicksPerFrame() * GetFrameRate());
+	}
+
 }	// namespace UE::MLDeformer
 
 #undef LOCTEXT_NAMESPACE
diff --git a/Engine/Plugins/Experimental/MLDeformer/MLDeformerFramework/Source/MLDeformerFrameworkEditor/Private/MLDeformerEditorToolkit.cpp b/Engine/Plugins/Experimental/MLDeformer/MLDeformerFramework/Source/MLDeformerFrameworkEditor/Private/MLDeformerEditorToolkit.cpp
--- a/Engine/Plugins/Experimental/MLDeformer/MLDeformerFramework/Source/MLDeformerFrameworkEditor/Private/MLDeformerEditorToolkit.cpp
+++ b/Engine/Plugins/Experimental/MLDeformer/MLDeformerFramework/Source/MLDeformerFrameworkEditor/Private/MLDeformerEditorToolkit.cpp
@@ -22,6 +22,7 @@
 #include "Preferences/PersonaOptions.h"
 #include "UObject/Object.h"
 #include "SSimpleTimeSlider.h"
+#include "SMLDeformerTimeline.h"
 #include "Framework/MultiBox/MultiBoxBuilder.h"
 #include "Framework/Notifications/NotificationManager.h"
 #include "Misc/MessageDialog.h"
@@ -305,6 +306,12 @@ namespace UE::MLDeformer
 			VizSettingsDetailsView->ForceRefresh();
 		}
 
+		if (TimeSlider)
+		{
+			TWeakPtr<FMLDeformerEditorModel> WeakModel = ActiveModel; 
+			TimeSlider->SetModel(WeakModel);
+		}
+
 		GEngine->ForceGarbageCollection(true);
 		DeformerAsset->Modify();
 	}
@@ -686,9 +693,9 @@ namespace UE::MLDeformer
 
 				ActiveModel->SetHeatMapMaterialEnabled(VizSettings->GetShowHeatMap());
 				ActiveModel->UpdateDeformerGraph();
+				ActiveModel->HandleVizModeChanged(VizSettings->GetVisualizationMode());
 			}
 		}
-
 		ZoomOnActors();
 	}
 
@@ -747,12 +754,12 @@ namespace UE::MLDeformer
 		return VizSettingsDetailsView.Get();
 	}
 
-	void FMLDeformerEditorToolkit::SetTimeSlider(TSharedPtr<SSimpleTimeSlider> InTimeSlider)
+	void FMLDeformerEditorToolkit::SetTimeSlider(TSharedPtr<SMLDeformerTimeline> InTimeSlider)
 	{
 		TimeSlider = InTimeSlider;
 	}
 
-	SSimpleTimeSlider* FMLDeformerEditorToolkit::GetTimeSlider() const
+	SMLDeformerTimeline* FMLDeformerEditorToolkit::GetTimeSlider() const
 	{
 		return TimeSlider.Get();
 	}
@@ -811,12 +818,13 @@ namespace UE::MLDeformer
 
 	void FMLDeformerEditorToolkit::SetTimeSliderRange(double StartTime, double EndTime)
 	{
-		SSimpleTimeSlider* Slider = GetTimeSlider();
-		if (Slider)
+		TRange<double> ViewRange(StartTime, EndTime); 
+		if (ActiveModel.IsValid())
 		{
-			Slider->SetTimeRange(StartTime, EndTime);
-			Slider->SetClampRange(StartTime, EndTime);
+			ActiveModel->SetViewRange(ViewRange);
+
 		}
+
 	}
 }	// namespace UE::MLDeformer
 
diff --git a/Engine/Plugins/Experimental/MLDeformer/MLDeformerFramework/Source/MLDeformerFrameworkEditor/Private/MLDeformerTimelineTabSummoner.cpp b/Engine/Plugins/Experimental/MLDeformer/MLDeformerFramework/Source/MLDeformerFrameworkEditor/Private/MLDeformerTimelineTabSummoner.cpp
--- a/Engine/Plugins/Experimental/MLDeformer/MLDeformerFramework/Source/MLDeformerFrameworkEditor/Private/MLDeformerTimelineTabSummoner.cpp
+++ b/Engine/Plugins/Experimental/MLDeformer/MLDeformerFramework/Source/MLDeformerFrameworkEditor/Private/MLDeformerTimelineTabSummoner.cpp
@@ -8,9 +8,11 @@
 
 #include "IDocumentation.h"
 #include "SSimpleTimeSlider.h"
+
 #include "Widgets/SBoxPanel.h"
 #include "Widgets/Input/SButton.h"
 #include "Widgets/Images/SImage.h"
+#include "SMLDeformerTimeline.h"
 
 #define LOCTEXT_NAMESPACE "MLDeformerTimelineTabSummoner"
 
@@ -40,74 +42,15 @@ namespace UE::MLDeformer
 
 	TSharedRef<SWidget> FMLDeformerTimelineTabSummoner::CreateTabBody(const FWorkflowTabSpawnInfo& Info) const
 	{
-		// Create and setup the time slider widget.
-		TSharedRef<SSimpleTimeSlider> TimeSlider = SNew(SSimpleTimeSlider)
-			.ScrubPosition_Raw(Editor, &FMLDeformerEditorToolkit::CalcTimelinePosition)
-			.OnScrubPositionChanged_Raw(Editor, &FMLDeformerEditorToolkit::OnTimeSliderScrubPositionChanged);
+		TSharedRef<SMLDeformerTimeline> TimeSlider = SNew(SMLDeformerTimeline, Editor);
 
 		Editor->SetTimeSlider(TimeSlider);
-
 		// Add the time slider.
 		TSharedRef<SHorizontalBox> Content = SNew(SHorizontalBox)
 			+SHorizontalBox::Slot()
 			.VAlign(EVerticalAlignment::VAlign_Top)
 			[
 				TimeSlider
-			]
-			+SHorizontalBox::Slot()
-			.AutoWidth()
-			.HAlign(EHorizontalAlignment::HAlign_Right)
-			.VAlign(EVerticalAlignment::VAlign_Top)
-			[
-				SNew(SBox)
-				.WidthOverride(25)
-				.HeightOverride(25)
-				.Padding(FMargin(0.0f))
-				.Visibility_Lambda
-				(
-					[this]()
-					{
-						using namespace UE::MLDeformer;
-						const FMLDeformerEditorModel* ActiveModel = Editor->GetActiveModel();
-						if (ActiveModel)
-						{
-							return (Editor->GetActiveModel()->GetModel()->GetVizSettings()->GetVisualizationMode() == EMLDeformerVizMode::TestData) ? EVisibility::Visible : EVisibility::Collapsed; 
-						}
-						return EVisibility::Collapsed;
-					} 
-				)
-				[
-					SNew(SButton)
-					.ButtonStyle(FAppStyle::Get(), "SimpleButton")
-					.HAlign(EHorizontalAlignment::HAlign_Center)
-					.VAlign(EVerticalAlignment::VAlign_Center)
-					.ToolTipText(LOCTEXT("PlayButtonToolTip", "Play or pause the test animation sequence"))
-					.ContentPadding(FMargin(0.0f))
-					.OnClicked_Lambda([this]()
-					{ 
-							using namespace UE::MLDeformer;
-							FMLDeformerEditorModel* EditorModel = Editor->GetActiveModel();
-							if (EditorModel)
-							{ 
-								EditorModel->OnPlayButtonPressed(); 
-							} 
-							return FReply::Handled(); 
-					})
-					[
-						SNew(SImage)
-						.Image_Lambda
-						(
-							[this]()
-							{
-								using namespace UE::MLDeformer;
-								const FMLDeformerEditorModel* EditorModel = Editor->GetActiveModel();
-								return (EditorModel != nullptr && EditorModel->IsPlayingAnim())
-									? FMLDeformerEditorStyle::Get().GetBrush("MLDeformer.Timeline.PauseIcon")
-									: FMLDeformerEditorStyle::Get().GetBrush("MLDeformer.Timeline.PlayIcon");
-							}
-						)
-					]
-				]
 			];
 
 		return Content;
diff --git a/Engine/Plugins/Experimental/MLDeformer/MLDeformerFramework/Source/MLDeformerFrameworkEditor/Private/SMLDeformerTimeline.cpp b/Engine/Plugins/Experimental/MLDeformer/MLDeformerFramework/Source/MLDeformerFrameworkEditor/Private/SMLDeformerTimeline.cpp
new file mode 100644
--- /dev/null
+++ b/Engine/Plugins/Experimental/MLDeformer/MLDeformerFramework/Source/MLDeformerFrameworkEditor/Private/SMLDeformerTimeline.cpp
@@ -0,0 +1,1940 @@
+// Copyright Epic Games, Inc. All Rights Reserved.
+
+#include "SMLDeformerTimeline.h"
+#include "Styling/ISlateStyle.h"
+#include "Widgets/SWidget.h"
+#include "Widgets/Layout/SSplitter.h"
+#include "ISequencerWidgetsModule.h"
+#include "FrameNumberNumericInterface.h"
+#include "Widgets/Input/SSearchBox.h"
+#include "Widgets/Layout/SScrollBorder.h"
+#include "Styling/AppStyle.h"
+#include "Fonts/FontMeasure.h"
+#include "Widgets/Layout/SGridPanel.h"
+#include "Widgets/Layout/SSpacer.h"
+#include "Framework/Application/SlateApplication.h"
+#include "Widgets/Input/SButton.h"
+#include "Modules/ModuleManager.h"
+#include "Preferences/PersonaOptions.h"
+#include "IPersonaPreviewScene.h"
+#include "Animation/DebugSkelMeshComponent.h"
+#include "AnimPreviewInstance.h"
+#include "EditorWidgetsModule.h"
+#include "Framework/MultiBox/MultiBoxBuilder.h"
+#include "ScopedTransaction.h"
+#include "Widgets/Input/STextEntryPopup.h"
+#include "Widgets/Input/SSpinBox.h"
+#include "Animation/AnimSequence.h"
+#include "Animation/AnimData/AnimDataModel.h"
+#include "Runtime/Engine/Classes/Animation/AnimSequenceHelpers.h"
+#include "Widgets/Layout/SScrollBox.h"
+#include "MLDeformerEditorModel.h"
+#include "MLDeformerModel.h"
+
+#define LOCTEXT_NAMESPACE "SMLDeformerTimeline"
+
+namespace UE::MLDeformer
+{
+
+	struct FPaintSectionAreaViewArgs
+	{
+		FPaintSectionAreaViewArgs()
+			: bDisplayTickLines(false), bDisplayScrubPosition(false)
+		{}
+
+		/** Whether to display tick lines */
+		bool bDisplayTickLines;
+		/** Whether to display the scrub position */
+		bool bDisplayScrubPosition;
+		/** Optional Paint args for the playback range*/
+		TOptional<FPaintPlaybackRangeArgs> PlaybackRangeArgs;
+	};
+
+	/**
+	 * A time slider controller for the anim timeline
+	 */
+	class FMLTimeSliderController : public ITimeSliderController
+	{
+	public:
+		FMLTimeSliderController(const FTimeSliderArgs& InArgs, TWeakPtr<FMLDeformerEditorModel> InWeakModel, TWeakPtr<SMLDeformerTimeline> InWeakTimeline);
+
+		/**
+		* Determines the optimal spacing between tick marks in the slider for a given pixel density
+		* Increments until a minimum amount of slate units specified by MinTick is reached
+		*
+		* @param InPixelsPerInput	The density of pixels between each input
+		* @param MinTick			The minimum slate units per tick allowed
+		* @param MinTickSpacing	The minimum tick spacing in time units allowed
+		* @return the optimal spacing in time units
+		*/
+		float DetermineOptimalSpacing(float InPixelsPerInput, uint32 MinTick, float MinTickSpacing) const;
+		void SetModel(TWeakPtr<FMLDeformerEditorModel> InModel);
+		virtual int32 OnPaintTimeSlider(bool bMirrorLabels, const FGeometry& AllottedGeometry, const FSlateRect& MyCullingRect, FSlateWindowElementList& OutDrawElements, int32 LayerId, const FWidgetStyle& InWidgetStyle, bool bParentEnabled) const override;
+		virtual int32 OnPaintViewArea(const FGeometry& AllottedGeometry, const FSlateRect& MyCullingRect, FSlateWindowElementList& OutDrawElements, int32 LayerId, bool bEnabled, const FPaintViewAreaArgs& Args) const override;
+		virtual FReply OnMouseButtonDown(SWidget& WidgetOwner, const FGeometry& MyGeometry, const FPointerEvent& MouseEvent) override;
+		virtual FReply OnMouseButtonUp(SWidget& WidgetOwner, const FGeometry& MyGeometry, const FPointerEvent& MouseEvent) override;
+		virtual FReply OnMouseMove(SWidget& WidgetOwner, const FGeometry& MyGeometry, const FPointerEvent& MouseEvent) override;
+		virtual FReply OnMouseWheel(SWidget& WidgetOwner, const FGeometry& MyGeometry, const FPointerEvent& MouseEvent) override;
+		virtual FCursorReply OnCursorQuery(TSharedRef<const SWidget> WidgetOwner, const FGeometry& MyGeometry, const FPointerEvent& CursorEvent) const override;
+		virtual void SetViewRange(double NewRangeMin, double NewRangeMax, EViewRangeInterpolation Interpolation) override;
+		virtual void SetClampRange(double NewRangeMin, double NewRangeMax) override;
+		virtual void SetPlayRange(FFrameNumber RangeStart, int32 RangeDuration) override;
+		virtual FFrameRate GetDisplayRate() const override { return TimeSliderArgs.DisplayRate.Get(); }
+		virtual FFrameRate GetTickResolution() const override { return TimeSliderArgs.TickResolution.Get(); }
+		virtual FAnimatedRange GetViewRange() const override { return TimeSliderArgs.ViewRange.Get(); }
+		virtual FAnimatedRange GetClampRange() const override { return TimeSliderArgs.ClampRange.Get(); }
+		virtual TRange<FFrameNumber> GetPlayRange() const override { return TimeSliderArgs.PlaybackRange.Get(TRange<FFrameNumber>()); }
+		virtual FFrameTime GetScrubPosition() const { return TimeSliderArgs.ScrubPosition.Get(); }
+		virtual void SetScrubPosition(FFrameTime InTime, bool bEvaluate) { TimeSliderArgs.OnScrubPositionChanged.ExecuteIfBound(InTime, false, bEvaluate); }
+
+		/**
+		 * Clamp the given range to the clamp range
+		 *
+		 * @param NewRangeMin		The new lower bound of the range
+		 * @param NewRangeMax		The new upper bound of the range
+		 */
+		void ClampViewRange(double& NewRangeMin, double& NewRangeMax);
+
+		/**
+		 * Zoom the range by a given delta.
+		 *
+		 * @param InDelta		The total amount to zoom by (+ve = zoom out, -ve = zoom in)
+		 * @param ZoomBias		Bias to apply to lower/upper extents of the range. (0 = lower, 0.5 = equal, 1 = upper)
+		 */
+		bool ZoomByDelta(float InDelta, float ZoomBias = 0.5f);
+
+		/**
+		 * Pan the range by a given delta
+		 *
+		 * @param InDelta		The total amount to pan by (+ve = pan forwards in time, -ve = pan backwards in time)
+		 */
+		void PanByDelta(float InDelta);
+
+		/** Determine frame time from a mouse position */
+		FFrameTime GetFrameTimeFromMouse(const FGeometry& Geometry, FVector2D ScreenSpacePosition) const;
+
+	private:
+		// forward declared as class members to prevent name collision with similar types defined in other units
+		struct FDrawTickArgs;
+		struct FScrubRangeToScreen;
+
+		/**
+		 * Call this method when the user's interaction has changed the scrub position
+		 *
+		 * @param NewValue				Value resulting from the user's interaction
+		 * @param bIsScrubbing			True if done via scrubbing, false if just releasing scrubbing
+		 */
+		void CommitScrubPosition(FFrameTime NewValue, bool bIsScrubbing);
+
+		/**
+		 * Draw time tick marks
+		 *
+		 * @param OutDrawElements	List to add draw elements to
+		 * @param ViewRange			The currently visible time range in seconds
+		 * @param RangeToScreen		Time range to screen space converter
+		 * @param InArgs			Parameters for drawing the tick lines
+		 */
+		void DrawTicks(FSlateWindowElementList& OutDrawElements, const TRange<double>& ViewRange, const FScrubRangeToScreen& RangeToScreen, FDrawTickArgs& InArgs) const;
+
+		/**
+		 * Draw the selection range.
+		 *
+		 * @return The new layer ID.
+		 */
+		int32 DrawSelectionRange(const FGeometry& AllottedGeometry, const FSlateRect& MyCullingRect, FSlateWindowElementList& OutDrawElements, int32 LayerId, const FScrubRangeToScreen& RangeToScreen, const FPaintPlaybackRangeArgs& Args) const;
+
+		/**
+		 * Draw the playback range.
+		 *
+		 * @return the new layer ID
+		 */
+		int32 DrawPlaybackRange(const FGeometry& AllottedGeometry, const FSlateRect& MyCullingRect, FSlateWindowElementList& OutDrawElements, int32 LayerId, const FScrubRangeToScreen& RangeToScreen, const FPaintPlaybackRangeArgs& Args) const;
+
+	private:
+
+		/**
+		 * Hit test the lower bound of a range
+		 */
+		bool HitTestRangeStart(const FScrubRangeToScreen& RangeToScreen, const TRange<double>& Range, float HitPixel) const;
+
+		/**
+		 * Hit test the upper bound of a range
+		 */
+		bool HitTestRangeEnd(const FScrubRangeToScreen& RangeToScreen, const TRange<double>& Range, float HitPixel) const;
+
+		void SetPlaybackRangeStart(FFrameNumber NewStart);
+		void SetPlaybackRangeEnd(FFrameNumber NewEnd);
+
+		void SetSelectionRangeStart(FFrameNumber NewStart);
+		void SetSelectionRangeEnd(FFrameNumber NewEnd);
+
+		FFrameTime ComputeFrameTimeFromMouse(const FGeometry& Geometry, FVector2D ScreenSpacePosition, FScrubRangeToScreen RangeToScreen, bool CheckSnapping = true) const;
+	private:
+
+		struct FScrubPixelRange
+		{
+			TRange<float> Range;
+			TRange<float> HandleRange;
+			bool bClamped;
+		};
+
+		FScrubPixelRange GetHitTestScrubberPixelRange(FFrameTime ScrubTime, const FScrubRangeToScreen& RangeToScreen) const;
+
+		FScrubPixelRange GetScrubberPixelRange(FFrameTime ScrubTime, const FScrubRangeToScreen& RangeToScreen) const;
+		FScrubPixelRange GetScrubberPixelRange(FFrameTime ScrubTime, FFrameRate Resolution, FFrameRate PlayRate, const FScrubRangeToScreen& RangeToScreen, float DilationPixels = 0.f) const;
+
+	private:
+
+		/** Pointer back to the model object */
+		TWeakPtr<FMLDeformerEditorModel> WeakModel;
+
+		/** Pointer back to the timeline */
+		TWeakPtr<SMLDeformerTimeline> WeakTimeline;
+
+		FTimeSliderArgs TimeSliderArgs;
+
+		/** Brush for drawing the fill area on the scrubber */
+		const FSlateBrush* ScrubFillBrush;
+
+		/** Brush for drawing an upwards facing scrub handles */
+		const FSlateBrush* ScrubHandleUpBrush;
+
+		/** Brush for drawing a downwards facing scrub handle */
+		const FSlateBrush* ScrubHandleDownBrush;
+
+		/** Brush for drawing an editable time */
+		const FSlateBrush* EditableTimeBrush;
+
+		/** Total mouse delta during dragging **/
+		float DistanceDragged;
+
+		/** If we are dragging a scrubber or dragging to set the time range */
+		enum DragType
+		{
+			DRAG_SCRUBBING_TIME,
+			DRAG_PLAYBACK_START,
+			DRAG_PLAYBACK_END,
+			DRAG_SELECTION_START,
+			DRAG_SELECTION_END,
+			DRAG_TIME,
+			DRAG_NONE
+		};
+
+		DragType MouseDragType;
+
+		/** If we are currently panning the panel */
+		bool bPanning;
+
+		/** Index of the current dragged time */
+		int32 DraggedTimeIndex;
+
+		/** Mouse down position range */
+		FVector2D MouseDownPosition[2];
+
+		/** Geometry on mouse down */
+		FGeometry MouseDownGeometry;
+
+		/** Range stack */
+		TArray<TRange<double>> ViewRangeStack;
+
+	};
+
+	/** Utility struct for converting between scrub range space and local/absolute screen space */
+	struct FMLTimeSliderController::FScrubRangeToScreen
+	{
+		double ViewStart;
+		float  PixelsPerInput;
+
+		FScrubRangeToScreen(const TRange<double>& InViewInput, const FVector2D& InWidgetSize)
+		{
+			float ViewInputRange = InViewInput.Size<double>();
+
+			ViewStart = InViewInput.GetLowerBoundValue();
+			PixelsPerInput = ViewInputRange > 0 ? (InWidgetSize.X / ViewInputRange) : 0;
+		}
+
+		/** Local Widget Space -> Curve Input domain. */
+		double LocalXToInput(float ScreenX) const
+		{
+			return PixelsPerInput > 0 ? (ScreenX / PixelsPerInput) + ViewStart : ViewStart;
+		}
+
+		/** Curve Input domain -> local Widget Space */
+		float InputToLocalX(double Input) const
+		{
+			return (Input - ViewStart) * PixelsPerInput;
+		}
+	};
+
+
+	/**
+	 * Gets the the next spacing value in the series
+	 * to determine a good spacing value
+	 * E.g, .001,.005,.010,.050,.100,.500,1.000,etc
+	 */
+	static float GetNextSpacing(uint32 CurrentStep)
+	{
+		if (CurrentStep & 0x01)
+		{
+			// Odd numbers
+			return FMath::Pow(10.f, 0.5f * ((float)(CurrentStep - 1)) + 1.f);
+		}
+		else
+		{
+			// Even numbers
+			return 0.5f * FMath::Pow(10.f, 0.5f * ((float)(CurrentStep)) + 1.f);
+		}
+	}
+
+	FMLTimeSliderController::FMLTimeSliderController(const FTimeSliderArgs& InArgs, TWeakPtr<FMLDeformerEditorModel> InWeakModel, TWeakPtr<SMLDeformerTimeline> InWeakTimeline)
+		: WeakModel(InWeakModel)
+		, WeakTimeline(InWeakTimeline)
+		, TimeSliderArgs(InArgs)
+		, DistanceDragged(0.0f)
+		, MouseDragType(DRAG_NONE)
+		, bPanning(false)
+		, DraggedTimeIndex(INDEX_NONE)
+	{
+		ScrubFillBrush = FAppStyle::GetBrush(TEXT("Sequencer.Timeline.ScrubFill"));
+		ScrubHandleUpBrush = FAppStyle::GetBrush(TEXT("Sequencer.Timeline.VanillaScrubHandleUp"));
+		ScrubHandleDownBrush = FAppStyle::GetBrush(TEXT("Sequencer.Timeline.VanillaScrubHandleDown"));
+		EditableTimeBrush = FAppStyle::GetBrush(TEXT("AnimTimeline.SectionMarker"));
+	}
+
+	void FMLTimeSliderController::SetModel(TWeakPtr<FMLDeformerEditorModel> InModel)
+	{
+		WeakModel = InModel;
+	}
+
+	FFrameTime FMLTimeSliderController::ComputeFrameTimeFromMouse(const FGeometry& Geometry, FVector2D ScreenSpacePosition, FScrubRangeToScreen RangeToScreen, bool CheckSnapping) const
+	{
+		FVector2D CursorPos = Geometry.AbsoluteToLocal(ScreenSpacePosition);
+		double    MouseValue = RangeToScreen.LocalXToInput(CursorPos.X);
+
+		return MouseValue * GetTickResolution();
+	}
+
+	FMLTimeSliderController::FScrubPixelRange FMLTimeSliderController::GetHitTestScrubberPixelRange(FFrameTime ScrubTime, const FScrubRangeToScreen& RangeToScreen) const
+	{
+		static const float DragToleranceSlateUnits = 2.f, MouseTolerance = 2.f;
+		return GetScrubberPixelRange(ScrubTime, GetTickResolution(), GetDisplayRate(), RangeToScreen, DragToleranceSlateUnits + MouseTolerance);
+	}
+
+	FMLTimeSliderController::FScrubPixelRange FMLTimeSliderController::GetScrubberPixelRange(FFrameTime ScrubTime, const FScrubRangeToScreen& RangeToScreen) const
+	{
+		return GetScrubberPixelRange(ScrubTime, GetTickResolution(), GetDisplayRate(), RangeToScreen);
+	}
+
+	FMLTimeSliderController::FScrubPixelRange FMLTimeSliderController::GetScrubberPixelRange(FFrameTime ScrubTime, FFrameRate Resolution, FFrameRate PlayRate, const FScrubRangeToScreen& RangeToScreen, float DilationPixels) const
+	{
+		const FFrameNumber Frame = ScrubTime.FloorToFrame();// ConvertFrameTime(ScrubTime, Resolution, PlayRate).FloorToFrame();
+
+		float StartPixel = RangeToScreen.InputToLocalX(Frame / Resolution);
+		float EndPixel = RangeToScreen.InputToLocalX((Frame + 1) / Resolution);
+
+		{
+			float RoundedStartPixel = FMath::RoundToInt(StartPixel);
+			EndPixel -= (StartPixel - RoundedStartPixel);
+
+			StartPixel = RoundedStartPixel;
+			EndPixel = FMath::Max(EndPixel, StartPixel + 1);
+		}
+
+		FScrubPixelRange Range;
+
+		float MinScrubSize = 14.f;
+		Range.bClamped = EndPixel - StartPixel < MinScrubSize;
+		Range.Range = TRange<float>(StartPixel, EndPixel);
+		if (Range.bClamped)
+		{
+			Range.HandleRange = TRange<float>(
+				(StartPixel + EndPixel - MinScrubSize) * .5f,
+				(StartPixel + EndPixel + MinScrubSize) * .5f);
+		}
+		else
+		{
+			Range.HandleRange = Range.Range;
+		}
+
+		return Range;
+	}
+
+	struct FMLTimeSliderController::FDrawTickArgs
+	{
+		/** Geometry of the area */
+		FGeometry AllottedGeometry;
+		/** Culling rect of the area */
+		FSlateRect CullingRect;
+		/** Color of each tick */
+		FLinearColor TickColor;
+		/** Offset in Y where to start the tick */
+		float TickOffset;
+		/** Height in of major ticks */
+		float MajorTickHeight;
+		/** Start layer for elements */
+		int32 StartLayer;
+		/** Draw effects to apply */
+		ESlateDrawEffect DrawEffects;
+		/** Whether or not to only draw major ticks */
+		bool bOnlyDrawMajorTicks;
+		/** Whether or not to mirror labels */
+		bool bMirrorLabels;
+
+	};
+
+	void FMLTimeSliderController::DrawTicks(FSlateWindowElementList& OutDrawElements, const TRange<double>& ViewRange, const FScrubRangeToScreen& RangeToScreen, FDrawTickArgs& InArgs) const
+	{
+		TSharedPtr<SMLDeformerTimeline> Timeline = WeakTimeline.Pin();
+		if (!Timeline.IsValid())
+		{
+			return;
+		}
+
+		if (!FMath::IsFinite(ViewRange.GetLowerBoundValue()) || !FMath::IsFinite(ViewRange.GetUpperBoundValue()))
+		{
+			return;
+		}
+
+		FFrameRate     FrameResolution = GetTickResolution();
+		FPaintGeometry PaintGeometry = InArgs.AllottedGeometry.ToPaintGeometry();
+		FSlateFontInfo SmallLayoutFont = FCoreStyle::GetDefaultFontStyle("Regular", 8);
+
+		double MajorGridStep = 0.0;
+		int32  MinorDivisions = 0;
+		if (!Timeline->GetGridMetrics(InArgs.AllottedGeometry.Size.X, MajorGridStep, MinorDivisions))
+		{
+			return;
+		}
+
+		if (InArgs.bOnlyDrawMajorTicks)
+		{
+			MinorDivisions = 0;
+		}
+
+		TArray<FVector2D> LinePoints;
+		LinePoints.SetNumUninitialized(2);
+
+		const bool bAntiAliasLines = false;
+
+		const double FirstMajorLine = FMath::FloorToDouble(ViewRange.GetLowerBoundValue() / MajorGridStep) * MajorGridStep;
+		const double LastMajorLine = FMath::CeilToDouble(ViewRange.GetUpperBoundValue() / MajorGridStep) * MajorGridStep;
+
+		for (double CurrentMajorLine = FirstMajorLine; CurrentMajorLine < LastMajorLine; CurrentMajorLine += MajorGridStep)
+		{
+			float MajorLinePx = RangeToScreen.InputToLocalX(CurrentMajorLine);
+
+			LinePoints[0] = FVector2D(MajorLinePx, InArgs.TickOffset);
+			LinePoints[1] = FVector2D(MajorLinePx, InArgs.TickOffset + InArgs.MajorTickHeight);
+
+			// Draw each tick mark
+			FSlateDrawElement::MakeLines(
+				OutDrawElements,
+				InArgs.StartLayer,
+				PaintGeometry,
+				LinePoints,
+				InArgs.DrawEffects,
+				InArgs.TickColor,
+				bAntiAliasLines
+			);
+
+			if (!InArgs.bOnlyDrawMajorTicks)
+			{
+				FString FrameString = TimeSliderArgs.NumericTypeInterface->ToString((CurrentMajorLine * FrameResolution).RoundToFrame().Value);
+
+				// Space the text between the tick mark but slightly above
+				FVector2D TextOffset(MajorLinePx + 5.f, InArgs.bMirrorLabels ? 3.f : FMath::Abs(InArgs.AllottedGeometry.Size.Y - (InArgs.MajorTickHeight + 3.f)));
+				FSlateDrawElement::MakeText(
+					OutDrawElements,
+					InArgs.StartLayer + 1,
+					InArgs.AllottedGeometry.ToPaintGeometry(TextOffset, InArgs.AllottedGeometry.Size),
+					FrameString,
+					SmallLayoutFont,
+					InArgs.DrawEffects,
+					InArgs.TickColor * 0.65f
+				);
+			}
+
+			for (int32 Step = 1; Step < MinorDivisions; ++Step)
+			{
+				// Compute the size of each tick mark.  If we are half way between to visible values display a slightly larger tick mark
+				const float MinorTickHeight = ((MinorDivisions % 2 == 0) && (Step % (MinorDivisions / 2)) == 0) ? 6.0f : 2.0f;
+				const float MinorLinePx = RangeToScreen.InputToLocalX(CurrentMajorLine + Step * MajorGridStep / MinorDivisions);
+
+				LinePoints[0] = FVector2D(MinorLinePx, InArgs.bMirrorLabels ? 0.0f : FMath::Abs(InArgs.AllottedGeometry.Size.Y - MinorTickHeight));
+				LinePoints[1] = FVector2D(MinorLinePx, LinePoints[0].Y + MinorTickHeight);
+
+				// Draw each sub mark
+				FSlateDrawElement::MakeLines(
+					OutDrawElements,
+					InArgs.StartLayer,
+					PaintGeometry,
+					LinePoints,
+					InArgs.DrawEffects,
+					InArgs.TickColor,
+					bAntiAliasLines
+				);
+			}
+		}
+	}
+
+	int32 FMLTimeSliderController::OnPaintTimeSlider(bool bMirrorLabels, const FGeometry& AllottedGeometry, const FSlateRect& MyCullingRect, FSlateWindowElementList& OutDrawElements, int32 LayerId, const FWidgetStyle& InWidgetStyle, bool bParentEnabled) const
+	{
+		const bool bEnabled = bParentEnabled;
+		const ESlateDrawEffect DrawEffects = bEnabled ? ESlateDrawEffect::None : ESlateDrawEffect::DisabledEffect;
+
+		TRange<double> LocalViewRange = TimeSliderArgs.ViewRange.Get();
+		const float    LocalViewRangeMin = LocalViewRange.GetLowerBoundValue();
+		const float    LocalViewRangeMax = LocalViewRange.GetUpperBoundValue();
+		const float    LocalSequenceLength = LocalViewRangeMax - LocalViewRangeMin;
+		TRange<FFrameNumber> LocalPlaybackRange = TimeSliderArgs.PlaybackRange.Get();
+
+		FVector2D Scale = FVector2D(1.0f, 1.0f);
+		if (LocalSequenceLength > 0)
+		{
+			FScrubRangeToScreen RangeToScreen(LocalViewRange, AllottedGeometry.Size);
+
+			// draw tick marks
+			const float MajorTickHeight = 9.0f;
+
+			FDrawTickArgs Args;
+			{
+				Args.AllottedGeometry = AllottedGeometry;
+				Args.bMirrorLabels = bMirrorLabels;
+				Args.bOnlyDrawMajorTicks = false;
+				Args.TickColor = FLinearColor::White;
+				Args.CullingRect = MyCullingRect;
+				Args.DrawEffects = DrawEffects;
+				Args.StartLayer = LayerId;
+				Args.TickOffset = bMirrorLabels ? 0.0f : FMath::Abs(AllottedGeometry.Size.Y - MajorTickHeight);
+				Args.MajorTickHeight = MajorTickHeight;
+			}
+
+			DrawTicks(OutDrawElements, LocalViewRange, RangeToScreen, Args);
+
+			// draw playback & selection range
+			FPaintPlaybackRangeArgs PlaybackRangeArgs(
+				bMirrorLabels ? FAppStyle::GetBrush("Sequencer.Timeline.PlayRange_Bottom_L") : FAppStyle::GetBrush("Sequencer.Timeline.PlayRange_Top_L"),
+				bMirrorLabels ? FAppStyle::GetBrush("Sequencer.Timeline.PlayRange_Bottom_R") : FAppStyle::GetBrush("Sequencer.Timeline.PlayRange_Top_R"),
+				6.f
+			);
+
+			LayerId = DrawPlaybackRange(AllottedGeometry, MyCullingRect, OutDrawElements, LayerId, RangeToScreen, PlaybackRangeArgs);
+
+			PlaybackRangeArgs.SolidFillOpacity = 0.05f;
+			LayerId = DrawSelectionRange(AllottedGeometry, MyCullingRect, OutDrawElements, LayerId, RangeToScreen, PlaybackRangeArgs);
+
+			// Draw the scrub handle
+			const float      HandleStart = RangeToScreen.InputToLocalX(TimeSliderArgs.ScrubPosition.Get().AsDecimal() / GetTickResolution().AsDecimal()) - 7.0f;
+			const float      HandleEnd = HandleStart + 13.0f;
+
+			const int32 ArrowLayer = LayerId + 2;
+			FPaintGeometry MyGeometry = AllottedGeometry.ToPaintGeometry(FVector2D(HandleStart, 0), FVector2D(HandleEnd - HandleStart, AllottedGeometry.Size.Y));
+			FLinearColor ScrubColor = InWidgetStyle.GetColorAndOpacityTint();
+			{
+				ScrubColor.A = ScrubColor.A * 0.75f;
+				ScrubColor.B *= 0.1f;
+				ScrubColor.G *= 0.2f;
+			}
+
+			const FSlateBrush* Brush = (bMirrorLabels ? ScrubHandleUpBrush : ScrubHandleDownBrush);
+
+			FSlateDrawElement::MakeBox(
+				OutDrawElements,
+				ArrowLayer,
+				MyGeometry,
+				Brush,
+				DrawEffects,
+				ScrubColor
+			);
+
+			{
+				// Draw the current time next to the scrub handle
+				FString FrameString = TimeSliderArgs.NumericTypeInterface->ToString(TimeSliderArgs.ScrubPosition.Get().GetFrame().Value);
+
+				if (GetDefault<UPersonaOptions>()->bTimelineDisplayPercentage)
+				{
+					double Percentage = FMath::Clamp(TimeSliderArgs.ScrubPosition.Get().AsDecimal() / FFrameTime(LocalPlaybackRange.Size<FFrameNumber>()).AsDecimal(), 0.0, 1.0);
+					FNumberFormattingOptions Options;
+					Options.MaximumFractionalDigits = 2;
+					FrameString += TEXT(" (") + FText::AsPercent(Percentage, &Options).ToString() + TEXT(")");
+				}
+
+				FSlateFontInfo SmallLayoutFont = FCoreStyle::GetDefaultFontStyle("Regular", 10);
+
+				const TSharedRef< FSlateFontMeasure > FontMeasureService = FSlateApplication::Get().GetRenderer()->GetFontMeasureService();
+				FVector2D TextSize = FontMeasureService->Measure(FrameString, SmallLayoutFont);
+
+				// Flip the text position if getting near the end of the view range
+				static const float TextOffsetPx = 2.f;
+				bool  bDrawLeft = (AllottedGeometry.Size.X - HandleEnd) < (TextSize.X + 14.f) - TextOffsetPx;
+				float TextPosition = bDrawLeft ? HandleStart - TextSize.X - TextOffsetPx : HandleEnd + TextOffsetPx;
+
+				FVector2D TextOffset(TextPosition, Args.bMirrorLabels ? TextSize.Y - 6.f : Args.AllottedGeometry.Size.Y - (Args.MajorTickHeight + TextSize.Y));
+
+				FSlateDrawElement::MakeText(
+					OutDrawElements,
+					Args.StartLayer + 1,
+					Args.AllottedGeometry.ToPaintGeometry(TextOffset, TextSize),
+					FrameString,
+					SmallLayoutFont,
+					Args.DrawEffects,
+					Args.TickColor
+				);
+			}
+
+			return ArrowLayer;
+		}
+
+		return LayerId;
+	}
+
+	int32 FMLTimeSliderController::DrawSelectionRange(const FGeometry& AllottedGeometry, const FSlateRect& MyCullingRect, FSlateWindowElementList& OutDrawElements, int32 LayerId, const FScrubRangeToScreen& RangeToScreen, const FPaintPlaybackRangeArgs& Args) const
+	{
+		TRange<double> SelectionRange = TimeSliderArgs.SelectionRange.Get() / GetTickResolution();
+
+		if (!SelectionRange.IsEmpty() && SelectionRange.HasLowerBound() && SelectionRange.HasUpperBound())
+		{
+			const float SelectionRangeL = RangeToScreen.InputToLocalX(SelectionRange.GetLowerBoundValue()) - 1;
+			const float SelectionRangeR = RangeToScreen.InputToLocalX(SelectionRange.GetUpperBoundValue()) + 1;
+			const auto DrawColor = FAppStyle::GetSlateColor("SelectionColor").GetColor(FWidgetStyle());
+
+			if (Args.SolidFillOpacity > 0.f)
+			{
+				FSlateDrawElement::MakeBox(
+					OutDrawElements,
+					LayerId + 1,
+					AllottedGeometry.ToPaintGeometry(FVector2D(SelectionRangeL, 0.f), FVector2D(SelectionRangeR - SelectionRangeL, AllottedGeometry.Size.Y)),
+					FAppStyle::GetBrush("WhiteBrush"),
+					ESlateDrawEffect::None,
+					DrawColor.CopyWithNewOpacity(Args.SolidFillOpacity)
+				);
+			}
+
+			FSlateDrawElement::MakeBox(
+				OutDrawElements,
+				LayerId + 1,
+				AllottedGeometry.ToPaintGeometry(FVector2D(SelectionRangeL, 0.f), FVector2D(Args.BrushWidth, AllottedGeometry.Size.Y)),
+				Args.StartBrush,
+				ESlateDrawEffect::None,
+				DrawColor
+			);
+
+			FSlateDrawElement::MakeBox(
+				OutDrawElements,
+				LayerId + 1,
+				AllottedGeometry.ToPaintGeometry(FVector2D(SelectionRangeR - Args.BrushWidth, 0.f), FVector2D(Args.BrushWidth, AllottedGeometry.Size.Y)),
+				Args.EndBrush,
+				ESlateDrawEffect::None,
+				DrawColor
+			);
+		}
+
+		return LayerId + 1;
+	}
+
+	int32 FMLTimeSliderController::DrawPlaybackRange(const FGeometry& AllottedGeometry, const FSlateRect& MyCullingRect, FSlateWindowElementList& OutDrawElements, int32 LayerId, const FScrubRangeToScreen& RangeToScreen, const FPaintPlaybackRangeArgs& Args) const
+	{
+		if (!TimeSliderArgs.PlaybackRange.IsSet())
+		{
+			return LayerId;
+		}
+
+		const uint8 OpacityBlend = TimeSliderArgs.SubSequenceRange.Get().IsSet() ? 128 : 255;
+
+		TRange<FFrameNumber> PlaybackRange = TimeSliderArgs.PlaybackRange.Get();
+		FFrameRate TickResolution = GetTickResolution();
+		const float PlaybackRangeL = RangeToScreen.InputToLocalX(PlaybackRange.GetLowerBoundValue() / TickResolution);
+		const float PlaybackRangeR = RangeToScreen.InputToLocalX(PlaybackRange.GetUpperBoundValue() / TickResolution) - 1;
+
+		FSlateDrawElement::MakeBox(
+			OutDrawElements,
+			LayerId + 1,
+			AllottedGeometry.ToPaintGeometry(FVector2D(PlaybackRangeL, 0.f), FVector2D(Args.BrushWidth, AllottedGeometry.Size.Y)),
+			Args.StartBrush,
+			ESlateDrawEffect::None,
+			FColor(32, 128, 32, OpacityBlend)	// 120, 75, 50 (HSV)
+		);
+
+		FSlateDrawElement::MakeBox(
+			OutDrawElements,
+			LayerId + 1,
+			AllottedGeometry.ToPaintGeometry(FVector2D(PlaybackRangeR - Args.BrushWidth, 0.f), FVector2D(Args.BrushWidth, AllottedGeometry.Size.Y)),
+			Args.EndBrush,
+			ESlateDrawEffect::None,
+			FColor(128, 32, 32, OpacityBlend)	// 0, 75, 50 (HSV)
+		);
+
+		// Black tint for excluded regions
+		FSlateDrawElement::MakeBox(
+			OutDrawElements,
+			LayerId + 1,
+			AllottedGeometry.ToPaintGeometry(FVector2D(0.f, 0.f), FVector2D(PlaybackRangeL, AllottedGeometry.Size.Y)),
+			FAppStyle::GetBrush("WhiteBrush"),
+			ESlateDrawEffect::None,
+			FLinearColor::Black.CopyWithNewOpacity(0.3f * OpacityBlend / 255.f)
+		);
+
+		FSlateDrawElement::MakeBox(
+			OutDrawElements,
+			LayerId + 1,
+			AllottedGeometry.ToPaintGeometry(FVector2D(PlaybackRangeR, 0.f), FVector2D(AllottedGeometry.Size.X - PlaybackRangeR, AllottedGeometry.Size.Y)),
+			FAppStyle::GetBrush("WhiteBrush"),
+			ESlateDrawEffect::None,
+			FLinearColor::Black.CopyWithNewOpacity(0.3f * OpacityBlend / 255.f)
+		);
+
+		return LayerId + 1;
+	}
+
+	FReply FMLTimeSliderController::OnMouseButtonDown(SWidget& WidgetOwner, const FGeometry& MyGeometry, const FPointerEvent& MouseEvent)
+	{
+		DistanceDragged = 0;
+		MouseDownPosition[0] = MouseDownPosition[1] = MouseEvent.GetScreenSpacePosition();
+		MouseDownGeometry = MyGeometry;
+		return FReply::Unhandled();
+	}
+
+	FReply FMLTimeSliderController::OnMouseButtonUp(SWidget& WidgetOwner, const FGeometry& MyGeometry, const FPointerEvent& MouseEvent)
+	{
+		bool bHandleLeftMouseButton = MouseEvent.GetEffectingButton() == EKeys::LeftMouseButton && WidgetOwner.HasMouseCapture();
+		bool bHandleRightMouseButton = MouseEvent.GetEffectingButton() == EKeys::RightMouseButton && WidgetOwner.HasMouseCapture() && TimeSliderArgs.AllowZoom;
+
+		FScrubRangeToScreen RangeToScreen = FScrubRangeToScreen(TimeSliderArgs.ViewRange.Get(), MyGeometry.Size);
+		FFrameTime          MouseTime = ComputeFrameTimeFromMouse(MyGeometry, MouseEvent.GetScreenSpacePosition(), RangeToScreen);
+
+		if (bHandleRightMouseButton)
+		{
+			if (!bPanning && DistanceDragged == 0.0f)
+			{
+				return WeakTimeline.Pin()->OnMouseButtonUp(MyGeometry, MouseEvent).ReleaseMouseCapture();
+			}
+
+			bPanning = false;
+			DistanceDragged = 0.f;
+
+			return FReply::Handled().ReleaseMouseCapture();
+		}
+		else if (bHandleLeftMouseButton)
+		{
+			if (MouseDragType == DRAG_PLAYBACK_START)
+			{
+				TimeSliderArgs.OnPlaybackRangeEndDrag.ExecuteIfBound();
+			}
+			else if (MouseDragType == DRAG_PLAYBACK_END)
+			{
+				TimeSliderArgs.OnPlaybackRangeEndDrag.ExecuteIfBound();
+			}
+			else if (MouseDragType == DRAG_SELECTION_START)
+			{
+				TimeSliderArgs.OnSelectionRangeEndDrag.ExecuteIfBound();
+			}
+			else if (MouseDragType == DRAG_SELECTION_END)
+			{
+				TimeSliderArgs.OnSelectionRangeEndDrag.ExecuteIfBound();
+			}
+			else
+			{
+				TimeSliderArgs.OnEndScrubberMovement.ExecuteIfBound();
+
+				CommitScrubPosition(MouseTime, /*bIsScrubbing=*/false);
+			}
+
+			MouseDragType = DRAG_NONE;
+			DistanceDragged = 0.f;
+
+			return FReply::Handled().ReleaseMouseCapture();
+		}
+
+		return FReply::Unhandled();
+	}
+
+	FReply FMLTimeSliderController::OnMouseMove(SWidget& WidgetOwner, const FGeometry& MyGeometry, const FPointerEvent& MouseEvent)
+	{
+		bool bHandleLeftMouseButton = MouseEvent.IsMouseButtonDown(EKeys::LeftMouseButton);
+		bool bHandleRightMouseButton = MouseEvent.IsMouseButtonDown(EKeys::RightMouseButton) && TimeSliderArgs.AllowZoom;
+
+		if (bHandleRightMouseButton)
+		{
+			if (!bPanning)
+			{
+				DistanceDragged += FMath::Abs(MouseEvent.GetCursorDelta().X);
+				if (DistanceDragged > FSlateApplication::Get().GetDragTriggerDistance())
+				{
+					bPanning = true;
+				}
+			}
+			else
+			{
+				TRange<double> LocalViewRange = TimeSliderArgs.ViewRange.Get();
+				double LocalViewRangeMin = LocalViewRange.GetLowerBoundValue();
+				double LocalViewRangeMax = LocalViewRange.GetUpperBoundValue();
+
+				FScrubRangeToScreen ScaleInfo(LocalViewRange, MyGeometry.Size);
+				FVector2D ScreenDelta = MouseEvent.GetCursorDelta();
+				FVector2D InputDelta;
+				InputDelta.X = ScreenDelta.X / ScaleInfo.PixelsPerInput;
+
+				double NewViewOutputMin = LocalViewRangeMin - InputDelta.X;
+				double NewViewOutputMax = LocalViewRangeMax - InputDelta.X;
+
+				ClampViewRange(NewViewOutputMin, NewViewOutputMax);
+				SetViewRange(NewViewOutputMin, NewViewOutputMax, EViewRangeInterpolation::Immediate);
+			}
+		}
+		else if (bHandleLeftMouseButton)
+		{
+			TRange<double> LocalViewRange = TimeSliderArgs.ViewRange.Get();
+			FScrubRangeToScreen RangeToScreen(LocalViewRange, MyGeometry.Size);
+			DistanceDragged += FMath::Abs(MouseEvent.GetCursorDelta().X);
+
+			if (MouseDragType == DRAG_NONE)
+			{
+				if (DistanceDragged > FSlateApplication::Get().GetDragTriggerDistance())
+				{
+					FFrameTime MouseDownFree = ComputeFrameTimeFromMouse(MyGeometry, MouseDownPosition[0], RangeToScreen, false);
+
+					const FFrameRate FrameResolution = GetTickResolution();
+					const bool       bLockedPlayRange = TimeSliderArgs.IsPlaybackRangeLocked.Get();
+					const float      MouseDownPixel = RangeToScreen.InputToLocalX(MouseDownFree / FrameResolution);
+					const bool       bHitScrubber = GetHitTestScrubberPixelRange(TimeSliderArgs.ScrubPosition.Get(), RangeToScreen).HandleRange.Contains(MouseDownPixel);
+
+					TRange<double>   SelectionRange = TimeSliderArgs.SelectionRange.Get() / FrameResolution;
+					TRange<double>   PlaybackRange = TimeSliderArgs.PlaybackRange.Get() / FrameResolution;
+
+					// Disable selection range test if it's empty so that the playback range scrubbing gets priority
+					if (!SelectionRange.IsEmpty() && !bHitScrubber && HitTestRangeEnd(RangeToScreen, SelectionRange, MouseDownPixel))
+					{
+						// selection range end scrubber
+						MouseDragType = DRAG_SELECTION_END;
+						TimeSliderArgs.OnSelectionRangeBeginDrag.ExecuteIfBound();
+					}
+					else if (!SelectionRange.IsEmpty() && !bHitScrubber && HitTestRangeStart(RangeToScreen, SelectionRange, MouseDownPixel))
+					{
+						// selection range start scrubber
+						MouseDragType = DRAG_SELECTION_START;
+						TimeSliderArgs.OnSelectionRangeBeginDrag.ExecuteIfBound();
+					}
+					else if (!bLockedPlayRange && !bHitScrubber && HitTestRangeEnd(RangeToScreen, PlaybackRange, MouseDownPixel))
+					{
+						// playback range end scrubber
+						MouseDragType = DRAG_PLAYBACK_END;
+						TimeSliderArgs.OnPlaybackRangeBeginDrag.ExecuteIfBound();
+					}
+					else if (!bLockedPlayRange && !bHitScrubber && HitTestRangeStart(RangeToScreen, PlaybackRange, MouseDownPixel))
+					{
+						// playback range start scrubber
+						MouseDragType = DRAG_PLAYBACK_START;
+						TimeSliderArgs.OnPlaybackRangeBeginDrag.ExecuteIfBound();
+					}
+					else
+					{
+						MouseDragType = DRAG_SCRUBBING_TIME;
+						TimeSliderArgs.OnBeginScrubberMovement.ExecuteIfBound();
+					}
+				}
+			}
+			else
+			{
+				FFrameTime MouseTime = ComputeFrameTimeFromMouse(MyGeometry, MouseEvent.GetScreenSpacePosition(), RangeToScreen);
+
+				// Set the start range time?
+				if (MouseDragType == DRAG_PLAYBACK_START)
+				{
+					SetPlaybackRangeStart(MouseTime.FrameNumber);
+				}
+				// Set the end range time?
+				else if (MouseDragType == DRAG_PLAYBACK_END)
+				{
+					SetPlaybackRangeEnd(MouseTime.FrameNumber - 1);
+				}
+				else if (MouseDragType == DRAG_SELECTION_START)
+				{
+					SetSelectionRangeStart(MouseTime.FrameNumber);
+				}
+				// Set the end range time?
+				else if (MouseDragType == DRAG_SELECTION_END)
+				{
+					SetSelectionRangeEnd(MouseTime.FrameNumber);
+				}
+				else if (MouseDragType == DRAG_SCRUBBING_TIME)
+				{
+					// Delegate responsibility for clamping to the current viewrange to the client
+					CommitScrubPosition(MouseTime, /*bIsScrubbing=*/true);
+				}
+			}
+		}
+
+		if (DistanceDragged != 0.f && (bHandleLeftMouseButton || bHandleRightMouseButton))
+		{
+			return FReply::Handled().CaptureMouse(WidgetOwner.AsShared());
+		}
+
+		return FReply::Handled();
+	}
+
+	void FMLTimeSliderController::CommitScrubPosition(FFrameTime NewValue, bool bIsScrubbing)
+	{
+		// Manage the scrub position ourselves if its not bound to a delegate
+		if (!TimeSliderArgs.ScrubPosition.IsBound())
+		{
+			TimeSliderArgs.ScrubPosition.Set(NewValue);
+		}
+
+		TimeSliderArgs.OnScrubPositionChanged.ExecuteIfBound(NewValue, bIsScrubbing, /*bEvaluate*/ true); //todo change if anim timeline needs to handle sequencer style middle mouse manipulation which changes time but doesn't evaluate
+	}
+
+	FReply FMLTimeSliderController::OnMouseWheel(SWidget& WidgetOwner, const FGeometry& MyGeometry, const FPointerEvent& MouseEvent)
+	{
+		TOptional<TRange<float>> NewTargetRange;
+
+		if (TimeSliderArgs.AllowZoom && MouseEvent.IsControlDown())
+		{
+			float MouseFractionX = MyGeometry.AbsoluteToLocal(MouseEvent.GetScreenSpacePosition()).X / MyGeometry.GetLocalSize().X;
+
+			const float ZoomDelta = -0.2f * MouseEvent.GetWheelDelta();
+			if (ZoomByDelta(ZoomDelta, MouseFractionX))
+			{
+				return FReply::Handled();
+			}
+		}
+		else if (MouseEvent.IsShiftDown())
+		{
+			PanByDelta(-MouseEvent.GetWheelDelta());
+			return FReply::Handled();
+		}
+
+		return FReply::Unhandled();
+	}
+
+	FCursorReply FMLTimeSliderController::OnCursorQuery(TSharedRef<const SWidget> WidgetOwner, const FGeometry& MyGeometry, const FPointerEvent& CursorEvent) const
+	{
+		FScrubRangeToScreen RangeToScreen(TimeSliderArgs.ViewRange.Get(), MyGeometry.Size);
+
+		const FFrameRate FrameResolution = GetTickResolution();
+		const bool       bLockedPlayRange = TimeSliderArgs.IsPlaybackRangeLocked.Get();
+		const float      HitTestPixel = MyGeometry.AbsoluteToLocal(CursorEvent.GetScreenSpacePosition()).X;
+		const bool       bHitScrubber = GetHitTestScrubberPixelRange(TimeSliderArgs.ScrubPosition.Get(), RangeToScreen).HandleRange.Contains(HitTestPixel);
+
+		TRange<double>   SelectionRange = TimeSliderArgs.SelectionRange.Get() / FrameResolution;
+		TRange<double>   PlaybackRange = TimeSliderArgs.PlaybackRange.Get() / FrameResolution;
+
+		if (MouseDragType == DRAG_SCRUBBING_TIME)
+		{
+			return FCursorReply::Unhandled();
+		}
+
+		// Use L/R resize cursor if we're dragging or hovering a playback range bound
+		if ((MouseDragType == DRAG_PLAYBACK_END) ||
+			(MouseDragType == DRAG_PLAYBACK_START) ||
+			(MouseDragType == DRAG_SELECTION_START) ||
+			(MouseDragType == DRAG_SELECTION_END) ||
+			(MouseDragType == DRAG_TIME) ||
+			(!bLockedPlayRange && !bHitScrubber && HitTestRangeStart(RangeToScreen, PlaybackRange, HitTestPixel)) ||
+			(!bLockedPlayRange && !bHitScrubber && HitTestRangeEnd(RangeToScreen, PlaybackRange, HitTestPixel)) ||
+			(!SelectionRange.IsEmpty() && !bHitScrubber && HitTestRangeStart(RangeToScreen, SelectionRange, HitTestPixel)) ||
+			(!SelectionRange.IsEmpty() && !bHitScrubber && HitTestRangeEnd(RangeToScreen, SelectionRange, HitTestPixel)))
+		{
+			return FCursorReply::Cursor(EMouseCursor::ResizeLeftRight);
+		}
+
+		return FCursorReply::Unhandled();
+	}
+
+	int32 FMLTimeSliderController::OnPaintViewArea(const FGeometry& AllottedGeometry, const FSlateRect& MyCullingRect, FSlateWindowElementList& OutDrawElements, int32 LayerId, bool bEnabled, const FPaintViewAreaArgs& Args) const
+	{
+		const ESlateDrawEffect DrawEffects = bEnabled ? ESlateDrawEffect::None : ESlateDrawEffect::DisabledEffect;
+
+		TRange<double> LocalViewRange = TimeSliderArgs.ViewRange.Get();
+		FScrubRangeToScreen RangeToScreen(LocalViewRange, AllottedGeometry.Size);
+
+		if (Args.PlaybackRangeArgs.IsSet())
+		{
+			FPaintPlaybackRangeArgs PaintArgs = Args.PlaybackRangeArgs.GetValue();
+			LayerId = DrawPlaybackRange(AllottedGeometry, MyCullingRect, OutDrawElements, LayerId, RangeToScreen, PaintArgs);
+			PaintArgs.SolidFillOpacity = 0.2f;
+			LayerId = DrawSelectionRange(AllottedGeometry, MyCullingRect, OutDrawElements, LayerId, RangeToScreen, PaintArgs);
+		}
+
+		if (Args.bDisplayTickLines)
+		{
+			static FLinearColor TickColor(0.1f, 0.1f, 0.1f, 0.3f);
+
+			// Draw major tick lines in the section area
+			FDrawTickArgs DrawTickArgs;
+			{
+				DrawTickArgs.AllottedGeometry = AllottedGeometry;
+				DrawTickArgs.bMirrorLabels = false;
+				DrawTickArgs.bOnlyDrawMajorTicks = true;
+				DrawTickArgs.TickColor = TickColor;
+				DrawTickArgs.CullingRect = MyCullingRect;
+				DrawTickArgs.DrawEffects = DrawEffects;
+				// Draw major ticks under sections
+				DrawTickArgs.StartLayer = LayerId - 1;
+				// Draw the tick the entire height of the section area
+				DrawTickArgs.TickOffset = 0.0f;
+				DrawTickArgs.MajorTickHeight = AllottedGeometry.Size.Y;
+			}
+
+			DrawTicks(OutDrawElements, LocalViewRange, RangeToScreen, DrawTickArgs);
+		}
+
+		if (Args.bDisplayScrubPosition)
+		{
+			// Draw a line for the scrub position
+			const float LinePos = RangeToScreen.InputToLocalX(TimeSliderArgs.ScrubPosition.Get().AsDecimal() / GetTickResolution().AsDecimal());
+
+			TArray<FVector2D> LinePoints;
+			{
+				LinePoints.AddUninitialized(2);
+				LinePoints[0] = FVector2D(0.0f, 0.0f);
+				LinePoints[1] = FVector2D(0.0f, FMath::FloorToFloat(AllottedGeometry.Size.Y));
+			}
+
+			FSlateDrawElement::MakeLines(
+				OutDrawElements,
+				LayerId + 1,
+				AllottedGeometry.ToPaintGeometry(FVector2D(LinePos, 0.0f), FVector2D(1.0f, 1.0f)),
+				LinePoints,
+				DrawEffects,
+				FLinearColor(1.f, 1.f, 1.f, .5f),
+				false
+			);
+		}
+
+		return LayerId;
+	}
+
+	FFrameTime FMLTimeSliderController::GetFrameTimeFromMouse(const FGeometry& Geometry, FVector2D ScreenSpacePosition) const
+	{
+		FScrubRangeToScreen ScrubRangeToScreen(TimeSliderArgs.ViewRange.Get(), Geometry.Size);
+		return ComputeFrameTimeFromMouse(Geometry, ScreenSpacePosition, ScrubRangeToScreen);
+	}
+
+	void FMLTimeSliderController::ClampViewRange(double& NewRangeMin, double& NewRangeMax)
+	{
+		bool bNeedsClampSet = false;
+		double NewClampRangeMin = TimeSliderArgs.ClampRange.Get().GetLowerBoundValue();
+		if (NewRangeMin < TimeSliderArgs.ClampRange.Get().GetLowerBoundValue())
+		{
+			NewClampRangeMin = NewRangeMin;
+			bNeedsClampSet = true;
+		}
+
+		double NewClampRangeMax = TimeSliderArgs.ClampRange.Get().GetUpperBoundValue();
+		if (NewRangeMax > TimeSliderArgs.ClampRange.Get().GetUpperBoundValue())
+		{
+			NewClampRangeMax = NewRangeMax;
+			bNeedsClampSet = true;
+		}
+
+		if (bNeedsClampSet)
+		{
+			SetClampRange(NewClampRangeMin, NewClampRangeMax);
+		}
+	}
+
+	void FMLTimeSliderController::SetViewRange(double NewRangeMin, double NewRangeMax, EViewRangeInterpolation Interpolation)
+	{
+		// Clamp to a minimum size to avoid zero-sized or negative visible ranges
+		double MinVisibleTimeRange = FFrameNumber(1) / GetTickResolution();
+		TRange<double> ExistingViewRange = TimeSliderArgs.ViewRange.Get();
+
+		if (NewRangeMax == ExistingViewRange.GetUpperBoundValue())
+		{
+			if (NewRangeMin > NewRangeMax - MinVisibleTimeRange)
+			{
+				NewRangeMin = NewRangeMax - MinVisibleTimeRange;
+			}
+		}
+		else if (NewRangeMax < NewRangeMin + MinVisibleTimeRange)
+		{
+			NewRangeMax = NewRangeMin + MinVisibleTimeRange;
+		}
+
+		// Clamp to the clamp range
+		const TRange<double> NewRange = TRange<double>(NewRangeMin, NewRangeMax);
+		TimeSliderArgs.OnViewRangeChanged.ExecuteIfBound(NewRange, Interpolation);
+
+		if (!TimeSliderArgs.ViewRange.IsBound())
+		{
+			// The  output is not bound to a delegate so we'll manage the value ourselves (no animation)
+			TimeSliderArgs.ViewRange.Set(NewRange);
+		}
+	}
+
+	void FMLTimeSliderController::SetClampRange(double NewRangeMin, double NewRangeMax)
+	{
+		const TRange<double> NewRange(NewRangeMin, NewRangeMax);
+
+		TimeSliderArgs.OnClampRangeChanged.ExecuteIfBound(NewRange);
+
+		if (!TimeSliderArgs.ClampRange.IsBound())
+		{
+			// The  output is not bound to a delegate so we'll manage the value ourselves (no animation)
+			TimeSliderArgs.ClampRange.Set(NewRange);
+		}
+	}
+
+	void FMLTimeSliderController::SetPlayRange(FFrameNumber RangeStart, int32 RangeDuration)
+	{
+		check(RangeDuration >= 0);
+
+		const TRange<FFrameNumber> NewRange(RangeStart, RangeStart + RangeDuration);
+
+		TimeSliderArgs.OnPlaybackRangeChanged.ExecuteIfBound(NewRange);
+
+		if (!TimeSliderArgs.PlaybackRange.IsBound())
+		{
+			// The  output is not bound to a delegate so we'll manage the value ourselves (no animation)
+			TimeSliderArgs.PlaybackRange.Set(NewRange);
+		}
+	}
+
+	bool FMLTimeSliderController::ZoomByDelta(float InDelta, float MousePositionFraction)
+	{
+		TRange<double> LocalViewRange = TimeSliderArgs.ViewRange.Get().GetAnimationTarget();
+		double LocalViewRangeMax = LocalViewRange.GetUpperBoundValue();
+		double LocalViewRangeMin = LocalViewRange.GetLowerBoundValue();
+		const double OutputViewSize = LocalViewRangeMax - LocalViewRangeMin;
+		const double OutputChange = OutputViewSize * InDelta;
+
+		double NewViewOutputMin = LocalViewRangeMin - (OutputChange * MousePositionFraction);
+		double NewViewOutputMax = LocalViewRangeMax + (OutputChange * (1.f - MousePositionFraction));
+
+		if (NewViewOutputMin < NewViewOutputMax)
+		{
+			ClampViewRange(NewViewOutputMin, NewViewOutputMax);
+			SetViewRange(NewViewOutputMin, NewViewOutputMax, EViewRangeInterpolation::Animated);
+			return true;
+		}
+
+		return false;
+	}
+
+	void FMLTimeSliderController::PanByDelta(float InDelta)
+	{
+		/**The fraction of the current view range to scroll per unit delta  */
+		const float ScrollPanFraction = 0.1f;
+
+		TRange<double> LocalViewRange = TimeSliderArgs.ViewRange.Get().GetAnimationTarget();
+
+		double CurrentMin = LocalViewRange.GetLowerBoundValue();
+		double CurrentMax = LocalViewRange.GetUpperBoundValue();
+
+		// Adjust the delta to be a percentage of the current range
+		InDelta *= ScrollPanFraction * (CurrentMax - CurrentMin);
+
+		double NewViewOutputMin = CurrentMin + InDelta;
+		double NewViewOutputMax = CurrentMax + InDelta;
+
+		ClampViewRange(NewViewOutputMin, NewViewOutputMax);
+		SetViewRange(NewViewOutputMin, NewViewOutputMax, EViewRangeInterpolation::Animated);
+	}
+
+	bool FMLTimeSliderController::HitTestRangeStart(const FScrubRangeToScreen& RangeToScreen, const TRange<double>& Range, float HitPixel) const
+	{
+		if (Range.HasLowerBound())
+		{
+			static float BrushSizeInStateUnits = 6.f, DragToleranceSlateUnits = 2.f, MouseTolerance = 2.f;
+			const float  RangeStartPixel = RangeToScreen.InputToLocalX(Range.GetLowerBoundValue());
+
+			// Hit test against the brush region to the right of the playback start position, +/- DragToleranceSlateUnits
+			return HitPixel >= RangeStartPixel - MouseTolerance - DragToleranceSlateUnits &&
+				HitPixel <= RangeStartPixel + MouseTolerance + BrushSizeInStateUnits + DragToleranceSlateUnits;
+		}
+
+		return false;
+	}
+
+	bool FMLTimeSliderController::HitTestRangeEnd(const FScrubRangeToScreen& RangeToScreen, const TRange<double>& Range, float HitPixel) const
+	{
+		if (Range.HasUpperBound())
+		{
+			static float BrushSizeInStateUnits = 6.f, DragToleranceSlateUnits = 2.f, MouseTolerance = 2.f;
+			const float  RangeEndPixel = RangeToScreen.InputToLocalX(Range.GetUpperBoundValue());
+
+			// Hit test against the brush region to the left of the playback end position, +/- DragToleranceSlateUnits
+			return HitPixel >= RangeEndPixel - MouseTolerance - BrushSizeInStateUnits - DragToleranceSlateUnits &&
+				HitPixel <= RangeEndPixel + MouseTolerance + DragToleranceSlateUnits;
+		}
+
+		return false;
+	}
+
+	void FMLTimeSliderController::SetPlaybackRangeStart(FFrameNumber NewStart)
+	{
+		TRange<FFrameNumber> PlaybackRange = TimeSliderArgs.PlaybackRange.Get();
+
+		if (NewStart <= UE::MovieScene::DiscreteExclusiveUpper(PlaybackRange))
+		{
+			TimeSliderArgs.OnPlaybackRangeChanged.ExecuteIfBound(TRange<FFrameNumber>(NewStart, PlaybackRange.GetUpperBound()));
+		}
+	}
+
+	void FMLTimeSliderController::SetPlaybackRangeEnd(FFrameNumber NewEnd)
+	{
+		TRange<FFrameNumber> PlaybackRange = TimeSliderArgs.PlaybackRange.Get();
+
+		if (NewEnd >= UE::MovieScene::DiscreteInclusiveLower(PlaybackRange))
+		{
+			TimeSliderArgs.OnPlaybackRangeChanged.ExecuteIfBound(TRange<FFrameNumber>(PlaybackRange.GetLowerBound(), NewEnd));
+		}
+	}
+
+	void FMLTimeSliderController::SetSelectionRangeStart(FFrameNumber NewStart)
+	{
+		TRange<FFrameNumber> SelectionRange = TimeSliderArgs.SelectionRange.Get();
+
+		if (SelectionRange.IsEmpty())
+		{
+			TimeSliderArgs.OnSelectionRangeChanged.ExecuteIfBound(TRange<FFrameNumber>(NewStart, NewStart + 1));
+		}
+		else if (NewStart <= UE::MovieScene::DiscreteExclusiveUpper(SelectionRange))
+		{
+			TimeSliderArgs.OnSelectionRangeChanged.ExecuteIfBound(TRange<FFrameNumber>(NewStart, SelectionRange.GetUpperBound()));
+		}
+	}
+
+	void FMLTimeSliderController::SetSelectionRangeEnd(FFrameNumber NewEnd)
+	{
+		TRange<FFrameNumber> SelectionRange = TimeSliderArgs.SelectionRange.Get();
+
+		if (SelectionRange.IsEmpty())
+		{
+			TimeSliderArgs.OnSelectionRangeChanged.ExecuteIfBound(TRange<FFrameNumber>(NewEnd - 1, NewEnd));
+		}
+		else if (NewEnd >= UE::MovieScene::DiscreteInclusiveLower(SelectionRange))
+		{
+			TimeSliderArgs.OnSelectionRangeChanged.ExecuteIfBound(TRange<FFrameNumber>(SelectionRange.GetLowerBound(), NewEnd));
+		}
+	}
+
+	class SAnimTimelineTransportControls : public SCompoundWidget
+	{
+	public:
+		SLATE_BEGIN_ARGS(SAnimTimelineTransportControls) {}
+
+		SLATE_END_ARGS()
+
+		void Construct(const FArguments& InArgs, TWeakPtr<FMLDeformerEditorModel>& InModel);
+
+		void SetModel(TWeakPtr<FMLDeformerEditorModel> InModel)
+		{
+			Model = InModel;
+		}
+
+	private:
+		UAnimSingleNodeInstance* GetPreviewInstance() const;
+
+		FReply OnClick_Forward_Step();
+
+		FReply OnClick_Forward_End();
+
+		FReply OnClick_Backward_Step();
+
+		FReply OnClick_Backward_End();
+
+		FReply OnClick_Forward();
+
+		FReply OnClick_ToggleLoop();
+
+		bool IsLoopStatusOn() const;
+
+		EPlaybackMode::Type GetPlaybackMode() const;
+
+	private:
+		/** Anim timeline model */
+		TWeakPtr<FMLDeformerEditorModel> Model;
+	};
+
+	void SAnimTimelineTransportControls::Construct(const FArguments& InArgs, TWeakPtr<FMLDeformerEditorModel>& InModel)
+	{
+		// Geoemetry Cache only supports forward playback, so backwards play is removed
+		Model = InModel;
+		FEditorWidgetsModule& EditorWidgetsModule = FModuleManager::LoadModuleChecked<FEditorWidgetsModule>("EditorWidgets");
+
+		FTransportControlArgs TransportControlArgs;
+		TransportControlArgs.OnForwardPlay = FOnClicked::CreateSP(this, &SAnimTimelineTransportControls::OnClick_Forward);
+		TransportControlArgs.OnForwardStep = FOnClicked::CreateSP(this, &SAnimTimelineTransportControls::OnClick_Forward_Step);
+		TransportControlArgs.OnBackwardStep = FOnClicked::CreateSP(this, &SAnimTimelineTransportControls::OnClick_Backward_Step);
+		TransportControlArgs.OnForwardEnd = FOnClicked::CreateSP(this, &SAnimTimelineTransportControls::OnClick_Forward_End);
+		TransportControlArgs.OnBackwardEnd = FOnClicked::CreateSP(this, &SAnimTimelineTransportControls::OnClick_Backward_End);
+
+		TransportControlArgs.OnGetPlaybackMode = FOnGetPlaybackMode::CreateSP(this, &SAnimTimelineTransportControls::GetPlaybackMode);
+
+		ChildSlot
+			[
+				EditorWidgetsModule.CreateTransportControl(TransportControlArgs)
+			];
+	}
+
+	FReply SAnimTimelineTransportControls::OnClick_Forward_Step()
+	{
+		if (Model.IsValid())
+		{
+			//Model.Pin()->OnPlayPressed();
+			FFrameNumber ScrubPosition = Model.Pin()->GetTickResScrubPosition();
+			ScrubPosition.Value += Model.Pin()->GetTicksPerFrame();
+			Model.Pin()->SetScrubPosition(ScrubPosition);
+			return FReply::Handled();
+		}
+		return FReply::Unhandled();
+	}
+
+	FReply SAnimTimelineTransportControls::OnClick_Forward_End()
+	{
+		if (Model.IsValid())
+		{
+			TRange<FFrameNumber> PlaybackRange = Model.Pin()->GetPlaybackRange();
+			FFrameNumber UpperBound = PlaybackRange.GetUpperBoundValue();
+			Model.Pin()->SetScrubPosition(UpperBound);
+			return FReply::Handled();
+		}
+		return FReply::Unhandled();
+	}
+
+	FReply SAnimTimelineTransportControls::OnClick_Backward_Step()
+	{
+		if (Model.IsValid())
+		{
+			//Model.Pin()->OnPlayPressed(EPlaybackMode::Stopped);
+			FFrameNumber ScrubPosition = Model.Pin()->GetTickResScrubPosition();
+			ScrubPosition.Value -= Model.Pin()->GetTicksPerFrame();
+			Model.Pin()->SetScrubPosition(ScrubPosition);
+			return FReply::Handled();
+		}
+		return FReply::Unhandled();
+	}
+
+	FReply SAnimTimelineTransportControls::OnClick_Backward_End()
+	{
+		if (Model.IsValid())
+		{
+			TRange<FFrameNumber> PlaybackRange = Model.Pin()->GetPlaybackRange();
+			FFrameNumber LowerBound = PlaybackRange.GetLowerBoundValue();
+			Model.Pin()->SetScrubPosition(LowerBound);
+			return FReply::Handled();
+		}
+		return FReply::Unhandled();
+	}
+
+	FReply SAnimTimelineTransportControls::OnClick_Forward()
+	{
+		if (Model.IsValid())
+		{
+			Model.Pin()->OnPlayPressed();
+			return FReply::Handled();
+		}
+		return FReply::Unhandled();
+	}
+
+
+	EPlaybackMode::Type SAnimTimelineTransportControls::GetPlaybackMode() const
+	{
+		if (Model.IsValid())
+		{
+			if (Model.Pin()->IsPlayingAnim())
+			{
+				return EPlaybackMode::PlayingForward;
+			}
+		}
+		return EPlaybackMode::Stopped;
+	}
+
+	/**
+	 * An overlay that displays global information in the track area
+	 */
+	class SAnimTimelineOverlay : public SCompoundWidget
+	{
+	public:
+		SLATE_BEGIN_ARGS(SAnimTimelineOverlay)
+			: _DisplayTickLines(true)
+			, _DisplayScrubPosition(false)
+		{}
+
+		SLATE_ATTRIBUTE(bool, DisplayTickLines)
+			SLATE_ATTRIBUTE(bool, DisplayScrubPosition)
+			SLATE_ATTRIBUTE(FPaintPlaybackRangeArgs, PaintPlaybackRangeArgs)
+
+			SLATE_END_ARGS()
+
+			void Construct(const FArguments& InArgs, TSharedRef<FMLTimeSliderController> InTimeSliderController)
+		{
+			bDisplayScrubPosition = InArgs._DisplayScrubPosition;
+			bDisplayTickLines = InArgs._DisplayTickLines;
+			PaintPlaybackRangeArgs = InArgs._PaintPlaybackRangeArgs;
+			TimeSliderController = InTimeSliderController;
+		}
+
+	private:
+		/** SWidget Interface */
+		virtual int32 OnPaint(const FPaintArgs& Args, const FGeometry& AllottedGeometry, const FSlateRect& MyCullingRect, FSlateWindowElementList& OutDrawElements, int32 LayerId, const FWidgetStyle& InWidgetStyle, bool bParentEnabled) const override;
+
+	private:
+		/** Controller for manipulating time */
+		TSharedPtr<FMLTimeSliderController> TimeSliderController;
+		/** Whether or not to display the scrub position */
+		TAttribute<bool> bDisplayScrubPosition;
+		/** Whether or not to display tick lines */
+		TAttribute<bool> bDisplayTickLines;
+		/** User-supplied options for drawing playback range */
+		TAttribute<FPaintPlaybackRangeArgs> PaintPlaybackRangeArgs;
+
+	};
+
+	int32 SAnimTimelineOverlay::OnPaint(const FPaintArgs& Args, const FGeometry& AllottedGeometry, const FSlateRect& MyCullingRect, FSlateWindowElementList& OutDrawElements, int32 LayerId, const FWidgetStyle& InWidgetStyle, bool bParentEnabled) const
+	{
+		FPaintViewAreaArgs PaintArgs;
+		PaintArgs.bDisplayTickLines = bDisplayTickLines.Get();
+		PaintArgs.bDisplayScrubPosition = bDisplayScrubPosition.Get();
+
+		if (PaintPlaybackRangeArgs.IsSet())
+		{
+			PaintArgs.PlaybackRangeArgs = PaintPlaybackRangeArgs.Get();
+		}
+
+		TimeSliderController->OnPaintViewArea(AllottedGeometry, MyCullingRect, OutDrawElements, LayerId, ShouldBeEnabled(bParentEnabled), PaintArgs);
+
+		return SCompoundWidget::OnPaint(Args, AllottedGeometry, MyCullingRect, OutDrawElements, LayerId, InWidgetStyle, bParentEnabled);
+	}
+
+	void SMLDeformerTimeline::Construct(const FArguments& InArgs, FMLDeformerEditorToolkit* InEditor)
+	{
+		TWeakPtr<FMLDeformerEditorModel> WeakModel = InEditor->GetActiveModelPointer();
+
+		Model = WeakModel;
+		OnReceivedFocus = InArgs._OnReceivedFocus;
+
+		int32 TickResolutionValue = Model.Pin()->GetTickResolution();
+		int32 SequenceFrameRate = FMath::RoundToInt(Model.Pin()->GetFrameRate());
+
+		ViewRange = MakeAttributeLambda([this](){
+				if (Model.IsValid())
+				{
+					TRange<double> Range = Model.Pin()->GetViewRange();
+					return FAnimatedRange(Range.GetLowerBoundValue(), Range.GetUpperBoundValue());
+				}
+				else
+				{
+					return FAnimatedRange(0.0, 0.0);
+				}
+			});
+
+		TAttribute<EFrameNumberDisplayFormats> DisplayFormat = MakeAttributeLambda([this]()
+		{
+				if (Model.IsValid())
+				{
+					return Model.Pin()->IsDisplayingFrames() ? EFrameNumberDisplayFormats::Frames : EFrameNumberDisplayFormats::Seconds;
+				}
+				else
+				{
+					return EFrameNumberDisplayFormats::Frames;
+				}
+		});
+
+		TAttribute<FFrameRate> TickResolution = MakeAttributeLambda([TickResolutionValue]()
+		{
+			return FFrameRate(TickResolutionValue, 1);
+		});
+
+		TAttribute<FFrameRate> DisplayRate = MakeAttributeLambda([SequenceFrameRate]()
+		{
+			return FFrameRate(SequenceFrameRate, 1);
+		});
+
+		// Create our numeric type interface so we can pass it to the time slider below.
+		NumericTypeInterface = MakeShareable(new FFrameNumberInterface(DisplayFormat, 0, TickResolution, DisplayRate));
+
+		FTimeSliderArgs TimeSliderArgs;
+		{
+			TimeSliderArgs.ScrubPosition = MakeAttributeLambda([WeakModel](){ return WeakModel.IsValid() ? WeakModel.Pin()->GetTickResScrubPosition() : FFrameTime(0); });
+			TimeSliderArgs.ViewRange = ViewRange;
+			TimeSliderArgs.PlaybackRange = MakeAttributeLambda([WeakModel](){ return WeakModel.IsValid() ? WeakModel.Pin()->GetPlaybackRange() : TRange<FFrameNumber>(0, 0); });
+			TimeSliderArgs.ClampRange = MakeAttributeLambda([this](){
+					if (Model.IsValid())
+					{
+						TRange<double> Range = Model.Pin()->GetWorkingRange();
+						return FAnimatedRange(Range.GetLowerBoundValue(), Range.GetUpperBoundValue());
+					}
+					else
+					{
+						return FAnimatedRange(0.0, 0.0);
+					}
+				});
+			TimeSliderArgs.DisplayRate = DisplayRate;
+			TimeSliderArgs.TickResolution = TickResolution;
+			TimeSliderArgs.OnViewRangeChanged = FOnViewRangeChanged::CreateSP(this, &SMLDeformerTimeline::HandleViewRangeChanged);
+			TimeSliderArgs.OnClampRangeChanged = FOnTimeRangeChanged::CreateSP(Model.Pin().Get(), &FMLDeformerEditorModel::HandleWorkingRangeChanged);
+			TimeSliderArgs.IsPlaybackRangeLocked = true;
+			TimeSliderArgs.PlaybackStatus = EMovieScenePlayerStatus::Stopped;
+			TimeSliderArgs.NumericTypeInterface = NumericTypeInterface;
+			TimeSliderArgs.OnScrubPositionChanged = FOnScrubPositionChanged::CreateSP(this, &SMLDeformerTimeline::HandleScrubPositionChanged);
+		}
+
+		TimeSliderController = MakeShareable(new FMLTimeSliderController(TimeSliderArgs, Model, SharedThis(this)));
+	
+		TSharedRef<FMLTimeSliderController> TimeSliderControllerRef = TimeSliderController.ToSharedRef();
+
+		// Create the top slider
+		const bool bMirrorLabels = false;
+		ISequencerWidgetsModule& SequencerWidgets = FModuleManager::Get().LoadModuleChecked<ISequencerWidgetsModule>("SequencerWidgets");
+		TopTimeSlider = SequencerWidgets.CreateTimeSlider(TimeSliderControllerRef, bMirrorLabels);
+
+		// Create bottom time range slider
+		TSharedRef<ITimeSlider> BottomTimeRange = SequencerWidgets.CreateTimeRange(
+			FTimeRangeArgs(
+				EShowRange::ViewRange | EShowRange::WorkingRange | EShowRange::PlaybackRange,
+				EShowRange::ViewRange | EShowRange::WorkingRange,
+				TimeSliderControllerRef,
+				EVisibility::Visible,
+				NumericTypeInterface.ToSharedRef()
+			),
+			SequencerWidgets.CreateTimeRangeSlider(TimeSliderControllerRef)
+		);
+
+		TSharedRef<SScrollBar> ScrollBar = SNew(SScrollBar)
+			.Thickness(FVector2D(5.0f, 5.0f));
+
+		ColumnFillCoefficients[0] = 0.2f;
+		ColumnFillCoefficients[1] = 0.8f;
+
+		TransportControls = SNew(SAnimTimelineTransportControls, Model);
+
+		TAttribute<float> FillCoefficient_0, FillCoefficient_1;
+		{
+			FillCoefficient_0.Bind(TAttribute<float>::FGetter::CreateSP(this, &SMLDeformerTimeline::GetColumnFillCoefficient, 0));
+			FillCoefficient_1.Bind(TAttribute<float>::FGetter::CreateSP(this, &SMLDeformerTimeline::GetColumnFillCoefficient, 1));
+		}
+
+		const int32 Column0 = 0, Column1 = 1;
+		const int32 Row0 = 0, Row1 = 1, Row2 = 2, Row3 = 3, Row4 = 4;
+
+		const float CommonPadding = 3.f;
+		const FMargin ResizeBarPadding(4.f, 0, 0, 0);
+
+		ChildSlot
+		[
+			SNew(SOverlay)
+			+SOverlay::Slot()
+			[
+				SNew(SVerticalBox)
+				+SVerticalBox::Slot()
+				[
+					SNew(SOverlay)
+					+SOverlay::Slot()
+					[
+						SNew(SGridPanel)
+						.FillRow(1, 1.0f)
+						.FillColumn(0, FillCoefficient_0)
+						.FillColumn(1, FillCoefficient_1)
+
+						// outliner search box
+						+SGridPanel::Slot(Column0, Row0, SGridPanel::Layer(10))
+						[
+							SNew(SHorizontalBox)
+							+SHorizontalBox::Slot()
+							.FillWidth(1.0f)
+							.VAlign(VAlign_Center)
+							+SHorizontalBox::Slot()
+							.VAlign(VAlign_Center)
+							.HAlign(HAlign_Center)
+							.AutoWidth()
+							.Padding(2.0f, 0.0f, 2.0f, 0.0f)
+							[
+								SNew(SBox)
+								.MinDesiredWidth(30.0f)
+								.VAlign(VAlign_Center)
+								.HAlign(HAlign_Center)
+								[
+									// Current Play Time 
+									SNew(SSpinBox<double>)
+									.Style(&FAppStyle::GetWidgetStyle<FSpinBoxStyle>("Sequencer.PlayTimeSpinBox"))
+									.Value_Lambda([this]() -> double
+									{
+										if (Model.IsValid())
+										{
+											return Model.Pin()->GetTickResScrubPosition().Value;
+										}
+										return 0.0;
+									})
+									.OnValueChanged(this, &SMLDeformerTimeline::SetPlayTime)
+									.OnValueCommitted_Lambda([this](double InFrame, ETextCommit::Type)
+									{
+										SetPlayTime(InFrame);
+									})
+									.MinValue(TOptional<double>())
+									.MaxValue(TOptional<double>())
+									.TypeInterface(NumericTypeInterface)
+									.Delta(this, &SMLDeformerTimeline::GetSpinboxDelta)
+									.LinearDeltaSensitivity(25)
+								]
+							]
+						]
+						// Transport controls
+						+SGridPanel::Slot(Column0, Row3, SGridPanel::Layer(10))
+						.VAlign(VAlign_Center)
+						.HAlign(HAlign_Center)
+						[
+							TransportControls.ToSharedRef()
+						]
+
+						// Second column
+						+SGridPanel::Slot(Column1, Row0)
+						.Padding(ResizeBarPadding)
+						.RowSpan(2)
+						[
+							SNew(SBorder)
+							.BorderImage(FAppStyle::GetBrush("ToolPanel.GroupBorder"))
+							[
+								SNew(SSpacer)
+							]
+						]
+
+						+SGridPanel::Slot(Column1, Row0, SGridPanel::Layer(10))
+						.Padding(ResizeBarPadding)
+						[
+							SNew( SBorder )
+							.BorderImage( FAppStyle::GetBrush("ToolPanel.GroupBorder") )
+							.BorderBackgroundColor( FLinearColor(.50f, .50f, .50f, 1.0f ) )
+							.Padding(0)
+							.Clipping(EWidgetClipping::ClipToBounds)
+							[
+								TopTimeSlider.ToSharedRef()
+							]
+						]
+
+						// Overlay that draws the tick lines
+						+SGridPanel::Slot(Column1, Row1, SGridPanel::Layer(10))
+						.Padding(ResizeBarPadding)
+						[
+							SNew(SAnimTimelineOverlay, TimeSliderControllerRef)
+							.Visibility( EVisibility::HitTestInvisible )
+							.DisplayScrubPosition( false )
+							.DisplayTickLines( true )
+							.Clipping(EWidgetClipping::ClipToBounds)
+							.PaintPlaybackRangeArgs(FPaintPlaybackRangeArgs(FAppStyle::GetBrush("Sequencer.Timeline.PlayRange_L"), FAppStyle::GetBrush("Sequencer.Timeline.PlayRange_R"), 6.f))
+						]
+
+						// Overlay that draws the scrub position
+						+SGridPanel::Slot(Column1, Row1, SGridPanel::Layer(20))
+						.Padding(ResizeBarPadding)
+						[
+							SNew(SAnimTimelineOverlay, TimeSliderControllerRef)
+							.Visibility( EVisibility::HitTestInvisible )
+							.DisplayScrubPosition( true )
+							.DisplayTickLines( false )
+							.Clipping(EWidgetClipping::ClipToBounds)
+						]
+
+						// play range slider
+
+						+SGridPanel::Slot(Column1, Row3, SGridPanel::Layer(10))
+						.Padding(ResizeBarPadding)
+						[
+							SNew(SBorder)
+							.BorderImage( FAppStyle::GetBrush("ToolPanel.GroupBorder") )
+							.BorderBackgroundColor( FLinearColor(0.5f, 0.5f, 0.5f, 1.0f ) )
+							.Clipping(EWidgetClipping::ClipToBounds)
+							.Padding(0)
+							[
+								BottomTimeRange
+							]
+						]
+
+					]
+				]
+			]
+		];
+	}
+
+	FReply SMLDeformerTimeline::OnMouseButtonUp(const FGeometry& MyGeometry, const FPointerEvent& MouseEvent)
+	{
+		if(MouseEvent.GetEffectingButton() == EKeys::RightMouseButton)
+		{
+			FWidgetPath WidgetPath = MouseEvent.GetEventPath() != nullptr ? *MouseEvent.GetEventPath() : FWidgetPath();
+
+			const bool bCloseAfterSelection = true;
+			FMenuBuilder MenuBuilder(bCloseAfterSelection, NULL);
+
+			TWeakPtr<FMLDeformerEditorModel> WeakModel = Model;
+
+			MenuBuilder.BeginSection("TimelineOptions", LOCTEXT("TimelineOptions", "Timeline Options") );
+			{
+				MenuBuilder.AddSubMenu(
+					LOCTEXT("TimeFormat", "Time Format"),
+					LOCTEXT("TimeFormatTooltip", "Choose the format of times we display in the timeline"),
+					FNewMenuDelegate::CreateLambda([WeakModel](FMenuBuilder& InMenuBuilder)
+					{
+						InMenuBuilder.BeginSection("TimeFormat", LOCTEXT("TimeFormat", "Time Format") );
+						{
+							InMenuBuilder.AddMenuEntry(
+									LOCTEXT("DisplayFrames", "Display Frames"),
+									LOCTEXT("DisplayFrames_Tooltip", "Display Timeline in Frames."),
+									FSlateIcon(),
+									FUIAction(FExecuteAction::CreateLambda([WeakModel]() {
+											if (WeakModel.IsValid())
+											{
+												WeakModel.Pin()->SetDisplayFrames(true);
+											}
+										}
+									),
+									FCanExecuteAction(),
+									FIsActionChecked::CreateLambda([WeakModel]() {
+											if (WeakModel.IsValid()) 
+											{ 
+												return WeakModel.Pin()->IsDisplayingFrames();
+											}
+											return false;
+										}
+									)
+								),
+								NAME_None,
+								EUserInterfaceActionType::RadioButton
+							);
+							InMenuBuilder.AddMenuEntry(
+									LOCTEXT("DisplayFrames", "Display Seconds"),
+									LOCTEXT("DisplayFrames_Tooltip", "Display Timeline in Seconds."),
+									FSlateIcon(),
+									FUIAction(FExecuteAction::CreateLambda([WeakModel]() {
+											if (WeakModel.IsValid())
+											{
+												WeakModel.Pin()->SetDisplayFrames(false);
+											}
+										}
+									),
+									FCanExecuteAction(),
+									FIsActionChecked::CreateLambda([WeakModel]() {
+											if (WeakModel.IsValid())
+											{
+												return !WeakModel.Pin()->IsDisplayingFrames();
+											}
+											return false;
+										}
+									)
+								),
+								NAME_None,
+								EUserInterfaceActionType::RadioButton
+							);
+						}
+						InMenuBuilder.EndSection();
+					})
+				);
+			}
+
+			MenuBuilder.EndSection();
+
+			FSlateApplication::Get().PushMenu(SharedThis(this), WidgetPath, MenuBuilder.MakeWidget(), FSlateApplication::Get().GetCursorPos(), FPopupTransitionEffect(FPopupTransitionEffect::ContextMenu));
+
+			return FReply::Handled();
+		}
+
+		return FReply::Unhandled();
+	}
+
+
+	TSharedRef<INumericTypeInterface<double>> SMLDeformerTimeline::GetNumericTypeInterface() const
+	{
+		return NumericTypeInterface.ToSharedRef();
+	}
+
+	// FFrameRate::ComputeGridSpacing doesnt deal well with prime numbers, so we have a custom impl here
+	static bool ComputeGridSpacing(const FFrameRate& InFrameRate, float PixelsPerSecond, double& OutMajorInterval, int32& OutMinorDivisions, float MinTickPx, float DesiredMajorTickPx)
+	{
+		// First try built-in spacing
+		bool bResult = InFrameRate.ComputeGridSpacing(PixelsPerSecond, OutMajorInterval, OutMinorDivisions, MinTickPx, DesiredMajorTickPx);
+		if(!bResult || OutMajorInterval == 1.0)
+		{
+			if (PixelsPerSecond <= 0.f)
+			{
+				return false;
+			}
+
+			const int32 RoundedFPS = FMath::RoundToInt(InFrameRate.AsDecimal());
+
+			if (RoundedFPS > 0)
+			{
+				// Showing frames
+				TArray<int32, TInlineAllocator<10>> CommonBases;
+
+				// Divide the rounded frame rate by 2s, 3s or 5s recursively
+				{
+					const int32 Denominators[] = { 2, 3, 5 };
+
+					int32 LowestBase = RoundedFPS;
+					for (;;)
+					{
+						CommonBases.Add(LowestBase);
+	
+						if (LowestBase % 2 == 0)      { LowestBase = LowestBase / 2; }
+						else if (LowestBase % 3 == 0) { LowestBase = LowestBase / 3; }
+						else if (LowestBase % 5 == 0) { LowestBase = LowestBase / 5; }
+						else
+						{ 
+							int32 LowestResult = LowestBase;
+							for(int32 Denominator : Denominators)
+							{
+								int32 Result = LowestBase / Denominator;
+								if(Result > 0 && Result < LowestResult)
+								{
+									LowestResult = Result;
+								}
+							}
+
+							if(LowestResult < LowestBase)
+							{
+								LowestBase = LowestResult;
+							}
+							else
+							{
+								break;
+							}
+						}
+					}
+				}
+
+				Algo::Reverse(CommonBases);
+
+				const int32 Scale     = FMath::CeilToInt(DesiredMajorTickPx / PixelsPerSecond * InFrameRate.AsDecimal());
+				const int32 BaseIndex = FMath::Min(Algo::LowerBound(CommonBases, Scale), CommonBases.Num()-1);
+				const int32 Base      = CommonBases[BaseIndex];
+
+				int32 MajorIntervalFrames = FMath::CeilToInt(Scale / float(Base)) * Base;
+				OutMajorInterval  = MajorIntervalFrames * InFrameRate.AsInterval();
+
+				// Find the lowest number of divisions we can show that's larger than the minimum tick size
+				OutMinorDivisions = 0;
+				for (int32 DivIndex = 0; DivIndex < BaseIndex; ++DivIndex)
+				{
+					if (Base % CommonBases[DivIndex] == 0)
+					{
+						int32 MinorDivisions = MajorIntervalFrames/CommonBases[DivIndex];
+						if (OutMajorInterval / MinorDivisions * PixelsPerSecond >= MinTickPx)
+						{
+							OutMinorDivisions = MinorDivisions;
+							break;
+						}
+					}
+				}
+			}
+		}
+
+		return OutMajorInterval != 0;
+	}
+
+	void SMLDeformerTimeline::SetDisplayFormat(EFrameNumberDisplayFormats InFormat)
+	{
+		if (Model.IsValid())
+		{
+			if (InFormat == EFrameNumberDisplayFormats::Frames)
+			{
+				Model.Pin()->SetDisplayFrames(true);
+			}
+			else
+			{
+				Model.Pin()->SetDisplayFrames(false);
+			}
+		}
+	}
+
+	bool SMLDeformerTimeline::IsDisplayFormatChecked(EFrameNumberDisplayFormats InFormat) const
+	{
+		if (Model.IsValid())
+		{
+			if (InFormat == EFrameNumberDisplayFormats::Frames)
+			{
+				return Model.Pin()->IsDisplayingFrames();
+			}
+			else if (InFormat == EFrameNumberDisplayFormats::Seconds)
+			{
+				return !Model.Pin()->IsDisplayingFrames();
+			}
+		}
+		return false;
+	}
+
+	bool SMLDeformerTimeline::GetGridMetrics(float PhysicalWidth, double& OutMajorInterval, int32& OutMinorDivisions) const
+	{
+		FSlateFontInfo SmallLayoutFont = FCoreStyle::GetDefaultFontStyle("Regular", 8);
+		TSharedRef<FSlateFontMeasure> FontMeasureService = FSlateApplication::Get().GetRenderer()->GetFontMeasureService();
+		double FrameRate = 30.0; 
+		if (Model.IsValid())
+		{
+			FrameRate = Model.Pin()->GetFrameRate();
+		}
+		FFrameRate DisplayRate(FMath::Max(FMath::RoundToInt(FrameRate), 1), 1);
+		double BiggestTime = ViewRange.Get().GetUpperBoundValue();
+		FString TickString = NumericTypeInterface->ToString((BiggestTime * DisplayRate).FrameNumber.Value);
+		FVector2D MaxTextSize = FontMeasureService->Measure(TickString, SmallLayoutFont);
+
+		static float MajorTickMultiplier = 2.f;
+
+		float MinTickPx = MaxTextSize.X + 5.f;
+		float DesiredMajorTickPx = MaxTextSize.X * MajorTickMultiplier;
+
+		if (PhysicalWidth > 0 && DisplayRate.AsDecimal() > 0)
+		{
+			return ComputeGridSpacing(
+				DisplayRate,
+				PhysicalWidth / ViewRange.Get().Size<double>(),
+				OutMajorInterval,
+				OutMinorDivisions,
+				MinTickPx,
+				DesiredMajorTickPx);
+		}
+
+		return false;
+	}
+
+	TSharedPtr<ITimeSliderController> SMLDeformerTimeline::GetTimeSliderController() const 
+	{ 
+		return TimeSliderController; 
+	}
+
+	void SMLDeformerTimeline::SetModel(TWeakPtr<FMLDeformerEditorModel> InModel)
+	{
+		Model = InModel; 
+		TimeSliderController->SetModel(Model);
+		TransportControls->SetModel(Model);
+	}
+
+	void SMLDeformerTimeline::OnColumnFillCoefficientChanged(float FillCoefficient, int32 ColumnIndex)
+	{
+		ColumnFillCoefficients[ColumnIndex] = FillCoefficient;
+	}
+
+	void SMLDeformerTimeline::HandleScrubPositionChanged(FFrameTime NewScrubPosition, bool bIsScrubbing, bool bEvaluate)
+	{
+		if (Model.IsValid())
+		{
+			Model.Pin()->SetScrubPosition(NewScrubPosition);
+		}
+	}
+
+	void SMLDeformerTimeline::HandleViewRangeChanged(TRange<double> InRange, EViewRangeInterpolation InInterpolation)
+	{
+		if (Model.IsValid())
+		{
+			Model.Pin()->SetViewRange(InRange);
+		}
+	}
+
+	double SMLDeformerTimeline::GetSpinboxDelta() const
+	{
+		if (Model.IsValid())
+		{
+			return FFrameRate(Model.Pin()->GetTickResolution(), 1).AsDecimal() * FFrameRate(FMath::RoundToInt(Model.Pin()->GetFrameRate()), 1).AsInterval();
+		}
+		return 0.0;
+	}
+
+	void SMLDeformerTimeline::SetPlayTime(double InFrameTime)
+	{
+		FFrameTime FrameTime((int)InFrameTime, 0.0f);
+		if (Model.IsValid())
+		{
+			Model.Pin()->SetScrubPosition(FrameTime);
+		}
+	}
+}
+
+#undef LOCTEXT_NAMESPACE
diff --git a/Engine/Plugins/Experimental/MLDeformer/MLDeformerFramework/Source/MLDeformerFrameworkEditor/Private/SMLDeformerTimeline.h b/Engine/Plugins/Experimental/MLDeformer/MLDeformerFramework/Source/MLDeformerFrameworkEditor/Private/SMLDeformerTimeline.h
new file mode 100644
--- /dev/null
+++ b/Engine/Plugins/Experimental/MLDeformer/MLDeformerFramework/Source/MLDeformerFrameworkEditor/Private/SMLDeformerTimeline.h
@@ -0,0 +1,128 @@
+// Copyright Epic Games, Inc. All Rights Reserved.
+
+#pragma once
+
+#include "CoreTypes.h"
+#include "Templates/SharedPointer.h"
+#include "Widgets/DeclarativeSyntaxSupport.h"
+#include "Widgets/SCompoundWidget.h"
+#include "Widgets/Input/NumericTypeInterface.h"
+#include "Framework/MultiBox/MultiBoxBuilder.h"
+#include "ITimeSlider.h"
+#include "ITransportControl.h"
+#include "Framework/Commands/Commands.h"
+#include "MLDeformerEditorToolkit.h"
+
+class SSearchBox;
+struct FPaintPlaybackRangeArgs;
+class UAnimSingleNodeInstance;
+enum class EViewRangeInterpolation;
+enum class EFrameNumberDisplayFormats : uint8;
+
+namespace UE::MLDeformer
+{
+	class FMLDeformerEditorToolkit;
+	class FMLDeformerEditorModel;
+	class FMLTimeSliderController;
+	class SAnimTimelineTransportControls; 
+	/**
+	 * Implements a timeline widget for the ML Deformer Editor.
+	 */
+	class SMLDeformerTimeline : public SCompoundWidget
+	{
+	public:
+		SLATE_BEGIN_ARGS(SMLDeformerTimeline) { }
+
+		/** The current total range of frame indices */
+		SLATE_ATTRIBUTE(FInt32Interval, ViewIndices)
+
+			/** Called when any widget contained within the anim timeline has received focus */
+			SLATE_EVENT(FSimpleDelegate, OnReceivedFocus)
+
+		SLATE_END_ARGS()
+
+	public:
+
+		/**
+		 * Construct this widget.
+		 *
+		 * @param InArgs The declaration data for this widget.
+		 * @param InModel The model for the anim timeline
+		 */
+
+		void Construct(const FArguments& InArgs, FMLDeformerEditorToolkit* InEditor);
+		/** SWidget interface */
+		virtual FReply OnMouseButtonUp(const FGeometry& MyGeometry, const FPointerEvent& MouseEvent) override;
+
+		/** Compute a major grid interval and number of minor divisions to display */
+		bool GetGridMetrics(float PhysicalWidth, double& OutMajorInterval, int32& OutMinorDivisions) const;
+
+		/** Get the time slider controller */
+		TSharedPtr<ITimeSliderController> GetTimeSliderController() const;
+
+		void SetModel(TWeakPtr<FMLDeformerEditorModel> InModel);
+	private:
+		/**
+		 * @return The fill percentage of the animation outliner
+		 */
+		float GetColumnFillCoefficient(int32 ColumnIndex) const
+		{
+			return ColumnFillCoefficients[ColumnIndex];
+		}
+
+		/** Get numeric Type interface for converting between frame numbers and display formats. */
+		TSharedRef<INumericTypeInterface<double>> GetNumericTypeInterface() const;
+
+		/** Called when a column fill percentage is changed by a splitter slot. */
+		void OnColumnFillCoefficientChanged(float FillCoefficient, int32 ColumnIndex);
+
+		/** Handles an additive layer key being set */
+
+		void HandleScrubPositionChanged(FFrameTime NewScrubPosition, bool bIsScrubbing, bool bEvaluate);
+
+		void HandleViewRangeChanged(TRange<double> InRange, EViewRangeInterpolation InInterpolation);
+
+		double GetSpinboxDelta() const;
+
+		void SetPlayTime(double InFrameTime);
+
+		void SetDisplayFormat(EFrameNumberDisplayFormats InFormat);
+		bool IsDisplayFormatChecked(EFrameNumberDisplayFormats InFormat) const;
+
+	private:
+		/** Anim timeline model */
+		TWeakPtr<FMLDeformerEditorModel> Model;
+
+		/** The currently active editor model. */
+		UE::MLDeformer::FMLDeformerEditorToolkit* Editor;
+
+		/** The time slider controller */
+		TSharedPtr<FMLTimeSliderController> TimeSliderController;
+
+		TSharedPtr<SAnimTimelineTransportControls> TransportControls;
+		/** The top time slider widget */
+		TSharedPtr<ITimeSlider> TopTimeSlider;
+
+		/** The fill coefficients of each column in the grid. */
+		float ColumnFillCoefficients[2];
+
+		/** Called when the user has begun dragging the selection selection range */
+		FSimpleDelegate OnSelectionRangeBeginDrag;
+
+		/** Called when the user has finished dragging the selection selection range */
+		FSimpleDelegate OnSelectionRangeEndDrag;
+
+		/** Called when any widget contained within the anim timeline has received focus */
+		FSimpleDelegate OnReceivedFocus;
+
+		/** Numeric Type interface for converting between frame numbers and display formats. */
+		TSharedPtr<INumericTypeInterface<double>> NumericTypeInterface;
+
+		/** The view range */
+		TAttribute<FAnimatedRange> ViewRange;
+
+		/** Filter text used to search the tree */
+		FText FilterText;
+	};
+
+}// namespace UE::MLDeformer
diff --git a/Engine/Plugins/Experimental/MLDeformer/MLDeformerFramework/Source/MLDeformerFrameworkEditor/Public/MLDeformerEditorActor.h b/Engine/Plugins/Experimental/MLDeformer/MLDeformerFramework/Source/MLDeformerFrameworkEditor/Public/MLDeformerEditorActor.h
--- a/Engine/Plugins/Experimental/MLDeformer/MLDeformerFramework/Source/MLDeformerFrameworkEditor/Public/MLDeformerEditorActor.h
+++ b/Engine/Plugins/Experimental/MLDeformer/MLDeformerFramework/Source/MLDeformerFrameworkEditor/Public/MLDeformerEditorActor.h
@@ -40,6 +40,8 @@ namespace UE::MLDeformer
 		virtual void SetPlayPosition(float TimeInSeconds, bool bAutoPause=true);
 		virtual float GetPlayPosition() const;
 		virtual void SetPlaySpeed(float PlaySpeed);
+		virtual float GetPlaySpeed();
+
 		virtual FBox GetBoundingBox() const;
 		virtual void Pause(bool bPaused);
 		virtual bool IsGroundTruthActor() const;
diff --git a/Engine/Plugins/Experimental/MLDeformer/MLDeformerFramework/Source/MLDeformerFrameworkEditor/Public/MLDeformerEditorModel.h b/Engine/Plugins/Experimental/MLDeformer/MLDeformerFramework/Source/MLDeformerFrameworkEditor/Public/MLDeformerEditorModel.h
--- a/Engine/Plugins/Experimental/MLDeformer/MLDeformerFramework/Source/MLDeformerFrameworkEditor/Public/MLDeformerEditorModel.h
+++ b/Engine/Plugins/Experimental/MLDeformer/MLDeformerFramework/Source/MLDeformerFrameworkEditor/Public/MLDeformerEditorModel.h
@@ -7,6 +7,8 @@
 #include "UObject/ObjectMacros.h"
 #include "UObject/ObjectPtr.h"
 #include "MLDeformerEditorActor.h"
+#include "Misc/FrameTime.h"
+#include "MLDeformerVizSettings.h"
 
 class UMLDeformerModel;
 class UMLDeformerInputInfo;
@@ -15,6 +17,7 @@ class IPersonaPreviewScene;
 class FEditorViewportClient;
 class FSceneView;
 class FViewport;
+
 class FPrimitiveDrawInterface;
 class UWorld;
 class UMeshDeformer;
@@ -54,7 +57,7 @@ namespace UE::MLDeformer
 
 	/** The base class for the editor side of an UMLDeformerModel. */
 	class MLDEFORMERFRAMEWORKEDITOR_API FMLDeformerEditorModel
-		: public FGCObject
+		: public TSharedFromThis<FMLDeformerEditorModel>, public FGCObject
 	{
 	public:
 		/** 
@@ -98,11 +101,12 @@ namespace UE::MLDeformer
 		virtual void OnInputAssetsChanged();
 		virtual void OnPostInputAssetChanged();
 		virtual void HandleDefaultPropertyChanges(FPropertyChangedEvent& PropertyChangedEvent);
-		virtual void OnPlayButtonPressed();
+		virtual void OnPlayPressed();
 		virtual void OnPreTraining() {}
 		virtual void OnPostTraining(ETrainingResult TrainingResult, bool bUsePartiallyTrainedWhenAborted);
 		virtual void OnTrainingAborted(bool bUsePartiallyTrainedData) {}
 		virtual bool IsPlayingAnim() const;
+		virtual bool IsPlayingForward() const;
 		virtual double CalcTrainingTimelinePosition() const;
 		virtual double CalcTestTimelinePosition() const;
 		virtual void OnTimeSliderScrubPositionChanged(double NewScrubTime, bool bIsScrubbing);
@@ -133,6 +137,51 @@ namespace UE::MLDeformer
 		virtual FString GetDefaultDeformerGraphAssetPath() const;
 		virtual FString GetTrainedNetworkOnnxFile() const;
 
+		/** Get the current view range */
+		TRange<double> GetViewRange() const;
+
+		/** Set the current view range */
+		void SetViewRange(TRange<double> InRange);
+
+		/** Get the working range of the model's data */
+		TRange<double> GetWorkingRange() const;
+
+		/** Get the playback range of the model's data */
+		TRange<FFrameNumber> GetPlaybackRange() const;
+
+		/** Get the current scrub position */
+		FFrameNumber GetTickResScrubPosition() const;
+
+		int32 GetTicksPerFrame() const;
+		/** Get the current scrub time */
+		float GetScrubTime() const;
+
+		/** Set the current scrub position */
+		void SetScrubPosition(FFrameTime NewScrubPostion);
+
+		/** Set the current scrub position */
+		void SetScrubPosition(FFrameNumber NewScrubPostion);
+
+		/** Set if frames are displayed */
+		void SetDisplayFrames(bool bDisplayFrames);
+
+		bool IsDisplayingFrames() const;
+
+		void HandleModelChanged();
+		void HandleVizModeChanged(EMLDeformerVizMode Mode);
+
+		/** Handle the view range being changed */
+		void HandleViewRangeChanged(TRange<double> InRange);
+
+		/** Handle the working range being changed */
+		void HandleWorkingRangeChanged(TRange<double> InRange);
+
+		/** Get the framerate specified by the anim sequence */
+		double GetFrameRate() const;
+
+		/** Get the tick resolution we are displaying at */
+		int32 GetTickResolution() const;
+
 		FMLDeformerEditorToolkit* GetEditor() const { return Editor; }
 		UMLDeformerModel* GetModel() const { return Model.Get(); }
 		UWorld* GetWorld() const;
@@ -200,6 +249,13 @@ namespace UE::MLDeformer
 		 */
 		void CompressEngineMorphTargets(FMorphTargetVertexInfoBuffers& OutMorphBuffers, const TArray<UMorphTarget*>& MorphTargets, int32 LOD = 0, float MorphErrorTolerance = 0.01f);
 
+		FMLDeformerEditorActor* GetVisualizationModeBaseActor() const;
+
+		const UAnimSequence* GetAnimSequence() const;
+
+		double CalcTimelinePosition() const;
+
+		void UpdateRanges();
 	protected:
 		/** The runtime model associated with this editor model. */
 		TObjectPtr<UMLDeformerModel> Model = nullptr;
@@ -236,6 +292,19 @@ namespace UE::MLDeformer
 
 		/** Do we need to resample all input/output data? */
 		bool bNeedToResampleInputOutputs = true;
+
+		/** The range we are currently viewing */
+		TRange<double> ViewRange;
+
+		/** The working range of this model, encompassing the view range */
+		TRange<double> WorkingRange;
+
+		/** The playback range of this model for each timeframe */
+		TRange<double> PlaybackRange;
+
+		FFrameTime ScrubPosition;
+
+		bool bDisplayFrames = true;
 	};
 
 	template<class TrainingModelClass>
diff --git a/Engine/Plugins/Experimental/MLDeformer/MLDeformerFramework/Source/MLDeformerFrameworkEditor/Public/MLDeformerEditorToolkit.h b/Engine/Plugins/Experimental/MLDeformer/MLDeformerFramework/Source/MLDeformerFrameworkEditor/Public/MLDeformerEditorToolkit.h
--- a/Engine/Plugins/Experimental/MLDeformer/MLDeformerFramework/Source/MLDeformerFrameworkEditor/Public/MLDeformerEditorToolkit.h
+++ b/Engine/Plugins/Experimental/MLDeformer/MLDeformerFramework/Source/MLDeformerFrameworkEditor/Public/MLDeformerEditorToolkit.h
@@ -20,6 +20,8 @@ class SSimpleTimeSlider;
 
 namespace UE::MLDeformer
 {
+	class SMLDeformerTimeline;
+
 	namespace MLDeformerEditorModes
 	{
 		extern const FName Editor;
@@ -73,12 +75,15 @@ namespace UE::MLDeformer
 		IDetailsView* GetModelDetailsView() const;
 		IDetailsView* GetVizSettingsDetailsView() const;
 
-		void SetTimeSlider(TSharedPtr<SSimpleTimeSlider> InTimeSlider);
-		SSimpleTimeSlider* GetTimeSlider() const;
+		void SetTimeSlider(TSharedPtr<SMLDeformerTimeline> InTimeSlider);
+		SMLDeformerTimeline* GetTimeSlider() const;
 
 		UMLDeformerAsset* GetDeformerAsset() const { return DeformerAsset.Get(); }
 		FMLDeformerEditorModel* GetActiveModel() { return ActiveModel.Get(); }
 
+		TWeakPtr<FMLDeformerEditorModel> GetActiveModelPointer(){ return TWeakPtr<FMLDeformerEditorModel>(ActiveModel); }
+		const FMLDeformerEditorModel* GetActiveModel() const { return ActiveModel.Get(); }
+
 		double CalcTimelinePosition() const;
 		void OnTimeSliderScrubPositionChanged(double NewScrubTime, bool bIsScrubbing);
 		void UpdateTimeSliderRange();
@@ -122,7 +127,8 @@ namespace UE::MLDeformer
 		TSharedPtr<IDetailsView> VizSettingsDetailsView = nullptr;
 
 		/** The timeline slider widget. */
-		TSharedPtr<SSimpleTimeSlider> TimeSlider = nullptr;
+		TSharedPtr<SMLDeformerTimeline> TimeSlider = nullptr;
+
 
 		/** The currently active editor model. */
 		TSharedPtr<FMLDeformerEditorModel> ActiveModel = nullptr;
diff --git a/Engine/Plugins/Experimental/MLDeformer/NeuralMorphModel/Source/NeuralMorphModelEditor/Private/NeuralMorphEditorModelActor.cpp b/Engine/Plugins/Experimental/MLDeformer/NeuralMorphModel/Source/NeuralMorphModelEditor/Private/NeuralMorphEditorModelActor.cpp
--- a/Engine/Plugins/Experimental/MLDeformer/NeuralMorphModel/Source/NeuralMorphModelEditor/Private/NeuralMorphEditorModelActor.cpp
+++ b/Engine/Plugins/Experimental/MLDeformer/NeuralMorphModel/Source/NeuralMorphModelEditor/Private/NeuralMorphEditorModelActor.cpp
@@ -88,7 +88,11 @@ namespace UE::NeuralMorphModel
 
 		if (GeomCacheComponent)
 		{
-			return GeomCacheComponent->GetAnimationTime();
+			float Duration = GeomCacheComponent->GetDuration();
+			float AnimTime = GeomCacheComponent->GetAnimationTime();
+			float StartTime = GeomCacheComponent->GetStartTimeOffset();
+			float DeltaTime = AnimTime - StartTime;
+			return DeltaTime > Duration ? FMath::Fmod(DeltaTime, Duration) : DeltaTime;
 		}
 
 		return 0.0f;
diff --git a/Engine/Plugins/Experimental/MLDeformer/VertexDeltaModel/Source/VertexDeltaModelEditor/Private/VertexDeltaEditorModelActor.cpp b/Engine/Plugins/Experimental/MLDeformer/VertexDeltaModel/Source/VertexDeltaModelEditor/Private/VertexDeltaEditorModelActor.cpp
--- a/Engine/Plugins/Experimental/MLDeformer/VertexDeltaModel/Source/VertexDeltaModelEditor/Private/VertexDeltaEditorModelActor.cpp
+++ b/Engine/Plugins/Experimental/MLDeformer/VertexDeltaModel/Source/VertexDeltaModelEditor/Private/VertexDeltaEditorModelActor.cpp
@@ -88,7 +88,11 @@ namespace UE::VertexDeltaModel
 
 		if (GeomCacheComponent)
 		{
-			return GeomCacheComponent->GetAnimationTime();
+			float Duration = GeomCacheComponent->GetDuration(); 
+			float AnimTime =  GeomCacheComponent->GetAnimationTime();
+			float StartTime = GeomCacheComponent->GetStartTimeOffset();
+			float DeltaTime = AnimTime - StartTime; 
+			return DeltaTime > Duration ? FMath::Fmod(DeltaTime, Duration) : DeltaTime;
 		}
 
 		return 0.0f;
