[MLDeformer]
- Models can now specify whether they support bones and/or curves.
- Checkboxes will be shown rather than a combobox, to select bones/curves/bones+curves.
Models that only support one input, such as bones only, will not show the checkbox.
- Bones and curves show their counts next to the group title, and it shows whether curves are found or not, and whether all bones/curves are included or not.
- Fixed a bug where the selected model type wasn't the one that would be created.

diff --git a/Engine/Plugins/Experimental/MLDeformer/LegacyVertexDeltaModel/Source/LegacyVertexDeltaModelEditor/Private/LegacyVertexDeltaEditorModel.cpp b/Engine/Plugins/Experimental/MLDeformer/LegacyVertexDeltaModel/Source/LegacyVertexDeltaModelEditor/Private/LegacyVertexDeltaEditorModel.cpp
--- a/Engine/Plugins/Experimental/MLDeformer/LegacyVertexDeltaModel/Source/LegacyVertexDeltaModelEditor/Private/LegacyVertexDeltaEditorModel.cpp
+++ b/Engine/Plugins/Experimental/MLDeformer/LegacyVertexDeltaModel/Source/LegacyVertexDeltaModelEditor/Private/LegacyVertexDeltaEditorModel.cpp
@@ -23,11 +23,6 @@ namespace UE::LegacyVertexDeltaModel
 {
 	using namespace UE::MLDeformer;
 
-	FLegacyVertexDeltaEditorModel::~FLegacyVertexDeltaEditorModel()
-	{
-		Model->OnPostEditChangeProperty().Unbind();
-	}
-
 	FMLDeformerEditorModel* FLegacyVertexDeltaEditorModel::MakeInstance()
 	{
 		return new FLegacyVertexDeltaEditorModel();
@@ -45,12 +40,6 @@ namespace UE::LegacyVertexDeltaModel
 		return NewSampler;
 	}
 
-	void FLegacyVertexDeltaEditorModel::Init(const InitSettings& InitSettings)
-	{
-		FMLDeformerEditorModel::Init(InitSettings);
-		Model->OnPostEditChangeProperty().BindRaw(this, &FLegacyVertexDeltaEditorModel::OnPostEditChangeProperty);
-	}
-
 	void FLegacyVertexDeltaEditorModel::OnPostEditChangeProperty(FPropertyChangedEvent& PropertyChangedEvent)
 	{
 		const FProperty* Property = PropertyChangedEvent.Property;
diff --git a/Engine/Plugins/Experimental/MLDeformer/LegacyVertexDeltaModel/Source/LegacyVertexDeltaModelEditor/Public/LegacyVertexDeltaEditorModel.h b/Engine/Plugins/Experimental/MLDeformer/LegacyVertexDeltaModel/Source/LegacyVertexDeltaModelEditor/Public/LegacyVertexDeltaEditorModel.h
--- a/Engine/Plugins/Experimental/MLDeformer/LegacyVertexDeltaModel/Source/LegacyVertexDeltaModelEditor/Public/LegacyVertexDeltaEditorModel.h
+++ b/Engine/Plugins/Experimental/MLDeformer/LegacyVertexDeltaModel/Source/LegacyVertexDeltaModelEditor/Public/LegacyVertexDeltaEditorModel.h
@@ -26,8 +26,6 @@ namespace UE::LegacyVertexDeltaModel
 		: public UE::MLDeformer::FMLDeformerEditorModel
 	{
 	public:
-		virtual ~FLegacyVertexDeltaEditorModel() override;
-
 		// We need to implement this static MakeInstance method.
 		static FMLDeformerEditorModel* MakeInstance();
 
@@ -37,7 +35,6 @@ namespace UE::LegacyVertexDeltaModel
 		// ~END FGCObject overrides.
 	
 		// FMLDeformerEditorModel overrides.
-		virtual void Init(const InitSettings& Settings) override;
 		virtual FMLDeformerEditorActor* CreateEditorActor(const FMLDeformerEditorActor::FConstructSettings& Settings) const override;
 		virtual FMLDeformerSampler* CreateSampler() const;
 		virtual void CreateTrainingGroundTruthActor(UWorld* World) override;
@@ -52,7 +49,8 @@ namespace UE::LegacyVertexDeltaModel
 		virtual ETrainingResult Train() override;
 		virtual FString GetDefaultDeformerGraphAssetPath() const override;
 		virtual FString GetHeatMapDeformerGraphPath() const override;
-		FString GetTrainedNetworkOnnxFile() const override;
+		virtual FString GetTrainedNetworkOnnxFile() const override;
+		virtual void OnPostEditChangeProperty(FPropertyChangedEvent& PropertyChangedEvent) override;
 		// ~END FMLDeformerEditorModel overrides.
 
 		// Some helpers that cast to this model's variants of some classes.
@@ -60,8 +58,6 @@ namespace UE::LegacyVertexDeltaModel
 		ULegacyVertexDeltaModelVizSettings* GetVertexDeltaModelVizSettings() const;
 		FLegacyVertexDeltaEditorModelActor* FindVertexDeltaModelEditorActor(int32 TypeID) const;
 		FMLDeformerGeomCacheSampler* GetGeomCacheSampler() const { return static_cast<FMLDeformerGeomCacheSampler*>(Sampler); }
-	
-		void OnPostEditChangeProperty(FPropertyChangedEvent& PropertyChangedEvent);
 
 	protected:
 		void CreateGeomCacheActor(UWorld* World, int32 ActorID, const FName& Name, UGeometryCache* GeomCache, FLinearColor LabelColor, FLinearColor WireframeColor, const FText& LabelText, bool bIsTrainingActor);
diff --git a/Engine/Plugins/Experimental/MLDeformer/MLDeformerFramework/Source/MLDeformerFramework/Private/MLDeformerGeomCacheHelpers.cpp b/Engine/Plugins/Experimental/MLDeformer/MLDeformerFramework/Source/MLDeformerFramework/Private/MLDeformerGeomCacheHelpers.cpp
--- a/Engine/Plugins/Experimental/MLDeformer/MLDeformerFramework/Source/MLDeformerFramework/Private/MLDeformerGeomCacheHelpers.cpp
+++ b/Engine/Plugins/Experimental/MLDeformer/MLDeformerFramework/Source/MLDeformerFramework/Private/MLDeformerGeomCacheHelpers.cpp
@@ -224,7 +224,7 @@ namespace UE::MLDeformer
 				} // If the track name matches the skeletal meshes internal mesh name.
 			} // For all meshes in the Skeletal Mesh.
 
-			if (!bFoundMatch)
+			if (Track && !bFoundMatch)
 			{
 				OutFailedImportedMeshNames.Add(Track->GetName());
 				UE_LOG(LogMLDeformer, Warning, TEXT("Geometry cache '%s' cannot be matched with a mesh inside the Skeletal Mesh."), *Track->GetName());
diff --git a/Engine/Plugins/Experimental/MLDeformer/MLDeformerFramework/Source/MLDeformerFramework/Private/MLDeformerModel.cpp b/Engine/Plugins/Experimental/MLDeformer/MLDeformerFramework/Source/MLDeformerFramework/Private/MLDeformerModel.cpp
--- a/Engine/Plugins/Experimental/MLDeformer/MLDeformerFramework/Source/MLDeformerFramework/Private/MLDeformerModel.cpp
+++ b/Engine/Plugins/Experimental/MLDeformer/MLDeformerFramework/Source/MLDeformerFramework/Private/MLDeformerModel.cpp
@@ -140,6 +140,8 @@ USkeleton* UMLDeformerModel::GetSkeleton(bool& bInvalidSkeletonIsError, const IP
 
 void UMLDeformerModel::BeginDestroy()
 {
+	PostEditPropertyDelegate.RemoveAll(this);
+
 	BeginReleaseResource(&VertexMapBuffer);
 	RenderResourceDestroyFence.BeginFence();
 	Super::BeginDestroy();
@@ -184,7 +186,7 @@ void UMLDeformerModel::FloatArrayToVector3Array(const TArray<float>& FloatArray,
 			return;
 		}
 
-		OnPostEditChangeProperty().ExecuteIfBound(PropertyChangedEvent);
+		OnPostEditChangeProperty().Broadcast(PropertyChangedEvent);
 	}
 
 	void UMLDeformerModel::UpdateCachedNumVertices()
diff --git a/Engine/Plugins/Experimental/MLDeformer/MLDeformerFramework/Source/MLDeformerFramework/Private/MLDeformerModelInstance.cpp b/Engine/Plugins/Experimental/MLDeformer/MLDeformerFramework/Source/MLDeformerFramework/Private/MLDeformerModelInstance.cpp
--- a/Engine/Plugins/Experimental/MLDeformer/MLDeformerFramework/Source/MLDeformerFramework/Private/MLDeformerModelInstance.cpp
+++ b/Engine/Plugins/Experimental/MLDeformer/MLDeformerFramework/Source/MLDeformerFramework/Private/MLDeformerModelInstance.cpp
@@ -36,22 +36,25 @@ void UMLDeformerModelInstance::Init(USkeletalMeshComponent* SkelMeshComponent)
 		return;
 	}
 
-	USkeletalMesh* SkelMesh = SkelMeshComponent->GetSkeletalMeshAsset();
-	if (SkelMesh)
+	if (Model->DoesSupportBones())
 	{
-		// Init the bone mapping table.
-		const UMLDeformerInputInfo* InputInfo = Model->GetInputInfo();
-		const int32 NumAssetBones = InputInfo->GetNumBones();
-		AssetBonesToSkelMeshMappings.Reset();
-		AssetBonesToSkelMeshMappings.AddUninitialized(NumAssetBones);
-		TrainingBoneTransforms.SetNumUninitialized(NumAssetBones);
-
-		// For each bone in the deformer asset, find the matching bone index inside the skeletal mesh component.
-		for (int32 Index = 0; Index < NumAssetBones; ++Index)
+		USkeletalMesh* SkelMesh = SkelMeshComponent->GetSkeletalMeshAsset();
+		if (SkelMesh)
 		{
-			const FName BoneName = InputInfo->GetBoneName(Index);
-			const int32 SkelMeshBoneIndex = SkeletalMeshComponent->GetBaseComponent()->GetBoneIndex(BoneName);
-			AssetBonesToSkelMeshMappings[Index] = SkelMeshBoneIndex;
+			// Init the bone mapping table.
+			const UMLDeformerInputInfo* InputInfo = Model->GetInputInfo();
+			const int32 NumAssetBones = InputInfo->GetNumBones();
+			AssetBonesToSkelMeshMappings.Reset();
+			AssetBonesToSkelMeshMappings.AddUninitialized(NumAssetBones);
+			TrainingBoneTransforms.SetNumUninitialized(NumAssetBones);
+
+			// For each bone in the deformer asset, find the matching bone index inside the skeletal mesh component.
+			for (int32 Index = 0; Index < NumAssetBones; ++Index)
+			{
+				const FName BoneName = InputInfo->GetBoneName(Index);
+				const int32 SkelMeshBoneIndex = SkeletalMeshComponent->GetBaseComponent()->GetBoneIndex(BoneName);
+				AssetBonesToSkelMeshMappings[Index] = SkelMeshBoneIndex;
+			}
 		}
 	}
 
@@ -234,8 +237,16 @@ int64 UMLDeformerModelInstance::SetNeuralNetworkInputValues(float* InputData, in
 
 	// Feed data to the network inputs.
 	int64 BufferOffset = 0;
-	BufferOffset = SetBoneTransforms(InputData, NumInputFloats, BufferOffset);
-	BufferOffset = SetCurveValues(InputData, NumInputFloats, BufferOffset);
+	if (Model->DoesSupportBones())
+	{
+		BufferOffset = SetBoneTransforms(InputData, NumInputFloats, BufferOffset);
+	}
+
+	if (Model->DoesSupportCurves())
+	{
+		BufferOffset = SetCurveValues(InputData, NumInputFloats, BufferOffset);
+	}
+
 	return BufferOffset;
 }
 
diff --git a/Engine/Plugins/Experimental/MLDeformer/MLDeformerFramework/Source/MLDeformerFramework/Public/MLDeformerModel.h b/Engine/Plugins/Experimental/MLDeformer/MLDeformerFramework/Source/MLDeformerFramework/Public/MLDeformerModel.h
--- a/Engine/Plugins/Experimental/MLDeformer/MLDeformerFramework/Source/MLDeformerFramework/Public/MLDeformerModel.h
+++ b/Engine/Plugins/Experimental/MLDeformer/MLDeformerFramework/Source/MLDeformerFramework/Public/MLDeformerModel.h
@@ -35,16 +35,7 @@ namespace UE::MLDeformer
 	};
 }	// namespace UE::MLDeformer
 
-/** The training inputs. Specifies what data to include in training. */
-UENUM()
-enum class EMLDeformerTrainingInputFilter : uint8
-{
-	BonesAndCurves = 0, /** Include both bone rotations and curve values. */
-	BonesOnly,			/** Include only bone rotations. */
-	CurvesOnly			/** Include only curve values. */
-};
-
-DECLARE_DELEGATE_OneParam(FMLDeformerModelOnPostEditProperty, FPropertyChangedEvent&)
+DECLARE_EVENT_OneParam(UMLDeformerModel, FMLDeformerModelOnPostEditProperty, FPropertyChangedEvent&)
 
 /**
  * The ML Deformer runtime model base class.
@@ -67,6 +58,8 @@ public:
 	virtual UMLDeformerModelInstance* CreateModelInstance(UMLDeformerComponent* Component);
 	virtual FString GetDisplayName() const PURE_VIRTUAL(UMLDeformerModel::GetDisplayName, return FString(););
 	virtual void PostMLDeformerComponentInit(UMLDeformerModelInstance* ModelInstance) {}
+	virtual bool DoesSupportBones() const { return true; }
+	virtual bool DoesSupportCurves() const { return true; }
 	virtual bool IsNeuralNetworkOnGPU() const { return true; }	// GPU neural network.
 #if WITH_EDITORONLY_DATA
 	virtual bool HasTrainingGroundTruth() const { return false; }
@@ -91,7 +84,6 @@ public:
 	// ~END IBoneReferenceSkeletonProvider overrides.
 
 	UMLDeformerAsset* GetDeformerAsset() const;
-	EMLDeformerTrainingInputFilter GetTrainingInputs() const { return TrainingInputs; }
 	UMLDeformerInputInfo* GetInputInfo() const { return InputInfo.Get(); }
 
 	int32 GetNumBaseMeshVerts() const { return NumBaseMeshVerts; }
@@ -103,6 +95,7 @@ public:
 	UNeuralNetwork* GetNeuralNetwork() const { return NeuralNetwork.Get(); }
 	void SetNeuralNetwork(UNeuralNetwork* InNeuralNetwork);
 
+
 #if WITH_EDITORONLY_DATA
 	// UObject overrides.
 	virtual void PostEditChangeProperty(FPropertyChangedEvent& PropertyChangedEvent) override;
@@ -110,7 +103,10 @@ public:
 
 	void InitVertexMap();
 
-	FMLDeformerModelOnPostEditProperty& OnPostEditChangeProperty() { return PostEditPropertyEvent; }
+	bool ShouldIncludeBonesInTraining() const { return bIncludeBones; }
+	bool ShouldIncludeCurvesInTraining() const { return bIncludeCurves; }
+
+	FMLDeformerModelOnPostEditProperty& OnPostEditChangeProperty() { return PostEditPropertyDelegate; }
 
 	UMLDeformerVizSettings* GetVizSettings() const { return VizSettings; }
 	const USkeletalMesh* GetSkeletalMesh() const { return SkeletalMesh;  }
@@ -128,7 +124,7 @@ public:
 
 protected:
 	TObjectPtr<UMLDeformerAsset> DeformerAsset = nullptr;
-	FMLDeformerModelOnPostEditProperty PostEditPropertyEvent;
+	FMLDeformerModelOnPostEditProperty PostEditPropertyDelegate;
 
 	void SetInputInfo(UMLDeformerInputInfo* Input) { InputInfo = Input; }
 	void FloatArrayToVector3Array(const TArray<float>& FloatArray, TArray<FVector3f>& OutVectorArray);
@@ -142,10 +138,6 @@ public:
 	UPROPERTY()
 	int32 NumTargetMeshVerts = 0;
 
-	/** Describes what inputs we should train the neural network on. */
-	UPROPERTY(EditAnywhere, Category = "Inputs and Output")
-	EMLDeformerTrainingInputFilter TrainingInputs = EMLDeformerTrainingInputFilter::BonesOnly;
-
 	/** 
 	 * The information about the neural network inputs. This contains things such as bone names and curve names.
 	 */
@@ -174,6 +166,14 @@ public:
 	UPROPERTY()
 	TObjectPtr<UMLDeformerVizSettings> VizSettings = nullptr;
 
+	/** Include bones as inputs to training process? */
+	UPROPERTY(EditAnywhere, Category = "Inputs and Output")
+	bool bIncludeBones = true;
+
+	/** Include curves as inputs to the training process? */
+	UPROPERTY(EditAnywhere, Category = "Inputs and Output")
+	bool bIncludeCurves = false;
+
 	/** The skeletal mesh that represents the linear skinned mesh. */
 	UPROPERTY(EditAnywhere, BlueprintReadWrite, Category = "Base Mesh")
 	TObjectPtr<USkeletalMesh> SkeletalMesh = nullptr;
@@ -190,11 +190,11 @@ public:
 	FTransform AlignmentTransform = FTransform::Identity;
 
 	/** The bones to include during training. When none are provided, all bones of the Skeleton will be included. */
-	UPROPERTY(EditAnywhere, Category = "Inputs and Output", meta = (EditCondition = "TrainingInputs == EMLDeformerTrainingInputFilter::BonesAndCurves || TrainingInputs == EMLDeformerTrainingInputFilter::BonesOnly"))
+	UPROPERTY(EditAnywhere, Category = "Inputs and Output", meta = (EditCondition = "bIncludeBones"))
 	TArray<FBoneReference> BoneIncludeList;
 
 	/** The curves to include during training. When none are provided, all curves of the Skeleton will be included. */
-	UPROPERTY(EditAnywhere, Category = "Inputs and Output", meta = (EditCondition = "TrainingInputs == EMLDeformerTrainingInputFilter::BonesAndCurves || TrainingInputs == EMLDeformerTrainingInputFilter::CurvesOnly"))
+	UPROPERTY(EditAnywhere, Category = "Inputs and Output", meta = (EditCondition = "bIncludeCurves"))
 	TArray<FMLDeformerCurveReference> CurveIncludeList;
 
 	/** The maximum numer of training frames (samples) to train on. Use this to train on a sub-section of your full training data. */
diff --git a/Engine/Plugins/Experimental/MLDeformer/MLDeformerFramework/Source/MLDeformerFrameworkEditor/Private/MLDeformerEditorModel.cpp b/Engine/Plugins/Experimental/MLDeformer/MLDeformerFramework/Source/MLDeformerFrameworkEditor/Private/MLDeformerEditorModel.cpp
--- a/Engine/Plugins/Experimental/MLDeformer/MLDeformerFramework/Source/MLDeformerFrameworkEditor/Private/MLDeformerEditorModel.cpp
+++ b/Engine/Plugins/Experimental/MLDeformer/MLDeformerFramework/Source/MLDeformerFrameworkEditor/Private/MLDeformerEditorModel.cpp
@@ -39,6 +39,11 @@ namespace UE::MLDeformer
 {
 	FMLDeformerEditorModel::~FMLDeformerEditorModel()
 	{
+		if (PostEditPropertyDelegateHandle.IsValid())
+		{
+			Model->OnPostEditChangeProperty().Remove(PostEditPropertyDelegateHandle);
+		}
+
 		DeleteEditorActors();
 
 		FMLDeformerEditorModule& EditorModule = FModuleManager::GetModuleChecked<FMLDeformerEditorModule>("MLDeformerFrameworkEditor");
@@ -63,6 +68,8 @@ namespace UE::MLDeformer
 		ViewRange = TRange<double>(0.0, 100.0);
 		WorkingRange = TRange<double>(0.0, 100.0);
 		PlaybackRange = TRange<double>(0.0, 100.0);
+
+		PostEditPropertyDelegateHandle = Model->OnPostEditChangeProperty().AddRaw(this, &FMLDeformerEditorModel::OnPostEditChangeProperty);
 	}
 
 	void FMLDeformerEditorModel::UpdateEditorInputInfo()
@@ -607,7 +614,9 @@ namespace UE::MLDeformer
 		{
 			TriggerInputAssetChanged();
 		}
-		else if (Property->GetFName() == GET_MEMBER_NAME_CHECKED(UMLDeformerModel, TrainingInputs))
+		else
+		if (Property->GetFName() == GET_MEMBER_NAME_CHECKED(UMLDeformerModel, bIncludeBones) ||
+			Property->GetFName() == GET_MEMBER_NAME_CHECKED(UMLDeformerModel, bIncludeCurves))
 		{
 			if (PropertyChangedEvent.ChangeType == EPropertyChangeType::ValueSet)
 			{
@@ -621,6 +630,7 @@ namespace UE::MLDeformer
 			Property->GetFName() == GET_MEMBER_NAME_CHECKED(UMLDeformerModel, CurveIncludeList))
 		{
 			UpdateEditorInputInfo();
+			GetEditor()->GetModelDetailsView()->ForceRefresh();
 		}
 		else if (Property->GetFName() == GET_MEMBER_NAME_CHECKED(UMLDeformerVizSettings, AnimPlaySpeed))
 		{
@@ -883,13 +893,23 @@ namespace UE::MLDeformer
 	{
 		if (Model->SkeletalMesh && GetEditorInputInfo()->IsEmpty())
 		{
-			switch (Model->TrainingInputs)
+			FString ErrorString;
+			if (Model->DoesSupportBones() && Model->ShouldIncludeBonesInTraining())
 			{
-				case EMLDeformerTrainingInputFilter::BonesOnly:			return FText(LOCTEXT("InputsEmptyBonesErrorText", "Your base mesh has no bones to train on."));
-				case EMLDeformerTrainingInputFilter::CurvesOnly:		return FText(LOCTEXT("InputsEmptyCurvesErrorText", "Your base mesh has no curves to train on."));
-				case EMLDeformerTrainingInputFilter::BonesAndCurves:	return FText(LOCTEXT("InputsEmptyBonesCurvesErrorText", "Your base mesh has no bones or curves to train on."));
-				default: return FText(LOCTEXT("InputsEmptyDefaultErrorText", "There are no inputs to train on. There are no bones, curves or other inputs we can use."));
+				ErrorString += FText(LOCTEXT("InputsEmptyBonesErrorText", "Your base mesh has no bones to train on.\n")).ToString();
 			}
+			else
+			if (Model->DoesSupportCurves() && Model->ShouldIncludeCurvesInTraining())
+			{
+				ErrorString += FText(LOCTEXT("InputsEmptyCurvesErrorText", "Your base mesh has no curves to train on.\n")).ToString();
+			}
+			else
+			{
+				ErrorString += FText(LOCTEXT("InputsEmptyCurvesErrorText", "The training process needs inputs.\n")).ToString();
+			}
+
+			ErrorString.RemoveFromEnd("\n");
+			return FText::FromString(ErrorString);
 		}
 
 		return FText();
@@ -945,8 +965,8 @@ namespace UE::MLDeformer
 		InputInfo->SetNumBaseVertices(Model->GetNumBaseMeshVerts());
 		InputInfo->SetNumTargetVertices(Model->GetNumTargetMeshVerts());
 
-		const bool bIncludeBones = (Model->GetTrainingInputs() == EMLDeformerTrainingInputFilter::BonesAndCurves || Model->GetTrainingInputs() == EMLDeformerTrainingInputFilter::BonesOnly);
-		const bool bIncludeCurves = (Model->GetTrainingInputs() == EMLDeformerTrainingInputFilter::BonesAndCurves || Model->GetTrainingInputs() == EMLDeformerTrainingInputFilter::CurvesOnly);
+		const bool bIncludeBones = Model->DoesSupportBones() && Model->ShouldIncludeBonesInTraining();
+		const bool bIncludeCurves = Model->DoesSupportCurves() && Model->ShouldIncludeCurvesInTraining();
 		const USkeleton* Skeleton = Model->GetSkeletalMesh() ? Model->GetSkeletalMesh()->GetSkeleton() : nullptr;
 
 		// Handle bones.
@@ -1111,6 +1131,7 @@ namespace UE::MLDeformer
 			Model->BoneIncludeList.Empty();
 			UE_LOG(LogMLDeformer, Warning, TEXT("There are no animated bone rotations in Anim Sequence '%s'."), *(Model->AnimSequence->GetName()));
 		}
+		UpdateEditorInputInfo();
 	}
 
 	void FMLDeformerEditorModel::InitCurveIncludeListToAnimatedCurvesOnly()
@@ -1193,6 +1214,7 @@ namespace UE::MLDeformer
 			Model->CurveIncludeList.Empty();
 			UE_LOG(LogMLDeformer, Warning, TEXT("There are no animated curves in Anim Sequence '%s'."), *(Model->AnimSequence->GetName()));
 		}
+		UpdateEditorInputInfo();
 	}
 
 	void FMLDeformerEditorModel::Render(const FSceneView* View, FViewport* Viewport, FPrimitiveDrawInterface* PDI)
@@ -1675,6 +1697,26 @@ namespace UE::MLDeformer
 		}
 	}
 
+	int32 FMLDeformerEditorModel::GetNumCurvesOnSkeletalMesh(USkeletalMesh* SkelMesh) const
+	{
+		if (SkelMesh)
+		{
+			const USkeleton* Skeleton = SkelMesh->GetSkeleton();
+			if (Skeleton)
+			{
+				const FSmartNameMapping* Mapping = Skeleton->GetSmartNameContainer(USkeleton::AnimCurveMappingName);
+				if (Mapping)
+				{
+					TArray<FName> SkeletonCurveNames;
+					Mapping->FillNameArray(SkeletonCurveNames);
+					return SkeletonCurveNames.Num();
+				}
+			}
+		}
+
+		return 0;
+	}
+
 	FString FMLDeformerEditorModel::GetHeatMapMaterialPath() const
 	{
 		return FString(TEXT("/MLDeformerFramework/Materials/MLDeformerHeatMapMat.MLDeformerHeatMapMat"));
diff --git a/Engine/Plugins/Experimental/MLDeformer/MLDeformerFramework/Source/MLDeformerFrameworkEditor/Private/MLDeformerEditorToolkit.cpp b/Engine/Plugins/Experimental/MLDeformer/MLDeformerFramework/Source/MLDeformerFrameworkEditor/Private/MLDeformerEditorToolkit.cpp
--- a/Engine/Plugins/Experimental/MLDeformer/MLDeformerFramework/Source/MLDeformerFrameworkEditor/Private/MLDeformerEditorToolkit.cpp
+++ b/Engine/Plugins/Experimental/MLDeformer/MLDeformerFramework/Source/MLDeformerFrameworkEditor/Private/MLDeformerEditorToolkit.cpp
@@ -214,9 +214,9 @@ namespace UE::MLDeformer
 			int32 UnsortedModelIndex = -1;
 			for (int32 Index = 0; Index < ModelTypes.Num(); ++Index)
 			{
-				const UClass* Model = ModelTypes[ModelIndex];
+				const UClass* Model = ModelTypes[Index];
 				const UMLDeformerModel* DefaultModel = Cast<UMLDeformerModel>(Model->GetDefaultObject());
-				if (DefaultModel->GetDisplayName() == ModelNames[Index])
+				if (DefaultModel->GetDisplayName() == ModelNames[ModelIndex])
 				{
 					UnsortedModelIndex = Index;
 					break;
diff --git a/Engine/Plugins/Experimental/MLDeformer/MLDeformerFramework/Source/MLDeformerFrameworkEditor/Private/MLDeformerModelDetails.cpp b/Engine/Plugins/Experimental/MLDeformer/MLDeformerFramework/Source/MLDeformerFrameworkEditor/Private/MLDeformerModelDetails.cpp
--- a/Engine/Plugins/Experimental/MLDeformer/MLDeformerFramework/Source/MLDeformerFrameworkEditor/Private/MLDeformerModelDetails.cpp
+++ b/Engine/Plugins/Experimental/MLDeformer/MLDeformerFramework/Source/MLDeformerFrameworkEditor/Private/MLDeformerModelDetails.cpp
@@ -162,9 +162,13 @@ namespace UE::MLDeformer
 
 		TargetMeshCategoryBuilder->AddProperty(GET_MEMBER_NAME_CHECKED(UMLDeformerModel, AlignmentTransform), UMLDeformerModel::StaticClass());
 
-		// Input and output.
-		InputOutputCategoryBuilder->AddProperty(GET_MEMBER_NAME_CHECKED(UMLDeformerModel, TrainingInputs), UMLDeformerModel::StaticClass());
+		InputOutputCategoryBuilder->AddProperty(GET_MEMBER_NAME_CHECKED(UMLDeformerModel, bIncludeBones), UMLDeformerModel::StaticClass())
+			.Visibility(IsBonesFlagVisible() ? EVisibility::Visible : EVisibility::Collapsed);
 
+		InputOutputCategoryBuilder->AddProperty(GET_MEMBER_NAME_CHECKED(UMLDeformerModel, bIncludeCurves), UMLDeformerModel::StaticClass())
+			.Visibility(IsCurvesFlagVisible() ? EVisibility::Visible : EVisibility::Collapsed);
+
+		AddTrainingInputFlags();
 		AddTrainingInputErrors();
 
 		const FText ErrorText = EditorModel->GetInputsErrorText();
@@ -185,36 +189,122 @@ namespace UE::MLDeformer
 		InputOutputCategoryBuilder->AddProperty(GET_MEMBER_NAME_CHECKED(UMLDeformerModel, DeltaCutoffLength), UMLDeformerModel::StaticClass());
 
 		// Bone include list group.
-		IDetailGroup& BoneIncludeGroup = InputOutputCategoryBuilder->AddGroup("BoneIncludeGroup", LOCTEXT("BoneIncludeGroup", "Bones"), false, false);
-		BoneIncludeGroup.AddWidgetRow()
+		if (Model->DoesSupportBones())
+		{
+			const int NumBones = EditorModel->GetEditorInputInfo() ? EditorModel->GetEditorInputInfo()->GetNumBones() : 0;
+			const FText BonesText = FText::Format(FTextFormat(LOCTEXT("BonesGroupName", "Bones ({0})")), NumBones);;
+			FText AllBonesText;
+			if (Model->GetSkeletalMesh() && (Model->GetSkeletalMesh()->GetRefSkeleton().GetNum() == NumBones || NumBones == 0))
+			{
+				AllBonesText = LOCTEXT("BoneGroupValue", "All Bones Included");
+			}
+			IDetailGroup& BoneIncludeGroup = InputOutputCategoryBuilder->AddGroup("BoneIncludeGroup", FText(), false, false);
+			BoneIncludeGroup.HeaderRow()
+			.NameContent()
+			[
+				SNew(SHorizontalBox)
+				+ SHorizontalBox::Slot()
+				.HAlign(HAlign_Right)
+				[
+					SNew(STextBlock)
+					.Text(BonesText)
+					.Font(IDetailLayoutBuilder::GetDetailFont())
+				]
+			]
 			.ValueContent()
 			[
-				SNew(SButton)
-				.HAlign(HAlign_Center)
-				.VAlign(VAlign_Center)
-				.Text(LOCTEXT("AnimatedBonesButton", "Animated Bones Only"))
-				.OnClicked(FOnClicked::CreateSP(this, &FMLDeformerModelDetails::OnFilterAnimatedBonesOnly))
-				.IsEnabled_Lambda([this](){ return (Model->GetTrainingInputs() == EMLDeformerTrainingInputFilter::BonesAndCurves) || (Model->GetTrainingInputs() == EMLDeformerTrainingInputFilter::BonesOnly); })
+				SNew(SHorizontalBox)
+				+ SHorizontalBox::Slot()
+				.HAlign(HAlign_Left)
+				[
+					SNew(STextBlock)
+					.Text(AllBonesText)
+					.Font(IDetailLayoutBuilder::GetDetailFontItalic())
+				]
 			];
-		BoneIncludeGroup.AddPropertyRow(DetailBuilder.GetProperty(GET_MEMBER_NAME_CHECKED(UMLDeformerModel, BoneIncludeList), UMLDeformerModel::StaticClass()));
+			BoneIncludeGroup.AddWidgetRow()
+				.ValueContent()
+				[
+					SNew(SButton)
+					.HAlign(HAlign_Center)
+					.VAlign(VAlign_Center)
+					.Text(LOCTEXT("AnimatedBonesButton", "Animated Bones Only"))
+					.ToolTipText(LOCTEXT("AnimatedBonesButtonTooltip", "Initialize the bone include list to bones that are animated."))
+					.OnClicked(FOnClicked::CreateSP(this, &FMLDeformerModelDetails::OnFilterAnimatedBonesOnly))
+					.IsEnabled_Lambda([this](){ return Model->ShouldIncludeBonesInTraining(); })
+				];
+			BoneIncludeGroup.AddPropertyRow(DetailBuilder.GetProperty(GET_MEMBER_NAME_CHECKED(UMLDeformerModel, BoneIncludeList), UMLDeformerModel::StaticClass()));
 
-		AddBoneInputErrors();
+			AddBoneInputErrors();
+		}
+		else
+		{
+			InputOutputCategoryBuilder->AddProperty(GET_MEMBER_NAME_CHECKED(UMLDeformerModel, BoneIncludeList), UMLDeformerModel::StaticClass())
+				.Visibility(EVisibility::Collapsed);
+		}
 
 		// Curve include list group.
-		IDetailGroup& CurveIncludeGroup = InputOutputCategoryBuilder->AddGroup("CurveIncludeGroup", LOCTEXT("CurveIncludeGroup", "Curves"), false, false);
-		CurveIncludeGroup.AddWidgetRow()
+		if (Model->DoesSupportCurves())
+		{
+			const int NumCurves = EditorModel->GetEditorInputInfo() ? EditorModel->GetEditorInputInfo()->GetNumCurves() : 0;
+			const FText CurvesText = FText::Format(FTextFormat(LOCTEXT("CurvesGroupName", "Curves ({0})")), NumCurves);
+			FText AllCurvesText;
+			const int32 NumCurvesOnSkelMesh = EditorModel->GetNumCurvesOnSkeletalMesh(Model->GetSkeletalMesh());
+			if (NumCurvesOnSkelMesh == 0)
+			{
+				AllCurvesText = LOCTEXT("CurvesGroupValueNoCurves", "No Curves Found");
+			}
+			else if (NumCurves == NumCurvesOnSkelMesh)
+			{
+				AllCurvesText = LOCTEXT("CurvesGroupValue", "All Curves Included");
+			}
+			IDetailGroup& CurvesIncludeGroup = InputOutputCategoryBuilder->AddGroup("CurveIncludeGroup", LOCTEXT("CurveIncludeGroup", "Curves"), false, false);
+			CurvesIncludeGroup.HeaderRow()
+			.NameContent()
+			[
+				SNew(SHorizontalBox)
+				+ SHorizontalBox::Slot()
+				.HAlign(HAlign_Right)
+				[
+					SNew(STextBlock)
+					.Text(CurvesText)
+					.Font(IDetailLayoutBuilder::GetDetailFont())
+				]
+			]
 			.ValueContent()
 			[
-				SNew(SButton)
-				.HAlign(HAlign_Center)
-				.VAlign(VAlign_Center)
-				.Text(LOCTEXT("AnimatedCurvesButton", "Animated Curves Only"))
-				.OnClicked(FOnClicked::CreateSP(this, &FMLDeformerModelDetails::OnFilterAnimatedCurvesOnly))
-				.IsEnabled_Lambda([this](){ return (Model->GetTrainingInputs() == EMLDeformerTrainingInputFilter::BonesAndCurves) || (Model->GetTrainingInputs() == EMLDeformerTrainingInputFilter::CurvesOnly); })
+				SNew(SHorizontalBox)
+				+ SHorizontalBox::Slot()
+				.HAlign(HAlign_Left)
+				[
+					SNew(STextBlock)
+					.Text(AllCurvesText)
+					.Font(IDetailLayoutBuilder::GetDetailFontItalic())
+				]
 			];
-		CurveIncludeGroup.AddPropertyRow(DetailBuilder.GetProperty(GET_MEMBER_NAME_CHECKED(UMLDeformerModel, CurveIncludeList), UMLDeformerModel::StaticClass()));
 
-		AddCurveInputErrors();
+			CurvesIncludeGroup.AddWidgetRow()
+				.ValueContent()
+				[
+					SNew(SButton)
+					.HAlign(HAlign_Center)
+					.VAlign(VAlign_Center)
+					.Text(LOCTEXT("AnimatedCurvesButton", "Animated Curves Only"))
+					.ToolTipText(LOCTEXT("AnimatedCurvesButtonTooltip", "Initialize the curve include list to curves that are animated."))
+					.OnClicked(FOnClicked::CreateSP(this, &FMLDeformerModelDetails::OnFilterAnimatedCurvesOnly))
+					.IsEnabled_Lambda([this](){ return Model->ShouldIncludeCurvesInTraining() && EditorModel->GetNumCurvesOnSkeletalMesh(Model->GetSkeletalMesh()) > 0; })
+				];
+			CurvesIncludeGroup.AddPropertyRow(DetailBuilder.GetProperty(GET_MEMBER_NAME_CHECKED(UMLDeformerModel, CurveIncludeList), UMLDeformerModel::StaticClass()));
+
+			AddCurveInputErrors();
+		}
+		else
+		{
+			InputOutputCategoryBuilder->AddProperty(GET_MEMBER_NAME_CHECKED(UMLDeformerModel, CurveIncludeList), UMLDeformerModel::StaticClass())
+				.Visibility(EVisibility::Collapsed);
+		}
+
+		AddTrainingInputFilters();
 
 		// Show a warning when no neural network has been set.
 		{		
@@ -291,6 +381,15 @@ namespace UE::MLDeformer
 				]
 			];
 	}
+	bool FMLDeformerModelDetails::IsBonesFlagVisible() const
+	{
+		return (Model->DoesSupportBones() && Model->DoesSupportCurves());
+	}
+
+	bool FMLDeformerModelDetails::IsCurvesFlagVisible() const
+	{
+		return (Model->DoesSupportBones() && Model->DoesSupportCurves() && EditorModel->GetNumCurvesOnSkeletalMesh(Model->GetSkeletalMesh()) > 0);
+	}
 }	// namespace UE::MLDeformer
 
 #undef LOCTEXT_NAMESPACE
diff --git a/Engine/Plugins/Experimental/MLDeformer/MLDeformerFramework/Source/MLDeformerFrameworkEditor/Public/MLDeformerEditorModel.h b/Engine/Plugins/Experimental/MLDeformer/MLDeformerFramework/Source/MLDeformerFrameworkEditor/Public/MLDeformerEditorModel.h
--- a/Engine/Plugins/Experimental/MLDeformer/MLDeformerFramework/Source/MLDeformerFrameworkEditor/Public/MLDeformerEditorModel.h
+++ b/Engine/Plugins/Experimental/MLDeformer/MLDeformerFramework/Source/MLDeformerFrameworkEditor/Public/MLDeformerEditorModel.h
@@ -105,6 +105,7 @@ namespace UE::MLDeformer
 		virtual void OnPreTraining() {}
 		virtual void OnPostTraining(ETrainingResult TrainingResult, bool bUsePartiallyTrainedWhenAborted);
 		virtual void OnTrainingAborted(bool bUsePartiallyTrainedData) {}
+		virtual void OnPostEditChangeProperty(FPropertyChangedEvent& PropertyChangedEvent) {}
 		virtual bool IsPlayingAnim() const;
 		virtual bool IsPlayingForward() const;
 		virtual double CalcTrainingTimelinePosition() const;
@@ -206,6 +207,7 @@ namespace UE::MLDeformer
 
 		void InitBoneIncludeListToAnimatedBonesOnly();
 		void InitCurveIncludeListToAnimatedCurvesOnly();
+		int32 GetNumCurvesOnSkeletalMesh(USkeletalMesh* SkelMesh) const;
 
 		UNeuralNetwork* LoadNeuralNetworkFromOnnx(const FString& Filename) const;
 		int32 GetCurrentTrainingFrame() const { return CurrentTrainingFrame; }
@@ -281,6 +283,9 @@ namespace UE::MLDeformer
 		/** The heatmap deformer graph. */
 		TObjectPtr<UMeshDeformer> HeatMapDeformerGraph = nullptr;
 
+		/** The delegate handle to the post edit property event. */
+		FDelegateHandle PostEditPropertyDelegateHandle;
+
 		/** The current training frame. */
 		int32 CurrentTrainingFrame = -1;
 
diff --git a/Engine/Plugins/Experimental/MLDeformer/MLDeformerFramework/Source/MLDeformerFrameworkEditor/Public/MLDeformerModelDetails.h b/Engine/Plugins/Experimental/MLDeformer/MLDeformerFramework/Source/MLDeformerFrameworkEditor/Public/MLDeformerModelDetails.h
--- a/Engine/Plugins/Experimental/MLDeformer/MLDeformerFramework/Source/MLDeformerFrameworkEditor/Public/MLDeformerModelDetails.h
+++ b/Engine/Plugins/Experimental/MLDeformer/MLDeformerFramework/Source/MLDeformerFrameworkEditor/Public/MLDeformerModelDetails.h
@@ -31,9 +31,13 @@ namespace UE::MLDeformer
 		virtual void AddTargetMesh() {}
 		virtual void AddBaseMeshErrors() {}
 		virtual void AddAnimSequenceErrors() {}
-		virtual void AddTrainingInputErrors() {}
 		virtual void AddBoneInputErrors() {}
 		virtual void AddCurveInputErrors() {}
+		virtual void AddTrainingInputFlags() {}
+		virtual void AddTrainingInputFilters() {}
+		virtual void AddTrainingInputErrors() {}
+		virtual bool IsBonesFlagVisible() const;
+		virtual bool IsCurvesFlagVisible() const;
 
 	protected:
 		bool FilterAnimSequences(const FAssetData& AssetData, USkeleton* Skeleton);
diff --git a/Engine/Plugins/Experimental/MLDeformer/NeuralMorphModel/Source/NeuralMorphModelEditor/Private/NeuralMorphEditorModel.cpp b/Engine/Plugins/Experimental/MLDeformer/NeuralMorphModel/Source/NeuralMorphModelEditor/Private/NeuralMorphEditorModel.cpp
--- a/Engine/Plugins/Experimental/MLDeformer/NeuralMorphModel/Source/NeuralMorphModelEditor/Private/NeuralMorphEditorModel.cpp
+++ b/Engine/Plugins/Experimental/MLDeformer/NeuralMorphModel/Source/NeuralMorphModelEditor/Private/NeuralMorphEditorModel.cpp
@@ -27,7 +27,7 @@ namespace UE::NeuralMorphModel
 
 	FNeuralMorphEditorModel::~FNeuralMorphEditorModel()
 	{
-		Model->OnPostEditChangeProperty().Unbind();
+		Model->OnPostEditChangeProperty().Remove(PostEditPropertyDelegateHandle);
 	}
 
 	FMLDeformerEditorModel* FNeuralMorphEditorModel::MakeInstance()
@@ -47,12 +47,6 @@ namespace UE::NeuralMorphModel
 		return NewSampler;
 	}
 
-	void FNeuralMorphEditorModel::Init(const InitSettings& InitSettings)
-	{
-		FMLDeformerEditorModel::Init(InitSettings);
-		Model->OnPostEditChangeProperty().BindRaw(this, &FNeuralMorphEditorModel::OnPostEditChangeProperty);
-	}
-
 	void FNeuralMorphEditorModel::OnPostEditChangeProperty(FPropertyChangedEvent& PropertyChangedEvent)
 	{
 		const FProperty* Property = PropertyChangedEvent.Property;
diff --git a/Engine/Plugins/Experimental/MLDeformer/NeuralMorphModel/Source/NeuralMorphModelEditor/Public/NeuralMorphEditorModel.h b/Engine/Plugins/Experimental/MLDeformer/NeuralMorphModel/Source/NeuralMorphModelEditor/Public/NeuralMorphEditorModel.h
--- a/Engine/Plugins/Experimental/MLDeformer/NeuralMorphModel/Source/NeuralMorphModelEditor/Public/NeuralMorphEditorModel.h
+++ b/Engine/Plugins/Experimental/MLDeformer/NeuralMorphModel/Source/NeuralMorphModelEditor/Public/NeuralMorphEditorModel.h
@@ -36,7 +36,6 @@ namespace UE::NeuralMorphModel
 		// ~END FGCObject overrides.
 	
 		// FMLDeformerEditorModel overrides.
-		virtual void Init(const InitSettings& Settings) override;
 		virtual FMLDeformerEditorActor* CreateEditorActor(const FMLDeformerEditorActor::FConstructSettings& Settings) const override;
 		virtual FMLDeformerSampler* CreateSampler() const;
 		virtual void CreateTrainingGroundTruthActor(UWorld* World) override;
@@ -55,6 +54,7 @@ namespace UE::NeuralMorphModel
 		virtual FString GetHeatMapDeformerGraphPath() const override;
 		virtual void OnPreTraining() override;
 		virtual void OnPostTraining(ETrainingResult TrainingResult, bool bUsePartiallyTrainedWhenAborted) override;
+		virtual void OnPostEditChangeProperty(FPropertyChangedEvent& PropertyChangedEvent) override;
 		// ~END FMLDeformerEditorModel overrides.
 
 		// Some helpers that cast to this model's variants of some classes.
@@ -63,7 +63,6 @@ namespace UE::NeuralMorphModel
 		FNeuralMorphEditorModelActor* FindNeuralMorphModelEditorActor(int32 TypeID) const;
 		FMLDeformerGeomCacheSampler* GetGeomCacheSampler() const { return static_cast<FMLDeformerGeomCacheSampler*>(Sampler); }
 	
-		void OnPostEditChangeProperty(FPropertyChangedEvent& PropertyChangedEvent);
 
 	protected:
 		void CreateGeomCacheActor(UWorld* World, int32 ActorID, const FName& Name, UGeometryCache* GeomCache, FLinearColor LabelColor, FLinearColor WireframeColor, const FText& LabelText, bool bIsTrainingActor);
diff --git a/Engine/Plugins/Experimental/MLDeformer/VertexDeltaModel/Source/VertexDeltaModelEditor/Private/VertexDeltaEditorModel.cpp b/Engine/Plugins/Experimental/MLDeformer/VertexDeltaModel/Source/VertexDeltaModelEditor/Private/VertexDeltaEditorModel.cpp
--- a/Engine/Plugins/Experimental/MLDeformer/VertexDeltaModel/Source/VertexDeltaModelEditor/Private/VertexDeltaEditorModel.cpp
+++ b/Engine/Plugins/Experimental/MLDeformer/VertexDeltaModel/Source/VertexDeltaModelEditor/Private/VertexDeltaEditorModel.cpp
@@ -23,11 +23,6 @@ namespace UE::VertexDeltaModel
 {
 	using namespace UE::MLDeformer;
 
-	FVertexDeltaEditorModel::~FVertexDeltaEditorModel()
-	{
-		Model->OnPostEditChangeProperty().Unbind();
-	}
-
 	FMLDeformerEditorModel* FVertexDeltaEditorModel::MakeInstance()
 	{
 		return new FVertexDeltaEditorModel();
@@ -45,12 +40,6 @@ namespace UE::VertexDeltaModel
 		return NewSampler;
 	}
 
-	void FVertexDeltaEditorModel::Init(const InitSettings& InitSettings)
-	{
-		FMLDeformerEditorModel::Init(InitSettings);
-		Model->OnPostEditChangeProperty().BindRaw(this, &FVertexDeltaEditorModel::OnPostEditChangeProperty);
-	}
-
 	void FVertexDeltaEditorModel::OnPostEditChangeProperty(FPropertyChangedEvent& PropertyChangedEvent)
 	{
 		const FProperty* Property = PropertyChangedEvent.Property;
diff --git a/Engine/Plugins/Experimental/MLDeformer/VertexDeltaModel/Source/VertexDeltaModelEditor/Public/VertexDeltaEditorModel.h b/Engine/Plugins/Experimental/MLDeformer/VertexDeltaModel/Source/VertexDeltaModelEditor/Public/VertexDeltaEditorModel.h
--- a/Engine/Plugins/Experimental/MLDeformer/VertexDeltaModel/Source/VertexDeltaModelEditor/Public/VertexDeltaEditorModel.h
+++ b/Engine/Plugins/Experimental/MLDeformer/VertexDeltaModel/Source/VertexDeltaModelEditor/Public/VertexDeltaEditorModel.h
@@ -26,8 +26,6 @@ namespace UE::VertexDeltaModel
 		: public UE::MLDeformer::FMLDeformerEditorModel
 	{
 	public:
-		virtual ~FVertexDeltaEditorModel() override;
-
 		// We need to implement this static MakeInstance method.
 		static FMLDeformerEditorModel* MakeInstance();
 
@@ -37,7 +35,6 @@ namespace UE::VertexDeltaModel
 		// ~END FGCObject overrides.
 	
 		// FMLDeformerEditorModel overrides.
-		virtual void Init(const InitSettings& Settings) override;
 		virtual FMLDeformerEditorActor* CreateEditorActor(const FMLDeformerEditorActor::FConstructSettings& Settings) const override;
 		virtual FMLDeformerSampler* CreateSampler() const;
 		virtual void CreateTrainingGroundTruthActor(UWorld* World) override;
@@ -50,7 +47,8 @@ namespace UE::VertexDeltaModel
 		virtual int32 GetTestFrameAtTime(double TimeInSeconds) const override;
 		virtual void UpdateIsReadyForTrainingState() override;
 		virtual ETrainingResult Train() override;
-		FString GetTrainedNetworkOnnxFile() const override;
+		virtual FString GetTrainedNetworkOnnxFile() const override;
+		virtual void OnPostEditChangeProperty(FPropertyChangedEvent& PropertyChangedEvent) override;
 		// ~END FMLDeformerEditorModel overrides.
 
 		// Some helpers that cast to this model's variants of some classes.
@@ -58,8 +56,6 @@ namespace UE::VertexDeltaModel
 		UVertexDeltaModelVizSettings* GetVertexDeltaModelVizSettings() const;
 		FVertexDeltaEditorModelActor* FindVertexDeltaModelEditorActor(int32 TypeID) const;
 		FMLDeformerGeomCacheSampler* GetGeomCacheSampler() const { return static_cast<FMLDeformerGeomCacheSampler*>(Sampler); }
-	
-		void OnPostEditChangeProperty(FPropertyChangedEvent& PropertyChangedEvent);
 
 	protected:
 		void CreateGeomCacheActor(UWorld* World, int32 ActorID, const FName& Name, UGeometryCache* GeomCache, FLinearColor LabelColor, FLinearColor WireframeColor, const FText& LabelText, bool bIsTrainingActor);
