[MLDeformer] Localization fixes and some reformatting / improvements on the timeline code.

diff --git a/Engine/Plugins/Experimental/MLDeformer/MLDeformerFramework/Source/MLDeformerFrameworkEditor/Private/MLDeformerEditorModel.cpp b/Engine/Plugins/Experimental/MLDeformer/MLDeformerFramework/Source/MLDeformerFrameworkEditor/Private/MLDeformerEditorModel.cpp
--- a/Engine/Plugins/Experimental/MLDeformer/MLDeformerFramework/Source/MLDeformerFrameworkEditor/Private/MLDeformerEditorModel.cpp
+++ b/Engine/Plugins/Experimental/MLDeformer/MLDeformerFramework/Source/MLDeformerFrameworkEditor/Private/MLDeformerEditorModel.cpp
@@ -901,11 +901,11 @@ namespace UE::MLDeformer
 			else
 			if (Model->DoesSupportCurves() && Model->ShouldIncludeCurvesInTraining())
 			{
-				ErrorString += FText(LOCTEXT("InputsEmptyCurvesErrorText", "Your base mesh has no curves to train on.\n")).ToString();
+				ErrorString += FText(LOCTEXT("InputsNoCurvesToTrainText", "Your base mesh has no curves to train on.\n")).ToString();
 			}
 			else
 			{
-				ErrorString += FText(LOCTEXT("InputsEmptyCurvesErrorText", "The training process needs inputs.\n")).ToString();
+				ErrorString += FText(LOCTEXT("InputsNeededErrorText", "The training process needs inputs.\n")).ToString();
 			}
 
 			ErrorString.RemoveFromEnd("\n");
diff --git a/Engine/Plugins/Experimental/MLDeformer/MLDeformerFramework/Source/MLDeformerFrameworkEditor/Private/SMLDeformerTimeline.cpp b/Engine/Plugins/Experimental/MLDeformer/MLDeformerFramework/Source/MLDeformerFrameworkEditor/Private/SMLDeformerTimeline.cpp
--- a/Engine/Plugins/Experimental/MLDeformer/MLDeformerFramework/Source/MLDeformerFrameworkEditor/Private/SMLDeformerTimeline.cpp
+++ b/Engine/Plugins/Experimental/MLDeformer/MLDeformerFramework/Source/MLDeformerFrameworkEditor/Private/SMLDeformerTimeline.cpp
@@ -35,38 +35,37 @@
 
 namespace UE::MLDeformer
 {
-
 	struct FPaintSectionAreaViewArgs
 	{
 		FPaintSectionAreaViewArgs()
 			: bDisplayTickLines(false), bDisplayScrubPosition(false)
 		{}
 
-		/** Whether to display tick lines */
+		/** Whether to display tick lines. */
 		bool bDisplayTickLines;
-		/** Whether to display the scrub position */
+		/** Whether to display the scrub position. */
 		bool bDisplayScrubPosition;
-		/** Optional Paint args for the playback range*/
+		/** Optional Paint args for the playback range. */
 		TOptional<FPaintPlaybackRangeArgs> PlaybackRangeArgs;
 	};
 
 	/**
-	 * A time slider controller for the anim timeline
+	 * A time slider controller for the anim timeline.
 	 */
-	class FMLTimeSliderController : public ITimeSliderController
+	class FMLTimeSliderController 
+		: public ITimeSliderController
 	{
 	public:
 		FMLTimeSliderController(const FTimeSliderArgs& InArgs, TWeakPtr<FMLDeformerEditorModel> InWeakModel, TWeakPtr<SMLDeformerTimeline> InWeakTimeline);
 
 		/**
-		* Determines the optimal spacing between tick marks in the slider for a given pixel density
-		* Increments until a minimum amount of slate units specified by MinTick is reached
-		*
-		* @param InPixelsPerInput	The density of pixels between each input
-		* @param MinTick			The minimum slate units per tick allowed
-		* @param MinTickSpacing	The minimum tick spacing in time units allowed
-		* @return the optimal spacing in time units
-		*/
+		 * Determines the optimal spacing between tick marks in the slider for a given pixel density.
+		 * Increments until a minimum amount of slate units specified by MinTick is reached.
+		 * @param InPixelsPerInput The density of pixels between each input.
+		 * @param MinTick The minimum slate units per tick allowed.
+		 * @param MinTickSpacing The minimum tick spacing in time units allowed.
+		 * @return The optimal spacing in time units.
+		 */
 		float DetermineOptimalSpacing(float InPixelsPerInput, uint32 MinTick, float MinTickSpacing) const;
 		void SetModel(TWeakPtr<FMLDeformerEditorModel> InModel);
 		virtual int32 OnPaintTimeSlider(bool bMirrorLabels, const FGeometry& AllottedGeometry, const FSlateRect& MyCullingRect, FSlateWindowElementList& OutDrawElements, int32 LayerId, const FWidgetStyle& InWidgetStyle, bool bParentEnabled) const override;
@@ -88,77 +87,69 @@ namespace UE::MLDeformer
 		virtual void SetScrubPosition(FFrameTime InTime, bool bEvaluate) { TimeSliderArgs.OnScrubPositionChanged.ExecuteIfBound(InTime, false, bEvaluate); }
 
 		/**
-		 * Clamp the given range to the clamp range
-		 *
-		 * @param NewRangeMin		The new lower bound of the range
-		 * @param NewRangeMax		The new upper bound of the range
+		 * Clamp the given range to the clamp range.
+		 * @param NewRangeMin The new lower bound of the range.
+		 * @param NewRangeMax The new upper bound of the range.
 		 */
 		void ClampViewRange(double& NewRangeMin, double& NewRangeMax);
 
 		/**
 		 * Zoom the range by a given delta.
-		 *
-		 * @param InDelta		The total amount to zoom by (+ve = zoom out, -ve = zoom in)
-		 * @param ZoomBias		Bias to apply to lower/upper extents of the range. (0 = lower, 0.5 = equal, 1 = upper)
+		 * @param InDelta The total amount to zoom by (+ve = zoom out, -ve = zoom in).
+		 * @param ZoomBias Bias to apply to lower/upper extents of the range (0 = lower, 0.5 = equal, 1 = upper).
 		 */
 		bool ZoomByDelta(float InDelta, float ZoomBias = 0.5f);
 
 		/**
-		 * Pan the range by a given delta
-		 *
-		 * @param InDelta		The total amount to pan by (+ve = pan forwards in time, -ve = pan backwards in time)
+		 * Pan the range by a given delta.
+		 * @param InDelta The total amount to pan by (+ve = pan forwards in time, -ve = pan backwards in time).
 		 */
 		void PanByDelta(float InDelta);
 
-		/** Determine frame time from a mouse position */
+		/** Determine frame time from a mouse position. */
 		FFrameTime GetFrameTimeFromMouse(const FGeometry& Geometry, FVector2D ScreenSpacePosition) const;
 
 	private:
-		// forward declared as class members to prevent name collision with similar types defined in other units
+		// Forward declared as class members to prevent name collision with similar types defined in other units.
 		struct FDrawTickArgs;
 		struct FScrubRangeToScreen;
 
 		/**
-		 * Call this method when the user's interaction has changed the scrub position
-		 *
-		 * @param NewValue				Value resulting from the user's interaction
-		 * @param bIsScrubbing			True if done via scrubbing, false if just releasing scrubbing
+		 * Call this method when the user's interaction has changed the scrub position.
+		 * @param NewValue Value resulting from the user's interaction.
+		 * @param bIsScrubbing True if done via scrubbing, false if just releasing scrubbing.
 		 */
 		void CommitScrubPosition(FFrameTime NewValue, bool bIsScrubbing);
 
 		/**
-		 * Draw time tick marks
-		 *
-		 * @param OutDrawElements	List to add draw elements to
-		 * @param ViewRange			The currently visible time range in seconds
-		 * @param RangeToScreen		Time range to screen space converter
-		 * @param InArgs			Parameters for drawing the tick lines
+		 * Draw time tick marks.
+		 * @param OutDrawElements List to add draw elements to.
+		 * @param ViewRange The currently visible time range in seconds.
+		 * @param RangeToScreen Time range to screen space converter.
+		 * @param InArgs Parameters for drawing the tick lines.
 		 */
 		void DrawTicks(FSlateWindowElementList& OutDrawElements, const TRange<double>& ViewRange, const FScrubRangeToScreen& RangeToScreen, FDrawTickArgs& InArgs) const;
 
 		/**
 		 * Draw the selection range.
-		 *
 		 * @return The new layer ID.
 		 */
 		int32 DrawSelectionRange(const FGeometry& AllottedGeometry, const FSlateRect& MyCullingRect, FSlateWindowElementList& OutDrawElements, int32 LayerId, const FScrubRangeToScreen& RangeToScreen, const FPaintPlaybackRangeArgs& Args) const;
 
 		/**
 		 * Draw the playback range.
-		 *
-		 * @return the new layer ID
+		 * @return the new layer ID.
 		 */
 		int32 DrawPlaybackRange(const FGeometry& AllottedGeometry, const FSlateRect& MyCullingRect, FSlateWindowElementList& OutDrawElements, int32 LayerId, const FScrubRangeToScreen& RangeToScreen, const FPaintPlaybackRangeArgs& Args) const;
 
 	private:
-
 		/**
-		 * Hit test the lower bound of a range
+		 * Hit test the lower bound of a range.
 		 */
 		bool HitTestRangeStart(const FScrubRangeToScreen& RangeToScreen, const TRange<double>& Range, float HitPixel) const;
 
 		/**
-		 * Hit test the upper bound of a range
+		 * Hit test the upper bound of a range.
 		 */
 		bool HitTestRangeEnd(const FScrubRangeToScreen& RangeToScreen, const TRange<double>& Range, float HitPixel) const;
 
@@ -169,8 +160,8 @@ namespace UE::MLDeformer
 		void SetSelectionRangeEnd(FFrameNumber NewEnd);
 
 		FFrameTime ComputeFrameTimeFromMouse(const FGeometry& Geometry, FVector2D ScreenSpacePosition, FScrubRangeToScreen RangeToScreen, bool CheckSnapping = true) const;
-	private:
 
+	private:
 		struct FScrubPixelRange
 		{
 			TRange<float> Range;
@@ -179,36 +170,34 @@ namespace UE::MLDeformer
 		};
 
 		FScrubPixelRange GetHitTestScrubberPixelRange(FFrameTime ScrubTime, const FScrubRangeToScreen& RangeToScreen) const;
-
 		FScrubPixelRange GetScrubberPixelRange(FFrameTime ScrubTime, const FScrubRangeToScreen& RangeToScreen) const;
 		FScrubPixelRange GetScrubberPixelRange(FFrameTime ScrubTime, FFrameRate Resolution, FFrameRate PlayRate, const FScrubRangeToScreen& RangeToScreen, float DilationPixels = 0.f) const;
 
 	private:
-
-		/** Pointer back to the model object */
+		/** Pointer back to the model object. */
 		TWeakPtr<FMLDeformerEditorModel> WeakModel;
 
-		/** Pointer back to the timeline */
+		/** Pointer back to the timeline. */
 		TWeakPtr<SMLDeformerTimeline> WeakTimeline;
 
 		FTimeSliderArgs TimeSliderArgs;
 
-		/** Brush for drawing the fill area on the scrubber */
+		/** Brush for drawing the fill area on the scrubber. */
 		const FSlateBrush* ScrubFillBrush;
 
-		/** Brush for drawing an upwards facing scrub handles */
+		/** Brush for drawing an upwards facing scrub handles. */
 		const FSlateBrush* ScrubHandleUpBrush;
 
-		/** Brush for drawing a downwards facing scrub handle */
+		/** Brush for drawing a downwards facing scrub handle. */
 		const FSlateBrush* ScrubHandleDownBrush;
 
-		/** Brush for drawing an editable time */
+		/** Brush for drawing an editable time. */
 		const FSlateBrush* EditableTimeBrush;
 
-		/** Total mouse delta during dragging **/
+		/** Total mouse delta during dragging. **/
 		float DistanceDragged;
 
-		/** If we are dragging a scrubber or dragging to set the time range */
+		/** If we are dragging a scrubber or dragging to set the time range. */
 		enum DragType
 		{
 			DRAG_SCRUBBING_TIME,
@@ -222,33 +211,31 @@ namespace UE::MLDeformer
 
 		DragType MouseDragType;
 
-		/** If we are currently panning the panel */
+		/** If we are currently panning the panel. */
 		bool bPanning;
 
-		/** Index of the current dragged time */
+		/** Index of the current dragged time. */
 		int32 DraggedTimeIndex;
 
-		/** Mouse down position range */
+		/** Mouse down position range. */
 		FVector2D MouseDownPosition[2];
 
-		/** Geometry on mouse down */
+		/** Geometry on mouse down. */
 		FGeometry MouseDownGeometry;
 
-		/** Range stack */
+		/** Range stack. */
 		TArray<TRange<double>> ViewRangeStack;
-
 	};
 
-	/** Utility struct for converting between scrub range space and local/absolute screen space */
+	/** Utility struct for converting between scrub range space and local/absolute screen space. */
 	struct FMLTimeSliderController::FScrubRangeToScreen
 	{
 		double ViewStart;
-		float  PixelsPerInput;
+		float PixelsPerInput;
 
 		FScrubRangeToScreen(const TRange<double>& InViewInput, const FVector2D& InWidgetSize)
 		{
-			float ViewInputRange = InViewInput.Size<double>();
-
+			const float ViewInputRange = InViewInput.Size<double>();
 			ViewStart = InViewInput.GetLowerBoundValue();
 			PixelsPerInput = ViewInputRange > 0 ? (InWidgetSize.X / ViewInputRange) : 0;
 		}
@@ -259,7 +246,7 @@ namespace UE::MLDeformer
 			return PixelsPerInput > 0 ? (ScreenX / PixelsPerInput) + ViewStart : ViewStart;
 		}
 
-		/** Curve Input domain -> local Widget Space */
+		/** Curve Input domain -> local Widget Space. */
 		float InputToLocalX(double Input) const
 		{
 			return (Input - ViewStart) * PixelsPerInput;
@@ -268,21 +255,20 @@ namespace UE::MLDeformer
 
 
 	/**
-	 * Gets the the next spacing value in the series
-	 * to determine a good spacing value
-	 * E.g, .001,.005,.010,.050,.100,.500,1.000,etc
+	 * Gets the the next spacing value in the series to determine a good spacing value.
+	 * E.g, .001, .005, .010, .050, .100, .500, 1.000, etc.
 	 */
 	static float GetNextSpacing(uint32 CurrentStep)
 	{
 		if (CurrentStep & 0x01)
 		{
 			// Odd numbers
-			return FMath::Pow(10.f, 0.5f * ((float)(CurrentStep - 1)) + 1.f);
+			return FMath::Pow(10.0f, 0.5f * ((float)(CurrentStep - 1)) + 1.0f);
 		}
 		else
 		{
 			// Even numbers
-			return 0.5f * FMath::Pow(10.f, 0.5f * ((float)(CurrentStep)) + 1.f);
+			return 0.5f * FMath::Pow(10.0f, 0.5f * ((float)(CurrentStep)) + 1.0f);
 		}
 	}
 
@@ -308,15 +294,14 @@ namespace UE::MLDeformer
 
 	FFrameTime FMLTimeSliderController::ComputeFrameTimeFromMouse(const FGeometry& Geometry, FVector2D ScreenSpacePosition, FScrubRangeToScreen RangeToScreen, bool CheckSnapping) const
 	{
-		FVector2D CursorPos = Geometry.AbsoluteToLocal(ScreenSpacePosition);
-		double    MouseValue = RangeToScreen.LocalXToInput(CursorPos.X);
-
+		const FVector2D CursorPos = Geometry.AbsoluteToLocal(ScreenSpacePosition);
+		const double MouseValue = RangeToScreen.LocalXToInput(CursorPos.X);
 		return MouseValue * GetTickResolution();
 	}
 
 	FMLTimeSliderController::FScrubPixelRange FMLTimeSliderController::GetHitTestScrubberPixelRange(FFrameTime ScrubTime, const FScrubRangeToScreen& RangeToScreen) const
 	{
-		static const float DragToleranceSlateUnits = 2.f, MouseTolerance = 2.f;
+		static const float DragToleranceSlateUnits = 2.0f, MouseTolerance = 2.0f;
 		return GetScrubberPixelRange(ScrubTime, GetTickResolution(), GetDisplayRate(), RangeToScreen, DragToleranceSlateUnits + MouseTolerance);
 	}
 
@@ -327,29 +312,25 @@ namespace UE::MLDeformer
 
 	FMLTimeSliderController::FScrubPixelRange FMLTimeSliderController::GetScrubberPixelRange(FFrameTime ScrubTime, FFrameRate Resolution, FFrameRate PlayRate, const FScrubRangeToScreen& RangeToScreen, float DilationPixels) const
 	{
-		const FFrameNumber Frame = ScrubTime.FloorToFrame();// ConvertFrameTime(ScrubTime, Resolution, PlayRate).FloorToFrame();
-
+		const FFrameNumber Frame = ScrubTime.FloorToFrame();
 		float StartPixel = RangeToScreen.InputToLocalX(Frame / Resolution);
 		float EndPixel = RangeToScreen.InputToLocalX((Frame + 1) / Resolution);
 
-		{
-			float RoundedStartPixel = FMath::RoundToInt(StartPixel);
-			EndPixel -= (StartPixel - RoundedStartPixel);
+		const float RoundedStartPixel = FMath::RoundToInt(StartPixel);
+		EndPixel -= (StartPixel - RoundedStartPixel);
 
-			StartPixel = RoundedStartPixel;
-			EndPixel = FMath::Max(EndPixel, StartPixel + 1);
-		}
+		StartPixel = RoundedStartPixel;
+		EndPixel = FMath::Max(EndPixel, StartPixel + 1);
 
+		const float MinScrubSize = 14.0f;
 		FScrubPixelRange Range;
-
-		float MinScrubSize = 14.f;
 		Range.bClamped = EndPixel - StartPixel < MinScrubSize;
 		Range.Range = TRange<float>(StartPixel, EndPixel);
 		if (Range.bClamped)
 		{
 			Range.HandleRange = TRange<float>(
-				(StartPixel + EndPixel - MinScrubSize) * .5f,
-				(StartPixel + EndPixel + MinScrubSize) * .5f);
+				(StartPixel + EndPixel - MinScrubSize) * 0.5f,
+				(StartPixel + EndPixel + MinScrubSize) * 0.5f);
 		}
 		else
 		{
@@ -361,25 +342,24 @@ namespace UE::MLDeformer
 
 	struct FMLTimeSliderController::FDrawTickArgs
 	{
-		/** Geometry of the area */
+		/** Geometry of the area. */
 		FGeometry AllottedGeometry;
-		/** Culling rect of the area */
+		/** Culling rect of the area. */
 		FSlateRect CullingRect;
-		/** Color of each tick */
+		/** Color of each tick. */
 		FLinearColor TickColor;
-		/** Offset in Y where to start the tick */
+		/** Offset in Y where to start the tick. */
 		float TickOffset;
-		/** Height in of major ticks */
+		/** Height in of major ticks. */
 		float MajorTickHeight;
-		/** Start layer for elements */
+		/** Start layer for elements. */
 		int32 StartLayer;
-		/** Draw effects to apply */
+		/** Draw effects to apply. */
 		ESlateDrawEffect DrawEffects;
-		/** Whether or not to only draw major ticks */
+		/** Whether or not to only draw major ticks. */
 		bool bOnlyDrawMajorTicks;
-		/** Whether or not to mirror labels */
+		/** Whether or not to mirror labels. */
 		bool bMirrorLabels;
-
 	};
 
 	void FMLTimeSliderController::DrawTicks(FSlateWindowElementList& OutDrawElements, const TRange<double>& ViewRange, const FScrubRangeToScreen& RangeToScreen, FDrawTickArgs& InArgs) const
@@ -395,12 +375,12 @@ namespace UE::MLDeformer
 			return;
 		}
 
-		FFrameRate     FrameResolution = GetTickResolution();
-		FPaintGeometry PaintGeometry = InArgs.AllottedGeometry.ToPaintGeometry();
-		FSlateFontInfo SmallLayoutFont = FCoreStyle::GetDefaultFontStyle("Regular", 8);
+		const FFrameRate     FrameResolution = GetTickResolution();
+		const FPaintGeometry PaintGeometry = InArgs.AllottedGeometry.ToPaintGeometry();
+		const FSlateFontInfo SmallLayoutFont = FCoreStyle::GetDefaultFontStyle("Regular", 8);
 
 		double MajorGridStep = 0.0;
-		int32  MinorDivisions = 0;
+		int32 MinorDivisions = 0;
 		if (!Timeline->GetGridMetrics(InArgs.AllottedGeometry.Size.X, MajorGridStep, MinorDivisions))
 		{
 			return;
@@ -415,18 +395,18 @@ namespace UE::MLDeformer
 		LinePoints.SetNumUninitialized(2);
 
 		const bool bAntiAliasLines = false;
-
 		const double FirstMajorLine = FMath::FloorToDouble(ViewRange.GetLowerBoundValue() / MajorGridStep) * MajorGridStep;
 		const double LastMajorLine = FMath::CeilToDouble(ViewRange.GetUpperBoundValue() / MajorGridStep) * MajorGridStep;
 
+		FString FrameString;
 		for (double CurrentMajorLine = FirstMajorLine; CurrentMajorLine < LastMajorLine; CurrentMajorLine += MajorGridStep)
 		{
-			float MajorLinePx = RangeToScreen.InputToLocalX(CurrentMajorLine);
+			const float MajorLinePx = RangeToScreen.InputToLocalX(CurrentMajorLine);
 
 			LinePoints[0] = FVector2D(MajorLinePx, InArgs.TickOffset);
 			LinePoints[1] = FVector2D(MajorLinePx, InArgs.TickOffset + InArgs.MajorTickHeight);
 
-			// Draw each tick mark
+			// Draw each tick mark.
 			FSlateDrawElement::MakeLines(
 				OutDrawElements,
 				InArgs.StartLayer,
@@ -439,10 +419,10 @@ namespace UE::MLDeformer
 
 			if (!InArgs.bOnlyDrawMajorTicks)
 			{
-				FString FrameString = TimeSliderArgs.NumericTypeInterface->ToString((CurrentMajorLine * FrameResolution).RoundToFrame().Value);
+				FrameString = TimeSliderArgs.NumericTypeInterface->ToString((CurrentMajorLine * FrameResolution).RoundToFrame().Value);
 
-				// Space the text between the tick mark but slightly above
-				FVector2D TextOffset(MajorLinePx + 5.f, InArgs.bMirrorLabels ? 3.f : FMath::Abs(InArgs.AllottedGeometry.Size.Y - (InArgs.MajorTickHeight + 3.f)));
+				// Space the text between the tick mark but slightly above.
+				const FVector2D TextOffset(MajorLinePx + 5.0f, InArgs.bMirrorLabels ? 3.0f : FMath::Abs(InArgs.AllottedGeometry.Size.Y - (InArgs.MajorTickHeight + 3.0f)));
 				FSlateDrawElement::MakeText(
 					OutDrawElements,
 					InArgs.StartLayer + 1,
@@ -456,14 +436,14 @@ namespace UE::MLDeformer
 
 			for (int32 Step = 1; Step < MinorDivisions; ++Step)
 			{
-				// Compute the size of each tick mark.  If we are half way between to visible values display a slightly larger tick mark
+				// Compute the size of each tick mark.  If we are half way between to visible values display a slightly larger tick mark.
 				const float MinorTickHeight = ((MinorDivisions % 2 == 0) && (Step % (MinorDivisions / 2)) == 0) ? 6.0f : 2.0f;
 				const float MinorLinePx = RangeToScreen.InputToLocalX(CurrentMajorLine + Step * MajorGridStep / MinorDivisions);
 
 				LinePoints[0] = FVector2D(MinorLinePx, InArgs.bMirrorLabels ? 0.0f : FMath::Abs(InArgs.AllottedGeometry.Size.Y - MinorTickHeight));
 				LinePoints[1] = FVector2D(MinorLinePx, LinePoints[0].Y + MinorTickHeight);
 
-				// Draw each sub mark
+				// Draw each sub mark.
 				FSlateDrawElement::MakeLines(
 					OutDrawElements,
 					InArgs.StartLayer,
@@ -471,8 +451,7 @@ namespace UE::MLDeformer
 					LinePoints,
 					InArgs.DrawEffects,
 					InArgs.TickColor,
-					bAntiAliasLines
-				);
+					bAntiAliasLines);
 			}
 		}
 	}
@@ -481,21 +460,18 @@ namespace UE::MLDeformer
 	{
 		const bool bEnabled = bParentEnabled;
 		const ESlateDrawEffect DrawEffects = bEnabled ? ESlateDrawEffect::None : ESlateDrawEffect::DisabledEffect;
+		const TRange<double> LocalViewRange = TimeSliderArgs.ViewRange.Get();
+		const TRange<FFrameNumber> LocalPlaybackRange = TimeSliderArgs.PlaybackRange.Get();
+		const float LocalViewRangeMin = LocalViewRange.GetLowerBoundValue();
+		const float LocalViewRangeMax = LocalViewRange.GetUpperBoundValue();
+		const float LocalSequenceLength = LocalViewRangeMax - LocalViewRangeMin;
 
-		TRange<double> LocalViewRange = TimeSliderArgs.ViewRange.Get();
-		const float    LocalViewRangeMin = LocalViewRange.GetLowerBoundValue();
-		const float    LocalViewRangeMax = LocalViewRange.GetUpperBoundValue();
-		const float    LocalSequenceLength = LocalViewRangeMax - LocalViewRangeMin;
-		TRange<FFrameNumber> LocalPlaybackRange = TimeSliderArgs.PlaybackRange.Get();
-
-		FVector2D Scale = FVector2D(1.0f, 1.0f);
 		if (LocalSequenceLength > 0)
 		{
 			FScrubRangeToScreen RangeToScreen(LocalViewRange, AllottedGeometry.Size);
 
-			// draw tick marks
+			// Draw the ticks.
 			const float MajorTickHeight = 9.0f;
-
 			FDrawTickArgs Args;
 			{
 				Args.AllottedGeometry = AllottedGeometry;
@@ -508,24 +484,22 @@ namespace UE::MLDeformer
 				Args.TickOffset = bMirrorLabels ? 0.0f : FMath::Abs(AllottedGeometry.Size.Y - MajorTickHeight);
 				Args.MajorTickHeight = MajorTickHeight;
 			}
-
 			DrawTicks(OutDrawElements, LocalViewRange, RangeToScreen, Args);
 
-			// draw playback & selection range
+			// Draw playback & selection range.
 			FPaintPlaybackRangeArgs PlaybackRangeArgs(
 				bMirrorLabels ? FAppStyle::GetBrush("Sequencer.Timeline.PlayRange_Bottom_L") : FAppStyle::GetBrush("Sequencer.Timeline.PlayRange_Top_L"),
 				bMirrorLabels ? FAppStyle::GetBrush("Sequencer.Timeline.PlayRange_Bottom_R") : FAppStyle::GetBrush("Sequencer.Timeline.PlayRange_Top_R"),
-				6.f
-			);
+				6.0f);
 
 			LayerId = DrawPlaybackRange(AllottedGeometry, MyCullingRect, OutDrawElements, LayerId, RangeToScreen, PlaybackRangeArgs);
 
 			PlaybackRangeArgs.SolidFillOpacity = 0.05f;
 			LayerId = DrawSelectionRange(AllottedGeometry, MyCullingRect, OutDrawElements, LayerId, RangeToScreen, PlaybackRangeArgs);
 
-			// Draw the scrub handle
-			const float      HandleStart = RangeToScreen.InputToLocalX(TimeSliderArgs.ScrubPosition.Get().AsDecimal() / GetTickResolution().AsDecimal()) - 7.0f;
-			const float      HandleEnd = HandleStart + 13.0f;
+			// Draw the scrub handle.
+			const float HandleStart = RangeToScreen.InputToLocalX(TimeSliderArgs.ScrubPosition.Get().AsDecimal() / GetTickResolution().AsDecimal()) - 7.0f;
+			const float HandleEnd = HandleStart + 13.0f;
 
 			const int32 ArrowLayer = LayerId + 2;
 			FPaintGeometry MyGeometry = AllottedGeometry.ToPaintGeometry(FVector2D(HandleStart, 0), FVector2D(HandleEnd - HandleStart, AllottedGeometry.Size.Y));
@@ -544,32 +518,29 @@ namespace UE::MLDeformer
 				MyGeometry,
 				Brush,
 				DrawEffects,
-				ScrubColor
-			);
+				ScrubColor);
 
 			{
-				// Draw the current time next to the scrub handle
+				// Draw the current time next to the scrub handle.
 				FString FrameString = TimeSliderArgs.NumericTypeInterface->ToString(TimeSliderArgs.ScrubPosition.Get().GetFrame().Value);
 
 				if (GetDefault<UPersonaOptions>()->bTimelineDisplayPercentage)
 				{
-					double Percentage = FMath::Clamp(TimeSliderArgs.ScrubPosition.Get().AsDecimal() / FFrameTime(LocalPlaybackRange.Size<FFrameNumber>()).AsDecimal(), 0.0, 1.0);
+					const double Percentage = FMath::Clamp(TimeSliderArgs.ScrubPosition.Get().AsDecimal() / FFrameTime(LocalPlaybackRange.Size<FFrameNumber>()).AsDecimal(), 0.0, 1.0);
 					FNumberFormattingOptions Options;
 					Options.MaximumFractionalDigits = 2;
 					FrameString += TEXT(" (") + FText::AsPercent(Percentage, &Options).ToString() + TEXT(")");
 				}
 
-				FSlateFontInfo SmallLayoutFont = FCoreStyle::GetDefaultFontStyle("Regular", 10);
+				const FSlateFontInfo SmallLayoutFont = FCoreStyle::GetDefaultFontStyle("Regular", 10);
+				const TSharedRef<FSlateFontMeasure> FontMeasureService = FSlateApplication::Get().GetRenderer()->GetFontMeasureService();
+				const FVector2D TextSize = FontMeasureService->Measure(FrameString, SmallLayoutFont);
 
-				const TSharedRef< FSlateFontMeasure > FontMeasureService = FSlateApplication::Get().GetRenderer()->GetFontMeasureService();
-				FVector2D TextSize = FontMeasureService->Measure(FrameString, SmallLayoutFont);
-
-				// Flip the text position if getting near the end of the view range
-				static const float TextOffsetPx = 2.f;
-				bool  bDrawLeft = (AllottedGeometry.Size.X - HandleEnd) < (TextSize.X + 14.f) - TextOffsetPx;
-				float TextPosition = bDrawLeft ? HandleStart - TextSize.X - TextOffsetPx : HandleEnd + TextOffsetPx;
-
-				FVector2D TextOffset(TextPosition, Args.bMirrorLabels ? TextSize.Y - 6.f : Args.AllottedGeometry.Size.Y - (Args.MajorTickHeight + TextSize.Y));
+				// Flip the text position if getting near the end of the view range.
+				static const float TextOffsetPx = 2.0f;
+				const bool bDrawLeft = (AllottedGeometry.Size.X - HandleEnd) < (TextSize.X + 14.0f) - TextOffsetPx;
+				const float TextPosition = bDrawLeft ? HandleStart - TextSize.X - TextOffsetPx : HandleEnd + TextOffsetPx;
+				const FVector2D TextOffset(TextPosition, Args.bMirrorLabels ? TextSize.Y - 6.f : Args.AllottedGeometry.Size.Y - (Args.MajorTickHeight + TextSize.Y));
 
 				FSlateDrawElement::MakeText(
 					OutDrawElements,
@@ -578,8 +549,7 @@ namespace UE::MLDeformer
 					FrameString,
 					SmallLayoutFont,
 					Args.DrawEffects,
-					Args.TickColor
-				);
+					Args.TickColor);
 			}
 
 			return ArrowLayer;
@@ -606,8 +576,7 @@ namespace UE::MLDeformer
 					AllottedGeometry.ToPaintGeometry(FVector2D(SelectionRangeL, 0.f), FVector2D(SelectionRangeR - SelectionRangeL, AllottedGeometry.Size.Y)),
 					FAppStyle::GetBrush("WhiteBrush"),
 					ESlateDrawEffect::None,
-					DrawColor.CopyWithNewOpacity(Args.SolidFillOpacity)
-				);
+					DrawColor.CopyWithNewOpacity(Args.SolidFillOpacity));
 			}
 
 			FSlateDrawElement::MakeBox(
@@ -616,8 +585,7 @@ namespace UE::MLDeformer
 				AllottedGeometry.ToPaintGeometry(FVector2D(SelectionRangeL, 0.f), FVector2D(Args.BrushWidth, AllottedGeometry.Size.Y)),
 				Args.StartBrush,
 				ESlateDrawEffect::None,
-				DrawColor
-			);
+				DrawColor);
 
 			FSlateDrawElement::MakeBox(
 				OutDrawElements,
@@ -625,8 +593,7 @@ namespace UE::MLDeformer
 				AllottedGeometry.ToPaintGeometry(FVector2D(SelectionRangeR - Args.BrushWidth, 0.f), FVector2D(Args.BrushWidth, AllottedGeometry.Size.Y)),
 				Args.EndBrush,
 				ESlateDrawEffect::None,
-				DrawColor
-			);
+				DrawColor);
 		}
 
 		return LayerId + 1;
@@ -639,17 +606,16 @@ namespace UE::MLDeformer
 			return LayerId;
 		}
 
-		const uint8 OpacityBlend = TimeSliderArgs.SubSequenceRange.Get().IsSet() ? 128 : 255;
-
-		TRange<FFrameNumber> PlaybackRange = TimeSliderArgs.PlaybackRange.Get();
-		FFrameRate TickResolution = GetTickResolution();
+		const TRange<FFrameNumber> PlaybackRange = TimeSliderArgs.PlaybackRange.Get();
+		const FFrameRate TickResolution = GetTickResolution();
 		const float PlaybackRangeL = RangeToScreen.InputToLocalX(PlaybackRange.GetLowerBoundValue() / TickResolution);
 		const float PlaybackRangeR = RangeToScreen.InputToLocalX(PlaybackRange.GetUpperBoundValue() / TickResolution) - 1;
 
+		const uint8 OpacityBlend = TimeSliderArgs.SubSequenceRange.Get().IsSet() ? 128 : 255;
 		FSlateDrawElement::MakeBox(
 			OutDrawElements,
 			LayerId + 1,
-			AllottedGeometry.ToPaintGeometry(FVector2D(PlaybackRangeL, 0.f), FVector2D(Args.BrushWidth, AllottedGeometry.Size.Y)),
+			AllottedGeometry.ToPaintGeometry(FVector2D(PlaybackRangeL, 0.0f), FVector2D(Args.BrushWidth, AllottedGeometry.Size.Y)),
 			Args.StartBrush,
 			ESlateDrawEffect::None,
 			FColor(32, 128, 32, OpacityBlend)	// 120, 75, 50 (HSV)
@@ -658,7 +624,7 @@ namespace UE::MLDeformer
 		FSlateDrawElement::MakeBox(
 			OutDrawElements,
 			LayerId + 1,
-			AllottedGeometry.ToPaintGeometry(FVector2D(PlaybackRangeR - Args.BrushWidth, 0.f), FVector2D(Args.BrushWidth, AllottedGeometry.Size.Y)),
+			AllottedGeometry.ToPaintGeometry(FVector2D(PlaybackRangeR - Args.BrushWidth, 0.0f), FVector2D(Args.BrushWidth, AllottedGeometry.Size.Y)),
 			Args.EndBrush,
 			ESlateDrawEffect::None,
 			FColor(128, 32, 32, OpacityBlend)	// 0, 75, 50 (HSV)
@@ -671,16 +637,16 @@ namespace UE::MLDeformer
 			AllottedGeometry.ToPaintGeometry(FVector2D(0.f, 0.f), FVector2D(PlaybackRangeL, AllottedGeometry.Size.Y)),
 			FAppStyle::GetBrush("WhiteBrush"),
 			ESlateDrawEffect::None,
-			FLinearColor::Black.CopyWithNewOpacity(0.3f * OpacityBlend / 255.f)
+			FLinearColor::Black.CopyWithNewOpacity(0.3f * OpacityBlend / 255.0f)
 		);
 
 		FSlateDrawElement::MakeBox(
 			OutDrawElements,
 			LayerId + 1,
-			AllottedGeometry.ToPaintGeometry(FVector2D(PlaybackRangeR, 0.f), FVector2D(AllottedGeometry.Size.X - PlaybackRangeR, AllottedGeometry.Size.Y)),
+			AllottedGeometry.ToPaintGeometry(FVector2D(PlaybackRangeR, 0.0f), FVector2D(AllottedGeometry.Size.X - PlaybackRangeR, AllottedGeometry.Size.Y)),
 			FAppStyle::GetBrush("WhiteBrush"),
 			ESlateDrawEffect::None,
-			FLinearColor::Black.CopyWithNewOpacity(0.3f * OpacityBlend / 255.f)
+			FLinearColor::Black.CopyWithNewOpacity(0.3f * OpacityBlend / 255.0f)
 		);
 
 		return LayerId + 1;
@@ -696,11 +662,11 @@ namespace UE::MLDeformer
 
 	FReply FMLTimeSliderController::OnMouseButtonUp(SWidget& WidgetOwner, const FGeometry& MyGeometry, const FPointerEvent& MouseEvent)
 	{
-		bool bHandleLeftMouseButton = MouseEvent.GetEffectingButton() == EKeys::LeftMouseButton && WidgetOwner.HasMouseCapture();
-		bool bHandleRightMouseButton = MouseEvent.GetEffectingButton() == EKeys::RightMouseButton && WidgetOwner.HasMouseCapture() && TimeSliderArgs.AllowZoom;
+		const bool bHandleLeftMouseButton = (MouseEvent.GetEffectingButton() == EKeys::LeftMouseButton) && WidgetOwner.HasMouseCapture();
+		const bool bHandleRightMouseButton = (MouseEvent.GetEffectingButton() == EKeys::RightMouseButton) && WidgetOwner.HasMouseCapture() && TimeSliderArgs.AllowZoom;
 
-		FScrubRangeToScreen RangeToScreen = FScrubRangeToScreen(TimeSliderArgs.ViewRange.Get(), MyGeometry.Size);
-		FFrameTime          MouseTime = ComputeFrameTimeFromMouse(MyGeometry, MouseEvent.GetScreenSpacePosition(), RangeToScreen);
+		const FScrubRangeToScreen RangeToScreen = FScrubRangeToScreen(TimeSliderArgs.ViewRange.Get(), MyGeometry.Size);
+		const FFrameTime MouseTime = ComputeFrameTimeFromMouse(MyGeometry, MouseEvent.GetScreenSpacePosition(), RangeToScreen);
 
 		if (bHandleRightMouseButton)
 		{
@@ -710,7 +676,7 @@ namespace UE::MLDeformer
 			}
 
 			bPanning = false;
-			DistanceDragged = 0.f;
+			DistanceDragged = 0.0f;
 
 			return FReply::Handled().ReleaseMouseCapture();
 		}
@@ -740,8 +706,7 @@ namespace UE::MLDeformer
 			}
 
 			MouseDragType = DRAG_NONE;
-			DistanceDragged = 0.f;
-
+			DistanceDragged = 0.0f;
 			return FReply::Handled().ReleaseMouseCapture();
 		}
 
@@ -750,8 +715,8 @@ namespace UE::MLDeformer
 
 	FReply FMLTimeSliderController::OnMouseMove(SWidget& WidgetOwner, const FGeometry& MyGeometry, const FPointerEvent& MouseEvent)
 	{
-		bool bHandleLeftMouseButton = MouseEvent.IsMouseButtonDown(EKeys::LeftMouseButton);
-		bool bHandleRightMouseButton = MouseEvent.IsMouseButtonDown(EKeys::RightMouseButton) && TimeSliderArgs.AllowZoom;
+		const bool bHandleLeftMouseButton = MouseEvent.IsMouseButtonDown(EKeys::LeftMouseButton);
+		const bool bHandleRightMouseButton = MouseEvent.IsMouseButtonDown(EKeys::RightMouseButton) && TimeSliderArgs.AllowZoom;
 
 		if (bHandleRightMouseButton)
 		{
@@ -765,17 +730,15 @@ namespace UE::MLDeformer
 			}
 			else
 			{
-				TRange<double> LocalViewRange = TimeSliderArgs.ViewRange.Get();
-				double LocalViewRangeMin = LocalViewRange.GetLowerBoundValue();
-				double LocalViewRangeMax = LocalViewRange.GetUpperBoundValue();
-
-				FScrubRangeToScreen ScaleInfo(LocalViewRange, MyGeometry.Size);
-				FVector2D ScreenDelta = MouseEvent.GetCursorDelta();
-				FVector2D InputDelta;
-				InputDelta.X = ScreenDelta.X / ScaleInfo.PixelsPerInput;
+				const TRange<double> LocalViewRange = TimeSliderArgs.ViewRange.Get();
+				const double LocalViewRangeMin = LocalViewRange.GetLowerBoundValue();
+				const double LocalViewRangeMax = LocalViewRange.GetUpperBoundValue();
 
-				double NewViewOutputMin = LocalViewRangeMin - InputDelta.X;
-				double NewViewOutputMax = LocalViewRangeMax - InputDelta.X;
+				const FScrubRangeToScreen ScaleInfo(LocalViewRange, MyGeometry.Size);
+				const FVector2D ScreenDelta = MouseEvent.GetCursorDelta();
+				const double InputDeltaX = ScreenDelta.X / ScaleInfo.PixelsPerInput;
+				double NewViewOutputMin = LocalViewRangeMin - InputDeltaX;
+				double NewViewOutputMax = LocalViewRangeMax - InputDeltaX;
 
 				ClampViewRange(NewViewOutputMin, NewViewOutputMax);
 				SetViewRange(NewViewOutputMin, NewViewOutputMax, EViewRangeInterpolation::Immediate);
@@ -783,15 +746,15 @@ namespace UE::MLDeformer
 		}
 		else if (bHandleLeftMouseButton)
 		{
-			TRange<double> LocalViewRange = TimeSliderArgs.ViewRange.Get();
-			FScrubRangeToScreen RangeToScreen(LocalViewRange, MyGeometry.Size);
+			const TRange<double> LocalViewRange = TimeSliderArgs.ViewRange.Get();
+			const FScrubRangeToScreen RangeToScreen(LocalViewRange, MyGeometry.Size);
 			DistanceDragged += FMath::Abs(MouseEvent.GetCursorDelta().X);
 
 			if (MouseDragType == DRAG_NONE)
 			{
 				if (DistanceDragged > FSlateApplication::Get().GetDragTriggerDistance())
 				{
-					FFrameTime MouseDownFree = ComputeFrameTimeFromMouse(MyGeometry, MouseDownPosition[0], RangeToScreen, false);
+					const FFrameTime MouseDownFree = ComputeFrameTimeFromMouse(MyGeometry, MouseDownPosition[0], RangeToScreen, false);
 
 					const FFrameRate FrameResolution = GetTickResolution();
 					const bool       bLockedPlayRange = TimeSliderArgs.IsPlaybackRangeLocked.Get();
@@ -801,28 +764,28 @@ namespace UE::MLDeformer
 					TRange<double>   SelectionRange = TimeSliderArgs.SelectionRange.Get() / FrameResolution;
 					TRange<double>   PlaybackRange = TimeSliderArgs.PlaybackRange.Get() / FrameResolution;
 
-					// Disable selection range test if it's empty so that the playback range scrubbing gets priority
+					// Disable selection range test if it's empty so that the playback range scrubbing gets priority.
 					if (!SelectionRange.IsEmpty() && !bHitScrubber && HitTestRangeEnd(RangeToScreen, SelectionRange, MouseDownPixel))
 					{
-						// selection range end scrubber
+						// selection range end scrubber.
 						MouseDragType = DRAG_SELECTION_END;
 						TimeSliderArgs.OnSelectionRangeBeginDrag.ExecuteIfBound();
 					}
 					else if (!SelectionRange.IsEmpty() && !bHitScrubber && HitTestRangeStart(RangeToScreen, SelectionRange, MouseDownPixel))
 					{
-						// selection range start scrubber
+						// selection range start scrubber.
 						MouseDragType = DRAG_SELECTION_START;
 						TimeSliderArgs.OnSelectionRangeBeginDrag.ExecuteIfBound();
 					}
 					else if (!bLockedPlayRange && !bHitScrubber && HitTestRangeEnd(RangeToScreen, PlaybackRange, MouseDownPixel))
 					{
-						// playback range end scrubber
+						// playback range end scrubber.
 						MouseDragType = DRAG_PLAYBACK_END;
 						TimeSliderArgs.OnPlaybackRangeBeginDrag.ExecuteIfBound();
 					}
 					else if (!bLockedPlayRange && !bHitScrubber && HitTestRangeStart(RangeToScreen, PlaybackRange, MouseDownPixel))
 					{
-						// playback range start scrubber
+						// playback range start scrubber.
 						MouseDragType = DRAG_PLAYBACK_START;
 						TimeSliderArgs.OnPlaybackRangeBeginDrag.ExecuteIfBound();
 					}
@@ -835,7 +798,7 @@ namespace UE::MLDeformer
 			}
 			else
 			{
-				FFrameTime MouseTime = ComputeFrameTimeFromMouse(MyGeometry, MouseEvent.GetScreenSpacePosition(), RangeToScreen);
+				const FFrameTime MouseTime = ComputeFrameTimeFromMouse(MyGeometry, MouseEvent.GetScreenSpacePosition(), RangeToScreen);
 
 				// Set the start range time?
 				if (MouseDragType == DRAG_PLAYBACK_START)
@@ -858,7 +821,7 @@ namespace UE::MLDeformer
 				}
 				else if (MouseDragType == DRAG_SCRUBBING_TIME)
 				{
-					// Delegate responsibility for clamping to the current viewrange to the client
+					// Delegate responsibility for clamping to the current viewrange to the client.
 					CommitScrubPosition(MouseTime, /*bIsScrubbing=*/true);
 				}
 			}
@@ -874,23 +837,21 @@ namespace UE::MLDeformer
 
 	void FMLTimeSliderController::CommitScrubPosition(FFrameTime NewValue, bool bIsScrubbing)
 	{
-		// Manage the scrub position ourselves if its not bound to a delegate
+		// Manage the scrub position ourselves if its not bound to a delegate.
 		if (!TimeSliderArgs.ScrubPosition.IsBound())
 		{
 			TimeSliderArgs.ScrubPosition.Set(NewValue);
 		}
 
-		TimeSliderArgs.OnScrubPositionChanged.ExecuteIfBound(NewValue, bIsScrubbing, /*bEvaluate*/ true); //todo change if anim timeline needs to handle sequencer style middle mouse manipulation which changes time but doesn't evaluate
+		// TODO: Change if anim timeline needs to handle sequencer style middle mouse manipulation which changes time but doesn't evaluate.
+		TimeSliderArgs.OnScrubPositionChanged.ExecuteIfBound(NewValue, bIsScrubbing, /*bEvaluate*/ true);
 	}
 
 	FReply FMLTimeSliderController::OnMouseWheel(SWidget& WidgetOwner, const FGeometry& MyGeometry, const FPointerEvent& MouseEvent)
 	{
-		TOptional<TRange<float>> NewTargetRange;
-
 		if (TimeSliderArgs.AllowZoom && MouseEvent.IsControlDown())
 		{
-			float MouseFractionX = MyGeometry.AbsoluteToLocal(MouseEvent.GetScreenSpacePosition()).X / MyGeometry.GetLocalSize().X;
-
+			const float MouseFractionX = MyGeometry.AbsoluteToLocal(MouseEvent.GetScreenSpacePosition()).X / MyGeometry.GetLocalSize().X;
 			const float ZoomDelta = -0.2f * MouseEvent.GetWheelDelta();
 			if (ZoomByDelta(ZoomDelta, MouseFractionX))
 			{
@@ -908,22 +869,20 @@ namespace UE::MLDeformer
 
 	FCursorReply FMLTimeSliderController::OnCursorQuery(TSharedRef<const SWidget> WidgetOwner, const FGeometry& MyGeometry, const FPointerEvent& CursorEvent) const
 	{
-		FScrubRangeToScreen RangeToScreen(TimeSliderArgs.ViewRange.Get(), MyGeometry.Size);
-
+		const FScrubRangeToScreen RangeToScreen(TimeSliderArgs.ViewRange.Get(), MyGeometry.Size);
 		const FFrameRate FrameResolution = GetTickResolution();
-		const bool       bLockedPlayRange = TimeSliderArgs.IsPlaybackRangeLocked.Get();
-		const float      HitTestPixel = MyGeometry.AbsoluteToLocal(CursorEvent.GetScreenSpacePosition()).X;
-		const bool       bHitScrubber = GetHitTestScrubberPixelRange(TimeSliderArgs.ScrubPosition.Get(), RangeToScreen).HandleRange.Contains(HitTestPixel);
-
-		TRange<double>   SelectionRange = TimeSliderArgs.SelectionRange.Get() / FrameResolution;
-		TRange<double>   PlaybackRange = TimeSliderArgs.PlaybackRange.Get() / FrameResolution;
+		const TRange<double> SelectionRange = TimeSliderArgs.SelectionRange.Get() / FrameResolution;
+		const TRange<double> PlaybackRange = TimeSliderArgs.PlaybackRange.Get() / FrameResolution;
+		const float HitTestPixel = MyGeometry.AbsoluteToLocal(CursorEvent.GetScreenSpacePosition()).X;
+		const bool bLockedPlayRange = TimeSliderArgs.IsPlaybackRangeLocked.Get();
+		const bool bHitScrubber = GetHitTestScrubberPixelRange(TimeSliderArgs.ScrubPosition.Get(), RangeToScreen).HandleRange.Contains(HitTestPixel);
 
 		if (MouseDragType == DRAG_SCRUBBING_TIME)
 		{
 			return FCursorReply::Unhandled();
 		}
 
-		// Use L/R resize cursor if we're dragging or hovering a playback range bound
+		// Use L/R resize cursor if we're dragging or hovering a playback range bound.
 		if ((MouseDragType == DRAG_PLAYBACK_END) ||
 			(MouseDragType == DRAG_PLAYBACK_START) ||
 			(MouseDragType == DRAG_SELECTION_START) ||
@@ -959,7 +918,7 @@ namespace UE::MLDeformer
 		{
 			static FLinearColor TickColor(0.1f, 0.1f, 0.1f, 0.3f);
 
-			// Draw major tick lines in the section area
+			// Draw major tick lines in the section area.
 			FDrawTickArgs DrawTickArgs;
 			{
 				DrawTickArgs.AllottedGeometry = AllottedGeometry;
@@ -968,9 +927,9 @@ namespace UE::MLDeformer
 				DrawTickArgs.TickColor = TickColor;
 				DrawTickArgs.CullingRect = MyCullingRect;
 				DrawTickArgs.DrawEffects = DrawEffects;
-				// Draw major ticks under sections
+				// Draw major ticks under sections.
 				DrawTickArgs.StartLayer = LayerId - 1;
-				// Draw the tick the entire height of the section area
+				// Draw the tick the entire height of the section area.
 				DrawTickArgs.TickOffset = 0.0f;
 				DrawTickArgs.MajorTickHeight = AllottedGeometry.Size.Y;
 			}
@@ -980,7 +939,7 @@ namespace UE::MLDeformer
 
 		if (Args.bDisplayScrubPosition)
 		{
-			// Draw a line for the scrub position
+			// Draw a line for the scrub position.
 			const float LinePos = RangeToScreen.InputToLocalX(TimeSliderArgs.ScrubPosition.Get().AsDecimal() / GetTickResolution().AsDecimal());
 
 			TArray<FVector2D> LinePoints;
@@ -996,7 +955,7 @@ namespace UE::MLDeformer
 				AllottedGeometry.ToPaintGeometry(FVector2D(LinePos, 0.0f), FVector2D(1.0f, 1.0f)),
 				LinePoints,
 				DrawEffects,
-				FLinearColor(1.f, 1.f, 1.f, .5f),
+				FLinearColor(1.0f, 1.0f, 1.0f, 0.5f),
 				false
 			);
 		}
@@ -1006,7 +965,7 @@ namespace UE::MLDeformer
 
 	FFrameTime FMLTimeSliderController::GetFrameTimeFromMouse(const FGeometry& Geometry, FVector2D ScreenSpacePosition) const
 	{
-		FScrubRangeToScreen ScrubRangeToScreen(TimeSliderArgs.ViewRange.Get(), Geometry.Size);
+		const FScrubRangeToScreen ScrubRangeToScreen(TimeSliderArgs.ViewRange.Get(), Geometry.Size);
 		return ComputeFrameTimeFromMouse(Geometry, ScreenSpacePosition, ScrubRangeToScreen);
 	}
 
@@ -1035,7 +994,7 @@ namespace UE::MLDeformer
 
 	void FMLTimeSliderController::SetViewRange(double NewRangeMin, double NewRangeMax, EViewRangeInterpolation Interpolation)
 	{
-		// Clamp to a minimum size to avoid zero-sized or negative visible ranges
+		// Clamp to a minimum size to avoid zero-sized or negative visible ranges.
 		double MinVisibleTimeRange = FFrameNumber(1) / GetTickResolution();
 		TRange<double> ExistingViewRange = TimeSliderArgs.ViewRange.Get();
 
@@ -1051,13 +1010,13 @@ namespace UE::MLDeformer
 			NewRangeMax = NewRangeMin + MinVisibleTimeRange;
 		}
 
-		// Clamp to the clamp range
+		// Clamp to the clamp range.
 		const TRange<double> NewRange = TRange<double>(NewRangeMin, NewRangeMax);
 		TimeSliderArgs.OnViewRangeChanged.ExecuteIfBound(NewRange, Interpolation);
 
 		if (!TimeSliderArgs.ViewRange.IsBound())
 		{
-			// The  output is not bound to a delegate so we'll manage the value ourselves (no animation)
+			// The  output is not bound to a delegate so we'll manage the value ourselves (no animation).
 			TimeSliderArgs.ViewRange.Set(NewRange);
 		}
 	}
@@ -1070,7 +1029,7 @@ namespace UE::MLDeformer
 
 		if (!TimeSliderArgs.ClampRange.IsBound())
 		{
-			// The  output is not bound to a delegate so we'll manage the value ourselves (no animation)
+			// The output is not bound to a delegate so we'll manage the value ourselves (no animation).
 			TimeSliderArgs.ClampRange.Set(NewRange);
 		}
 	}
@@ -1078,29 +1037,25 @@ namespace UE::MLDeformer
 	void FMLTimeSliderController::SetPlayRange(FFrameNumber RangeStart, int32 RangeDuration)
 	{
 		check(RangeDuration >= 0);
-
 		const TRange<FFrameNumber> NewRange(RangeStart, RangeStart + RangeDuration);
-
 		TimeSliderArgs.OnPlaybackRangeChanged.ExecuteIfBound(NewRange);
-
 		if (!TimeSliderArgs.PlaybackRange.IsBound())
 		{
-			// The  output is not bound to a delegate so we'll manage the value ourselves (no animation)
+			// The output is not bound to a delegate so we'll manage the value ourselves (no animation).
 			TimeSliderArgs.PlaybackRange.Set(NewRange);
 		}
 	}
 
 	bool FMLTimeSliderController::ZoomByDelta(float InDelta, float MousePositionFraction)
 	{
-		TRange<double> LocalViewRange = TimeSliderArgs.ViewRange.Get().GetAnimationTarget();
-		double LocalViewRangeMax = LocalViewRange.GetUpperBoundValue();
-		double LocalViewRangeMin = LocalViewRange.GetLowerBoundValue();
+		const TRange<double> LocalViewRange = TimeSliderArgs.ViewRange.Get().GetAnimationTarget();
+		const double LocalViewRangeMax = LocalViewRange.GetUpperBoundValue();
+		const double LocalViewRangeMin = LocalViewRange.GetLowerBoundValue();
 		const double OutputViewSize = LocalViewRangeMax - LocalViewRangeMin;
 		const double OutputChange = OutputViewSize * InDelta;
 
 		double NewViewOutputMin = LocalViewRangeMin - (OutputChange * MousePositionFraction);
-		double NewViewOutputMax = LocalViewRangeMax + (OutputChange * (1.f - MousePositionFraction));
-
+		double NewViewOutputMax = LocalViewRangeMax + (OutputChange * (1.0f - MousePositionFraction));
 		if (NewViewOutputMin < NewViewOutputMax)
 		{
 			ClampViewRange(NewViewOutputMin, NewViewOutputMax);
@@ -1113,15 +1068,13 @@ namespace UE::MLDeformer
 
 	void FMLTimeSliderController::PanByDelta(float InDelta)
 	{
-		/**The fraction of the current view range to scroll per unit delta  */
+		// The fraction of the current view range to scroll per unit delta.
 		const float ScrollPanFraction = 0.1f;
+		const TRange<double> LocalViewRange = TimeSliderArgs.ViewRange.Get().GetAnimationTarget();
+		const double CurrentMin = LocalViewRange.GetLowerBoundValue();
+		const double CurrentMax = LocalViewRange.GetUpperBoundValue();
 
-		TRange<double> LocalViewRange = TimeSliderArgs.ViewRange.Get().GetAnimationTarget();
-
-		double CurrentMin = LocalViewRange.GetLowerBoundValue();
-		double CurrentMax = LocalViewRange.GetUpperBoundValue();
-
-		// Adjust the delta to be a percentage of the current range
+		// Adjust the delta to be a percentage of the current range.
 		InDelta *= ScrollPanFraction * (CurrentMax - CurrentMin);
 
 		double NewViewOutputMin = CurrentMin + InDelta;
@@ -1135,12 +1088,14 @@ namespace UE::MLDeformer
 	{
 		if (Range.HasLowerBound())
 		{
-			static float BrushSizeInStateUnits = 6.f, DragToleranceSlateUnits = 2.f, MouseTolerance = 2.f;
-			const float  RangeStartPixel = RangeToScreen.InputToLocalX(Range.GetLowerBoundValue());
+			static const float BrushSizeInStateUnits = 6.0f;
+			static const float DragToleranceSlateUnits = 2.0f;
+			static const float MouseTolerance = 2.0f;
+			const float RangeStartPixel = RangeToScreen.InputToLocalX(Range.GetLowerBoundValue());
 
-			// Hit test against the brush region to the right of the playback start position, +/- DragToleranceSlateUnits
-			return HitPixel >= RangeStartPixel - MouseTolerance - DragToleranceSlateUnits &&
-				HitPixel <= RangeStartPixel + MouseTolerance + BrushSizeInStateUnits + DragToleranceSlateUnits;
+			// Hit test against the brush region to the right of the playback start position, +/- DragToleranceSlateUnits.
+			return (HitPixel >= RangeStartPixel - MouseTolerance - DragToleranceSlateUnits) &&
+				(HitPixel <= RangeStartPixel + MouseTolerance + BrushSizeInStateUnits + DragToleranceSlateUnits);
 		}
 
 		return false;
@@ -1150,12 +1105,14 @@ namespace UE::MLDeformer
 	{
 		if (Range.HasUpperBound())
 		{
-			static float BrushSizeInStateUnits = 6.f, DragToleranceSlateUnits = 2.f, MouseTolerance = 2.f;
-			const float  RangeEndPixel = RangeToScreen.InputToLocalX(Range.GetUpperBoundValue());
+			static const float BrushSizeInStateUnits = 6.0f;
+			static const float DragToleranceSlateUnits = 2.0f;
+			static const float MouseTolerance = 2.0f;
+			const float RangeEndPixel = RangeToScreen.InputToLocalX(Range.GetUpperBoundValue());
 
-			// Hit test against the brush region to the left of the playback end position, +/- DragToleranceSlateUnits
-			return HitPixel >= RangeEndPixel - MouseTolerance - BrushSizeInStateUnits - DragToleranceSlateUnits &&
-				HitPixel <= RangeEndPixel + MouseTolerance + DragToleranceSlateUnits;
+			// Hit test against the brush region to the left of the playback end position, +/- DragToleranceSlateUnits.
+			return (HitPixel >= RangeEndPixel - MouseTolerance - BrushSizeInStateUnits - DragToleranceSlateUnits) &&
+				(HitPixel <= RangeEndPixel + MouseTolerance + DragToleranceSlateUnits);
 		}
 
 		return false;
@@ -1163,8 +1120,7 @@ namespace UE::MLDeformer
 
 	void FMLTimeSliderController::SetPlaybackRangeStart(FFrameNumber NewStart)
 	{
-		TRange<FFrameNumber> PlaybackRange = TimeSliderArgs.PlaybackRange.Get();
-
+		const TRange<FFrameNumber> PlaybackRange = TimeSliderArgs.PlaybackRange.Get();
 		if (NewStart <= UE::MovieScene::DiscreteExclusiveUpper(PlaybackRange))
 		{
 			TimeSliderArgs.OnPlaybackRangeChanged.ExecuteIfBound(TRange<FFrameNumber>(NewStart, PlaybackRange.GetUpperBound()));
@@ -1173,8 +1129,7 @@ namespace UE::MLDeformer
 
 	void FMLTimeSliderController::SetPlaybackRangeEnd(FFrameNumber NewEnd)
 	{
-		TRange<FFrameNumber> PlaybackRange = TimeSliderArgs.PlaybackRange.Get();
-
+		const TRange<FFrameNumber> PlaybackRange = TimeSliderArgs.PlaybackRange.Get();
 		if (NewEnd >= UE::MovieScene::DiscreteInclusiveLower(PlaybackRange))
 		{
 			TimeSliderArgs.OnPlaybackRangeChanged.ExecuteIfBound(TRange<FFrameNumber>(PlaybackRange.GetLowerBound(), NewEnd));
@@ -1183,8 +1138,7 @@ namespace UE::MLDeformer
 
 	void FMLTimeSliderController::SetSelectionRangeStart(FFrameNumber NewStart)
 	{
-		TRange<FFrameNumber> SelectionRange = TimeSliderArgs.SelectionRange.Get();
-
+		const TRange<FFrameNumber> SelectionRange = TimeSliderArgs.SelectionRange.Get();
 		if (SelectionRange.IsEmpty())
 		{
 			TimeSliderArgs.OnSelectionRangeChanged.ExecuteIfBound(TRange<FFrameNumber>(NewStart, NewStart + 1));
@@ -1197,8 +1151,7 @@ namespace UE::MLDeformer
 
 	void FMLTimeSliderController::SetSelectionRangeEnd(FFrameNumber NewEnd)
 	{
-		TRange<FFrameNumber> SelectionRange = TimeSliderArgs.SelectionRange.Get();
-
+		const TRange<FFrameNumber> SelectionRange = TimeSliderArgs.SelectionRange.Get();
 		if (SelectionRange.IsEmpty())
 		{
 			TimeSliderArgs.OnSelectionRangeChanged.ExecuteIfBound(TRange<FFrameNumber>(NewEnd - 1, NewEnd));
@@ -1213,43 +1166,34 @@ namespace UE::MLDeformer
 	{
 	public:
 		SLATE_BEGIN_ARGS(SAnimTimelineTransportControls) {}
-
 		SLATE_END_ARGS()
 
 		void Construct(const FArguments& InArgs, TWeakPtr<FMLDeformerEditorModel>& InModel);
-
 		void SetModel(TWeakPtr<FMLDeformerEditorModel> InModel)
 		{
 			Model = InModel;
 		}
 
 	private:
-		UAnimSingleNodeInstance* GetPreviewInstance() const;
-
 		FReply OnClick_Forward_Step();
-
 		FReply OnClick_Forward_End();
-
 		FReply OnClick_Backward_Step();
-
 		FReply OnClick_Backward_End();
-
 		FReply OnClick_Forward();
-
 		FReply OnClick_ToggleLoop();
 
 		bool IsLoopStatusOn() const;
-
 		EPlaybackMode::Type GetPlaybackMode() const;
+		UAnimSingleNodeInstance* GetPreviewInstance() const;
 
 	private:
-		/** Anim timeline model */
+		/** Anim timeline model. */
 		TWeakPtr<FMLDeformerEditorModel> Model;
 	};
 
 	void SAnimTimelineTransportControls::Construct(const FArguments& InArgs, TWeakPtr<FMLDeformerEditorModel>& InModel)
 	{
-		// Geoemetry Cache only supports forward playback, so backwards play is removed
+		// Geometry Cache only supports forward playback, so backwards play is removed.
 		Model = InModel;
 		FEditorWidgetsModule& EditorWidgetsModule = FModuleManager::LoadModuleChecked<FEditorWidgetsModule>("EditorWidgets");
 
@@ -1259,20 +1203,18 @@ namespace UE::MLDeformer
 		TransportControlArgs.OnBackwardStep = FOnClicked::CreateSP(this, &SAnimTimelineTransportControls::OnClick_Backward_Step);
 		TransportControlArgs.OnForwardEnd = FOnClicked::CreateSP(this, &SAnimTimelineTransportControls::OnClick_Forward_End);
 		TransportControlArgs.OnBackwardEnd = FOnClicked::CreateSP(this, &SAnimTimelineTransportControls::OnClick_Backward_End);
-
 		TransportControlArgs.OnGetPlaybackMode = FOnGetPlaybackMode::CreateSP(this, &SAnimTimelineTransportControls::GetPlaybackMode);
 
 		ChildSlot
-			[
-				EditorWidgetsModule.CreateTransportControl(TransportControlArgs)
-			];
+		[
+			EditorWidgetsModule.CreateTransportControl(TransportControlArgs)
+		];
 	}
 
 	FReply SAnimTimelineTransportControls::OnClick_Forward_Step()
 	{
 		if (Model.IsValid())
 		{
-			//Model.Pin()->OnPlayPressed();
 			FFrameNumber ScrubPosition = Model.Pin()->GetTickResScrubPosition();
 			ScrubPosition.Value += Model.Pin()->GetTicksPerFrame();
 			Model.Pin()->SetScrubPosition(ScrubPosition);
@@ -1285,7 +1227,7 @@ namespace UE::MLDeformer
 	{
 		if (Model.IsValid())
 		{
-			TRange<FFrameNumber> PlaybackRange = Model.Pin()->GetPlaybackRange();
+			const TRange<FFrameNumber> PlaybackRange = Model.Pin()->GetPlaybackRange();
 			FFrameNumber UpperBound = PlaybackRange.GetUpperBoundValue();
 			Model.Pin()->SetScrubPosition(UpperBound);
 			return FReply::Handled();
@@ -1297,7 +1239,6 @@ namespace UE::MLDeformer
 	{
 		if (Model.IsValid())
 		{
-			//Model.Pin()->OnPlayPressed(EPlaybackMode::Stopped);
 			FFrameNumber ScrubPosition = Model.Pin()->GetTickResScrubPosition();
 			ScrubPosition.Value -= Model.Pin()->GetTicksPerFrame();
 			Model.Pin()->SetScrubPosition(ScrubPosition);
@@ -1310,7 +1251,7 @@ namespace UE::MLDeformer
 	{
 		if (Model.IsValid())
 		{
-			TRange<FFrameNumber> PlaybackRange = Model.Pin()->GetPlaybackRange();
+			const TRange<FFrameNumber> PlaybackRange = Model.Pin()->GetPlaybackRange();
 			FFrameNumber LowerBound = PlaybackRange.GetLowerBoundValue();
 			Model.Pin()->SetScrubPosition(LowerBound);
 			return FReply::Handled();
@@ -1342,7 +1283,7 @@ namespace UE::MLDeformer
 	}
 
 	/**
-	 * An overlay that displays global information in the track area
+	 * An overlay that displays global information in the track area.
 	 */
 	class SAnimTimelineOverlay : public SCompoundWidget
 	{
@@ -1355,10 +1296,9 @@ namespace UE::MLDeformer
 		SLATE_ATTRIBUTE(bool, DisplayTickLines)
 			SLATE_ATTRIBUTE(bool, DisplayScrubPosition)
 			SLATE_ATTRIBUTE(FPaintPlaybackRangeArgs, PaintPlaybackRangeArgs)
-
 			SLATE_END_ARGS()
 
-			void Construct(const FArguments& InArgs, TSharedRef<FMLTimeSliderController> InTimeSliderController)
+		void Construct(const FArguments& InArgs, TSharedRef<FMLTimeSliderController> InTimeSliderController)
 		{
 			bDisplayScrubPosition = InArgs._DisplayScrubPosition;
 			bDisplayTickLines = InArgs._DisplayTickLines;
@@ -1367,19 +1307,18 @@ namespace UE::MLDeformer
 		}
 
 	private:
-		/** SWidget Interface */
+		/** SWidget Interface. */
 		virtual int32 OnPaint(const FPaintArgs& Args, const FGeometry& AllottedGeometry, const FSlateRect& MyCullingRect, FSlateWindowElementList& OutDrawElements, int32 LayerId, const FWidgetStyle& InWidgetStyle, bool bParentEnabled) const override;
 
 	private:
-		/** Controller for manipulating time */
+		/** Controller for manipulating time. */
 		TSharedPtr<FMLTimeSliderController> TimeSliderController;
-		/** Whether or not to display the scrub position */
+		/** Whether or not to display the scrub position. */
 		TAttribute<bool> bDisplayScrubPosition;
-		/** Whether or not to display tick lines */
+		/** Whether or not to display tick lines. */
 		TAttribute<bool> bDisplayTickLines;
-		/** User-supplied options for drawing playback range */
+		/** User-supplied options for drawing playback range. */
 		TAttribute<FPaintPlaybackRangeArgs> PaintPlaybackRangeArgs;
-
 	};
 
 	int32 SAnimTimelineOverlay::OnPaint(const FPaintArgs& Args, const FGeometry& AllottedGeometry, const FSlateRect& MyCullingRect, FSlateWindowElementList& OutDrawElements, int32 LayerId, const FWidgetStyle& InWidgetStyle, bool bParentEnabled) const
@@ -1394,42 +1333,42 @@ namespace UE::MLDeformer
 		}
 
 		TimeSliderController->OnPaintViewArea(AllottedGeometry, MyCullingRect, OutDrawElements, LayerId, ShouldBeEnabled(bParentEnabled), PaintArgs);
-
 		return SCompoundWidget::OnPaint(Args, AllottedGeometry, MyCullingRect, OutDrawElements, LayerId, InWidgetStyle, bParentEnabled);
 	}
 
 	void SMLDeformerTimeline::Construct(const FArguments& InArgs, FMLDeformerEditorToolkit* InEditor)
 	{
-		TWeakPtr<FMLDeformerEditorModel> WeakModel = InEditor->GetActiveModelPointer();
+		const TWeakPtr<FMLDeformerEditorModel> WeakModel = InEditor->GetActiveModelPointer();
 
 		Model = WeakModel;
 		OnReceivedFocus = InArgs._OnReceivedFocus;
 
-		int32 TickResolutionValue = Model.Pin()->GetTickResolution();
-		int32 SequenceFrameRate = FMath::RoundToInt(Model.Pin()->GetFrameRate());
+		const int32 TickResolutionValue = Model.Pin()->GetTickResolution();
+		const int32 SequenceFrameRate = FMath::RoundToInt(Model.Pin()->GetFrameRate());
 
-		ViewRange = MakeAttributeLambda([this](){
-				if (Model.IsValid())
-				{
-					TRange<double> Range = Model.Pin()->GetViewRange();
-					return FAnimatedRange(Range.GetLowerBoundValue(), Range.GetUpperBoundValue());
-				}
-				else
-				{
-					return FAnimatedRange(0.0, 0.0);
-				}
-			});
+		ViewRange = MakeAttributeLambda([this]()
+		{
+			if (Model.IsValid())
+			{
+				const TRange<double> Range = Model.Pin()->GetViewRange();
+				return FAnimatedRange(Range.GetLowerBoundValue(), Range.GetUpperBoundValue());
+			}
+			else
+			{
+				return FAnimatedRange(0.0, 0.0);
+			}
+		});
 
 		TAttribute<EFrameNumberDisplayFormats> DisplayFormat = MakeAttributeLambda([this]()
 		{
-				if (Model.IsValid())
-				{
-					return Model.Pin()->IsDisplayingFrames() ? EFrameNumberDisplayFormats::Frames : EFrameNumberDisplayFormats::Seconds;
-				}
-				else
-				{
-					return EFrameNumberDisplayFormats::Frames;
-				}
+			if (Model.IsValid())
+			{
+				return Model.Pin()->IsDisplayingFrames() ? EFrameNumberDisplayFormats::Frames : EFrameNumberDisplayFormats::Seconds;
+			}
+			else
+			{
+				return EFrameNumberDisplayFormats::Frames;
+			}
 		});
 
 		TAttribute<FFrameRate> TickResolution = MakeAttributeLambda([TickResolutionValue]()
@@ -1450,17 +1389,18 @@ namespace UE::MLDeformer
 			TimeSliderArgs.ScrubPosition = MakeAttributeLambda([WeakModel](){ return WeakModel.IsValid() ? WeakModel.Pin()->GetTickResScrubPosition() : FFrameTime(0); });
 			TimeSliderArgs.ViewRange = ViewRange;
 			TimeSliderArgs.PlaybackRange = MakeAttributeLambda([WeakModel](){ return WeakModel.IsValid() ? WeakModel.Pin()->GetPlaybackRange() : TRange<FFrameNumber>(0, 0); });
-			TimeSliderArgs.ClampRange = MakeAttributeLambda([this](){
-					if (Model.IsValid())
-					{
-						TRange<double> Range = Model.Pin()->GetWorkingRange();
-						return FAnimatedRange(Range.GetLowerBoundValue(), Range.GetUpperBoundValue());
-					}
-					else
-					{
-						return FAnimatedRange(0.0, 0.0);
-					}
-				});
+			TimeSliderArgs.ClampRange = MakeAttributeLambda([this]()
+			{
+				if (Model.IsValid())
+				{
+					const TRange<double> Range = Model.Pin()->GetWorkingRange();
+					return FAnimatedRange(Range.GetLowerBoundValue(), Range.GetUpperBoundValue());
+				}
+				else
+				{
+					return FAnimatedRange(0.0, 0.0);
+				}
+			});
 			TimeSliderArgs.DisplayRate = DisplayRate;
 			TimeSliderArgs.TickResolution = TickResolution;
 			TimeSliderArgs.OnViewRangeChanged = FOnViewRangeChanged::CreateSP(this, &SMLDeformerTimeline::HandleViewRangeChanged);
@@ -1475,12 +1415,12 @@ namespace UE::MLDeformer
 	
 		TSharedRef<FMLTimeSliderController> TimeSliderControllerRef = TimeSliderController.ToSharedRef();
 
-		// Create the top slider
+		// Create the top slider.
 		const bool bMirrorLabels = false;
 		ISequencerWidgetsModule& SequencerWidgets = FModuleManager::Get().LoadModuleChecked<ISequencerWidgetsModule>("SequencerWidgets");
 		TopTimeSlider = SequencerWidgets.CreateTimeSlider(TimeSliderControllerRef, bMirrorLabels);
 
-		// Create bottom time range slider
+		// Create bottom time range slider.
 		TSharedRef<ITimeSlider> BottomTimeRange = SequencerWidgets.CreateTimeRange(
 			FTimeRangeArgs(
 				EShowRange::ViewRange | EShowRange::WorkingRange | EShowRange::PlaybackRange,
@@ -1506,11 +1446,8 @@ namespace UE::MLDeformer
 			FillCoefficient_1.Bind(TAttribute<float>::FGetter::CreateSP(this, &SMLDeformerTimeline::GetColumnFillCoefficient, 1));
 		}
 
-		const int32 Column0 = 0, Column1 = 1;
-		const int32 Row0 = 0, Row1 = 1, Row2 = 2, Row3 = 3, Row4 = 4;
-
-		const float CommonPadding = 3.f;
-		const FMargin ResizeBarPadding(4.f, 0, 0, 0);
+		const float CommonPadding = 3.0f;
+		const FMargin ResizeBarPadding(4.0f, 0, 0, 0);
 
 		ChildSlot
 		[
@@ -1528,8 +1465,8 @@ namespace UE::MLDeformer
 						.FillColumn(0, FillCoefficient_0)
 						.FillColumn(1, FillCoefficient_1)
 
-						// outliner search box
-						+SGridPanel::Slot(Column0, Row0, SGridPanel::Layer(10))
+						// Outliner search box.
+						+SGridPanel::Slot(0, 0, SGridPanel::Layer(10))
 						[
 							SNew(SHorizontalBox)
 							+SHorizontalBox::Slot()
@@ -1546,7 +1483,7 @@ namespace UE::MLDeformer
 								.VAlign(VAlign_Center)
 								.HAlign(HAlign_Center)
 								[
-									// Current Play Time 
+									// Current play time.
 									SNew(SSpinBox<double>)
 									.Style(&FAppStyle::GetWidgetStyle<FSpinBoxStyle>("Sequencer.PlayTimeSpinBox"))
 									.Value_Lambda([this]() -> double
@@ -1570,16 +1507,16 @@ namespace UE::MLDeformer
 								]
 							]
 						]
-						// Transport controls
-						+SGridPanel::Slot(Column0, Row3, SGridPanel::Layer(10))
+						// Transport controls.
+						+SGridPanel::Slot(0, 3, SGridPanel::Layer(10))
 						.VAlign(VAlign_Center)
 						.HAlign(HAlign_Center)
 						[
 							TransportControls.ToSharedRef()
 						]
 
-						// Second column
-						+SGridPanel::Slot(Column1, Row0)
+						// Second column.
+						+SGridPanel::Slot(1, 0)
 						.Padding(ResizeBarPadding)
 						.RowSpan(2)
 						[
@@ -1590,12 +1527,12 @@ namespace UE::MLDeformer
 							]
 						]
 
-						+SGridPanel::Slot(Column1, Row0, SGridPanel::Layer(10))
+						+SGridPanel::Slot(1, 0, SGridPanel::Layer(10))
 						.Padding(ResizeBarPadding)
 						[
-							SNew( SBorder )
-							.BorderImage( FAppStyle::GetBrush("ToolPanel.GroupBorder") )
-							.BorderBackgroundColor( FLinearColor(.50f, .50f, .50f, 1.0f ) )
+							SNew(SBorder)
+							.BorderImage(FAppStyle::GetBrush("ToolPanel.GroupBorder"))
+							.BorderBackgroundColor(FLinearColor(0.5f, 0.5f, 0.5f, 1.0f))
 							.Padding(0)
 							.Clipping(EWidgetClipping::ClipToBounds)
 							[
@@ -1603,44 +1540,42 @@ namespace UE::MLDeformer
 							]
 						]
 
-						// Overlay that draws the tick lines
-						+SGridPanel::Slot(Column1, Row1, SGridPanel::Layer(10))
+						// Overlay that draws the tick lines.
+						+SGridPanel::Slot(1, 1, SGridPanel::Layer(10))
 						.Padding(ResizeBarPadding)
 						[
 							SNew(SAnimTimelineOverlay, TimeSliderControllerRef)
 							.Visibility( EVisibility::HitTestInvisible )
-							.DisplayScrubPosition( false )
-							.DisplayTickLines( true )
+							.DisplayScrubPosition(false)
+							.DisplayTickLines(true)
 							.Clipping(EWidgetClipping::ClipToBounds)
-							.PaintPlaybackRangeArgs(FPaintPlaybackRangeArgs(FAppStyle::GetBrush("Sequencer.Timeline.PlayRange_L"), FAppStyle::GetBrush("Sequencer.Timeline.PlayRange_R"), 6.f))
+							.PaintPlaybackRangeArgs(FPaintPlaybackRangeArgs(FAppStyle::GetBrush("Sequencer.Timeline.PlayRange_L"), FAppStyle::GetBrush("Sequencer.Timeline.PlayRange_R"), 6.0f))
 						]
 
-						// Overlay that draws the scrub position
-						+SGridPanel::Slot(Column1, Row1, SGridPanel::Layer(20))
+						// Overlay that draws the scrub position.
+						+SGridPanel::Slot(1, 1, SGridPanel::Layer(20))
 						.Padding(ResizeBarPadding)
 						[
 							SNew(SAnimTimelineOverlay, TimeSliderControllerRef)
-							.Visibility( EVisibility::HitTestInvisible )
-							.DisplayScrubPosition( true )
-							.DisplayTickLines( false )
+							.Visibility(EVisibility::HitTestInvisible)
+							.DisplayScrubPosition(true)
+							.DisplayTickLines(false)
 							.Clipping(EWidgetClipping::ClipToBounds)
 						]
 
-						// play range slider
-
-						+SGridPanel::Slot(Column1, Row3, SGridPanel::Layer(10))
+						// Play range slider.
+						+SGridPanel::Slot(1, 3, SGridPanel::Layer(10))
 						.Padding(ResizeBarPadding)
 						[
 							SNew(SBorder)
-							.BorderImage( FAppStyle::GetBrush("ToolPanel.GroupBorder") )
-							.BorderBackgroundColor( FLinearColor(0.5f, 0.5f, 0.5f, 1.0f ) )
+							.BorderImage(FAppStyle::GetBrush("ToolPanel.GroupBorder"))
+							.BorderBackgroundColor(FLinearColor(0.5f, 0.5f, 0.5f, 1.0f))
 							.Clipping(EWidgetClipping::ClipToBounds)
 							.Padding(0)
 							[
 								BottomTimeRange
 							]
 						]
-
 					]
 				]
 			]
@@ -1651,12 +1586,11 @@ namespace UE::MLDeformer
 	{
 		if(MouseEvent.GetEffectingButton() == EKeys::RightMouseButton)
 		{
-			FWidgetPath WidgetPath = MouseEvent.GetEventPath() != nullptr ? *MouseEvent.GetEventPath() : FWidgetPath();
+			const FWidgetPath WidgetPath = (MouseEvent.GetEventPath() != nullptr) ? *MouseEvent.GetEventPath() : FWidgetPath();
 
 			const bool bCloseAfterSelection = true;
 			FMenuBuilder MenuBuilder(bCloseAfterSelection, NULL);
-
-			TWeakPtr<FMLDeformerEditorModel> WeakModel = Model;
+			const TWeakPtr<FMLDeformerEditorModel> WeakModel = Model;
 
 			MenuBuilder.BeginSection("TimelineOptions", LOCTEXT("TimelineOptions", "Timeline Options") );
 			{
@@ -1692,8 +1626,8 @@ namespace UE::MLDeformer
 								EUserInterfaceActionType::RadioButton
 							);
 							InMenuBuilder.AddMenuEntry(
-									LOCTEXT("DisplayFrames", "Display Seconds"),
-									LOCTEXT("DisplayFrames_Tooltip", "Display Timeline in Seconds."),
+									LOCTEXT("DisplaySeconds", "Display Seconds"),
+									LOCTEXT("DisplaySeconds_Tooltip", "Display Timeline in Seconds."),
 									FSlateIcon(),
 									FUIAction(FExecuteAction::CreateLambda([WeakModel]() {
 											if (WeakModel.IsValid())
@@ -1724,14 +1658,12 @@ namespace UE::MLDeformer
 			MenuBuilder.EndSection();
 
 			FSlateApplication::Get().PushMenu(SharedThis(this), WidgetPath, MenuBuilder.MakeWidget(), FSlateApplication::Get().GetCursorPos(), FPopupTransitionEffect(FPopupTransitionEffect::ContextMenu));
-
 			return FReply::Handled();
 		}
 
 		return FReply::Unhandled();
 	}
 
-
 	TSharedRef<INumericTypeInterface<double>> SMLDeformerTimeline::GetNumericTypeInterface() const
 	{
 		return NumericTypeInterface.ToSharedRef();
@@ -1740,11 +1672,11 @@ namespace UE::MLDeformer
 	// FFrameRate::ComputeGridSpacing doesnt deal well with prime numbers, so we have a custom impl here
 	static bool ComputeGridSpacing(const FFrameRate& InFrameRate, float PixelsPerSecond, double& OutMajorInterval, int32& OutMinorDivisions, float MinTickPx, float DesiredMajorTickPx)
 	{
-		// First try built-in spacing
-		bool bResult = InFrameRate.ComputeGridSpacing(PixelsPerSecond, OutMajorInterval, OutMinorDivisions, MinTickPx, DesiredMajorTickPx);
-		if(!bResult || OutMajorInterval == 1.0)
+		// First try built-in spacing.
+		const bool bResult = InFrameRate.ComputeGridSpacing(PixelsPerSecond, OutMajorInterval, OutMinorDivisions, MinTickPx, DesiredMajorTickPx);
+		if (!bResult || OutMajorInterval == 1.0)
 		{
-			if (PixelsPerSecond <= 0.f)
+			if (PixelsPerSecond <= 0.0f)
 			{
 				return false;
 			}
@@ -1753,10 +1685,10 @@ namespace UE::MLDeformer
 
 			if (RoundedFPS > 0)
 			{
-				// Showing frames
+				// Showing frames.
 				TArray<int32, TInlineAllocator<10>> CommonBases;
 
-				// Divide the rounded frame rate by 2s, 3s or 5s recursively
+				// Divide the rounded frame rate by 2s, 3s or 5s recursively.
 				{
 					const int32 Denominators[] = { 2, 3, 5 };
 
@@ -1771,16 +1703,16 @@ namespace UE::MLDeformer
 						else
 						{ 
 							int32 LowestResult = LowestBase;
-							for(int32 Denominator : Denominators)
+							for (int32 Denominator : Denominators)
 							{
-								int32 Result = LowestBase / Denominator;
+								const int32 Result = LowestBase / Denominator;
 								if(Result > 0 && Result < LowestResult)
 								{
 									LowestResult = Result;
 								}
 							}
 
-							if(LowestResult < LowestBase)
+							if (LowestResult < LowestBase)
 							{
 								LowestBase = LowestResult;
 							}
@@ -1794,20 +1726,20 @@ namespace UE::MLDeformer
 
 				Algo::Reverse(CommonBases);
 
-				const int32 Scale     = FMath::CeilToInt(DesiredMajorTickPx / PixelsPerSecond * InFrameRate.AsDecimal());
-				const int32 BaseIndex = FMath::Min(Algo::LowerBound(CommonBases, Scale), CommonBases.Num()-1);
-				const int32 Base      = CommonBases[BaseIndex];
+				const int32 Scale = FMath::CeilToInt(DesiredMajorTickPx / PixelsPerSecond * InFrameRate.AsDecimal());
+				const int32 BaseIndex = FMath::Min(Algo::LowerBound(CommonBases, Scale), CommonBases.Num() - 1);
+				const int32 Base = CommonBases[BaseIndex];
 
-				int32 MajorIntervalFrames = FMath::CeilToInt(Scale / float(Base)) * Base;
-				OutMajorInterval  = MajorIntervalFrames * InFrameRate.AsInterval();
+				const int32 MajorIntervalFrames = FMath::CeilToInt(Scale / float(Base)) * Base;
+				OutMajorInterval = MajorIntervalFrames * InFrameRate.AsInterval();
 
-				// Find the lowest number of divisions we can show that's larger than the minimum tick size
+				// Find the lowest number of divisions we can show that's larger than the minimum tick size.
 				OutMinorDivisions = 0;
 				for (int32 DivIndex = 0; DivIndex < BaseIndex; ++DivIndex)
 				{
 					if (Base % CommonBases[DivIndex] == 0)
 					{
-						int32 MinorDivisions = MajorIntervalFrames/CommonBases[DivIndex];
+						const int32 MinorDivisions = MajorIntervalFrames/CommonBases[DivIndex];
 						if (OutMajorInterval / MinorDivisions * PixelsPerSecond >= MinTickPx)
 						{
 							OutMinorDivisions = MinorDivisions;
@@ -1854,22 +1786,20 @@ namespace UE::MLDeformer
 
 	bool SMLDeformerTimeline::GetGridMetrics(float PhysicalWidth, double& OutMajorInterval, int32& OutMinorDivisions) const
 	{
-		FSlateFontInfo SmallLayoutFont = FCoreStyle::GetDefaultFontStyle("Regular", 8);
-		TSharedRef<FSlateFontMeasure> FontMeasureService = FSlateApplication::Get().GetRenderer()->GetFontMeasureService();
+		const FSlateFontInfo SmallLayoutFont = FCoreStyle::GetDefaultFontStyle("Regular", 8);
+		const TSharedRef<FSlateFontMeasure> FontMeasureService = FSlateApplication::Get().GetRenderer()->GetFontMeasureService();
 		double FrameRate = 30.0; 
 		if (Model.IsValid())
 		{
 			FrameRate = Model.Pin()->GetFrameRate();
 		}
-		FFrameRate DisplayRate(FMath::Max(FMath::RoundToInt(FrameRate), 1), 1);
-		double BiggestTime = ViewRange.Get().GetUpperBoundValue();
-		FString TickString = NumericTypeInterface->ToString((BiggestTime * DisplayRate).FrameNumber.Value);
-		FVector2D MaxTextSize = FontMeasureService->Measure(TickString, SmallLayoutFont);
-
-		static float MajorTickMultiplier = 2.f;
-
-		float MinTickPx = MaxTextSize.X + 5.f;
-		float DesiredMajorTickPx = MaxTextSize.X * MajorTickMultiplier;
+		const FFrameRate DisplayRate(FMath::Max(FMath::RoundToInt(FrameRate), 1), 1);
+		const double BiggestTime = ViewRange.Get().GetUpperBoundValue();
+		const FString TickString = NumericTypeInterface->ToString((BiggestTime * DisplayRate).FrameNumber.Value);
+		const FVector2D MaxTextSize = FontMeasureService->Measure(TickString, SmallLayoutFont);
+		static const float MajorTickMultiplier = 2.0f;
+		const float MinTickPx = MaxTextSize.X + 5.0f;
+		const float DesiredMajorTickPx = MaxTextSize.X * MajorTickMultiplier;
 
 		if (PhysicalWidth > 0 && DisplayRate.AsDecimal() > 0)
 		{
@@ -1929,12 +1859,12 @@ namespace UE::MLDeformer
 
 	void SMLDeformerTimeline::SetPlayTime(double InFrameTime)
 	{
-		FFrameTime FrameTime((int)InFrameTime, 0.0f);
+		const FFrameTime FrameTime((int)InFrameTime, 0.0f);
 		if (Model.IsValid())
 		{
 			Model.Pin()->SetScrubPosition(FrameTime);
 		}
 	}
-}
+}	// namespace UE::MLDeformer
 
 #undef LOCTEXT_NAMESPACE
diff --git a/Engine/Plugins/Experimental/MLDeformer/MLDeformerFramework/Source/MLDeformerFrameworkEditor/Private/SMLDeformerTimeline.h b/Engine/Plugins/Experimental/MLDeformer/MLDeformerFramework/Source/MLDeformerFrameworkEditor/Private/SMLDeformerTimeline.h
--- a/Engine/Plugins/Experimental/MLDeformer/MLDeformerFramework/Source/MLDeformerFrameworkEditor/Private/SMLDeformerTimeline.h
+++ b/Engine/Plugins/Experimental/MLDeformer/MLDeformerFramework/Source/MLDeformerFrameworkEditor/Private/SMLDeformerTimeline.h
@@ -13,8 +13,8 @@
 #include "Framework/Commands/Commands.h"
 #include "MLDeformerEditorToolkit.h"
 
-class SSearchBox;
 struct FPaintPlaybackRangeArgs;
+class SSearchBox;
 class UAnimSingleNodeInstance;
 enum class EViewRangeInterpolation;
 enum class EFrameNumberDisplayFormats : uint8;
@@ -25,50 +25,45 @@ namespace UE::MLDeformer
 	class FMLDeformerEditorModel;
 	class FMLTimeSliderController;
 	class SAnimTimelineTransportControls; 
+
 	/**
 	 * Implements a timeline widget for the ML Deformer Editor.
 	 */
 	class SMLDeformerTimeline : public SCompoundWidget
 	{
 	public:
-		SLATE_BEGIN_ARGS(SMLDeformerTimeline) { }
-
-		/** The current total range of frame indices */
-		SLATE_ATTRIBUTE(FInt32Interval, ViewIndices)
-
-			/** Called when any widget contained within the anim timeline has received focus */
+		SLATE_BEGIN_ARGS(SMLDeformerTimeline) {}
+			/** The current total range of frame indices. */
+			SLATE_ATTRIBUTE(FInt32Interval, ViewIndices)
+			/** Called when any widget contained within the anim timeline has received focus. */
 			SLATE_EVENT(FSimpleDelegate, OnReceivedFocus)
-
 		SLATE_END_ARGS()
 
 	public:
-
 		/**
 		 * Construct this widget.
-		 *
 		 * @param InArgs The declaration data for this widget.
-		 * @param InModel The model for the anim timeline
+		 * @param InModel The model for the anim timeline.
 		 */
-
 		void Construct(const FArguments& InArgs, FMLDeformerEditorToolkit* InEditor);
-		/** SWidget interface */
+
+		// SWidget interface.
 		virtual FReply OnMouseButtonUp(const FGeometry& MyGeometry, const FPointerEvent& MouseEvent) override;
+		// End SWidget interface.
 
-		/** Compute a major grid interval and number of minor divisions to display */
+		/** Compute a major grid interval and number of minor divisions to display. */
 		bool GetGridMetrics(float PhysicalWidth, double& OutMajorInterval, int32& OutMinorDivisions) const;
 
-		/** Get the time slider controller */
+		/** Get the time slider controller. */
 		TSharedPtr<ITimeSliderController> GetTimeSliderController() const;
 
 		void SetModel(TWeakPtr<FMLDeformerEditorModel> InModel);
+
 	private:
 		/**
-		 * @return The fill percentage of the animation outliner
+		 * @return The fill percentage of the animation outliner.
 		 */
-		float GetColumnFillCoefficient(int32 ColumnIndex) const
-		{
-			return ColumnFillCoefficients[ColumnIndex];
-		}
+		float GetColumnFillCoefficient(int32 ColumnIndex) const { return ColumnFillCoefficients[ColumnIndex]; }
 
 		/** Get numeric Type interface for converting between frame numbers and display formats. */
 		TSharedRef<INumericTypeInterface<double>> GetNumericTypeInterface() const;
@@ -76,52 +71,50 @@ namespace UE::MLDeformer
 		/** Called when a column fill percentage is changed by a splitter slot. */
 		void OnColumnFillCoefficientChanged(float FillCoefficient, int32 ColumnIndex);
 
-		/** Handles an additive layer key being set */
-
+		/** Handles an additive layer key being set. */
 		void HandleScrubPositionChanged(FFrameTime NewScrubPosition, bool bIsScrubbing, bool bEvaluate);
-
 		void HandleViewRangeChanged(TRange<double> InRange, EViewRangeInterpolation InInterpolation);
 
 		double GetSpinboxDelta() const;
-
 		void SetPlayTime(double InFrameTime);
-
 		void SetDisplayFormat(EFrameNumberDisplayFormats InFormat);
 		bool IsDisplayFormatChecked(EFrameNumberDisplayFormats InFormat) const;
 
 	private:
-		/** Anim timeline model */
+		/** Anim timeline model. */
 		TWeakPtr<FMLDeformerEditorModel> Model;
 
 		/** The currently active editor model. */
 		UE::MLDeformer::FMLDeformerEditorToolkit* Editor;
 
-		/** The time slider controller */
+		/** The time slider controller. */
 		TSharedPtr<FMLTimeSliderController> TimeSliderController;
 
+		/** The anim timeline transport controls, which contains the forward, backward and step actions. */
 		TSharedPtr<SAnimTimelineTransportControls> TransportControls;
-		/** The top time slider widget */
+
+		/** The top time slider widget. */
 		TSharedPtr<ITimeSlider> TopTimeSlider;
 
 		/** The fill coefficients of each column in the grid. */
 		float ColumnFillCoefficients[2];
 
-		/** Called when the user has begun dragging the selection selection range */
+		/** Called when the user has begun dragging the selection selection range. */
 		FSimpleDelegate OnSelectionRangeBeginDrag;
 
-		/** Called when the user has finished dragging the selection selection range */
+		/** Called when the user has finished dragging the selection selection range. */
 		FSimpleDelegate OnSelectionRangeEndDrag;
 
-		/** Called when any widget contained within the anim timeline has received focus */
+		/** Called when any widget contained within the anim timeline has received focus. */
 		FSimpleDelegate OnReceivedFocus;
 
 		/** Numeric Type interface for converting between frame numbers and display formats. */
 		TSharedPtr<INumericTypeInterface<double>> NumericTypeInterface;
 
-		/** The view range */
+		/** The view range. */
 		TAttribute<FAnimatedRange> ViewRange;
 
-		/** Filter text used to search the tree */
+		/** Filter text used to search the tree. */
 		FText FilterText;
 	};
 
