[MLDeformer] Fixed a crash when switching skeletal meshes in some specific case.

diff --git a/Engine/Plugins/Experimental/MLDeformer/MLDeformerFramework/Source/MLDeformerFrameworkEditor/Private/MLDeformerGeomCacheSampler.cpp b/Engine/Plugins/Experimental/MLDeformer/MLDeformerFramework/Source/MLDeformerFrameworkEditor/Private/MLDeformerGeomCacheSampler.cpp
--- a/Engine/Plugins/Experimental/MLDeformer/MLDeformerFramework/Source/MLDeformerFrameworkEditor/Private/MLDeformerGeomCacheSampler.cpp
+++ b/Engine/Plugins/Experimental/MLDeformer/MLDeformerFramework/Source/MLDeformerFrameworkEditor/Private/MLDeformerGeomCacheSampler.cpp
@@ -31,25 +31,22 @@ namespace UE::MLDeformer
 		}
 
 		UGeometryCache* GeomCache = OnGetGeometryCache().IsBound() ? OnGetGeometryCache().Execute() : nullptr;
-		if (GeomCache != GeometryCacheComponent->GetGeometryCache())
-		{
-			GeometryCacheComponent->SetGeometryCache(GeomCache);
-			GeometryCacheComponent->SetManualTick(true);
-			GeometryCacheComponent->SetVisibility(false);
+		GeometryCacheComponent->SetGeometryCache(GeomCache);
+		GeometryCacheComponent->SetManualTick(true);
+		GeometryCacheComponent->SetVisibility(false);
 
-			// Generate mappings between the meshes in the SkeletalMesh and the geometry cache tracks.
-			FailedImportedMeshNames.Reset();
-			MeshMappings.Reset();
-			GenerateGeomCacheMeshMappings(
-				Model->GetSkeletalMesh(), 
-				GeomCache, 
-				MeshMappings, 
-				FailedImportedMeshNames,
-				VertexCountMisMatchNames);
+		// Generate mappings between the meshes in the SkeletalMesh and the geometry cache tracks.
+		FailedImportedMeshNames.Reset();
+		MeshMappings.Reset();
+		GenerateGeomCacheMeshMappings(
+			Model->GetSkeletalMesh(), 
+			GeomCache, 
+			MeshMappings, 
+			FailedImportedMeshNames,
+			VertexCountMisMatchNames);
 
-			GeomCacheMeshDatas.Reset(MeshMappings.Num());
-			GeomCacheMeshDatas.AddDefaulted(MeshMappings.Num());
-		}
+		GeomCacheMeshDatas.Reset(MeshMappings.Num());
+		GeomCacheMeshDatas.AddDefaulted(MeshMappings.Num());
 	}
 
 	void FMLDeformerGeomCacheSampler::Sample(int32 InAnimFrameIndex)
