Optimus: Don't use void for empty parameter list in data interface read functions.
There is at least one platform where shader compilation doesn't like that.
#rnx

diff --git a/Engine/Plugins/Animation/DeformerGraph/Shaders/Private/DataInterfaceCloth.ush b/Engine/Plugins/Animation/DeformerGraph/Shaders/Private/DataInterfaceCloth.ush
--- a/Engine/Plugins/Animation/DeformerGraph/Shaders/Private/DataInterfaceCloth.ush
+++ b/Engine/Plugins/Animation/DeformerGraph/Shaders/Private/DataInterfaceCloth.ush
@@ -11,12 +11,12 @@ float4x4 DI_LOCAL(ClothToLocal);
 Buffer<float4> DI_LOCAL(ClothBuffer);
 Buffer<float2> DI_LOCAL(ClothPositionsAndNormalsBuffer);
 
-DI_IMPL_READ(ReadNumVertices, uint, void)
+DI_IMPL_READ(ReadNumVertices, uint, )
 {
 	return DI_LOCAL(NumVertices);
 }
 
-DI_IMPL_READ(ReadClothToLocal, float4x4, void)
+DI_IMPL_READ(ReadClothToLocal, float4x4, )
 {
 #if ENABLE_DEFORMER_CLOTH
 	return DI_LOCAL(ClothToLocal);
diff --git a/Engine/Plugins/Animation/DeformerGraph/Shaders/Private/DataInterfaceConnectivity.ush b/Engine/Plugins/Animation/DeformerGraph/Shaders/Private/DataInterfaceConnectivity.ush
--- a/Engine/Plugins/Animation/DeformerGraph/Shaders/Private/DataInterfaceConnectivity.ush
+++ b/Engine/Plugins/Animation/DeformerGraph/Shaders/Private/DataInterfaceConnectivity.ush
@@ -7,7 +7,7 @@ uint DI_LOCAL(InputStreamStart);
 uint DI_LOCAL(MaxConnectedVertexCount);
 Buffer<uint> DI_LOCAL(ConnectivityBuffer);
 
-DI_IMPL_READ(ReadNumVertices, uint, void)
+DI_IMPL_READ(ReadNumVertices, uint, )
 {
 	return DI_LOCAL(NumVertices);
 }
diff --git a/Engine/Plugins/Animation/DeformerGraph/Shaders/Private/DataInterfaceDebugDraw.ush b/Engine/Plugins/Animation/DeformerGraph/Shaders/Private/DataInterfaceDebugDraw.ush
--- a/Engine/Plugins/Animation/DeformerGraph/Shaders/Private/DataInterfaceDebugDraw.ush
+++ b/Engine/Plugins/Animation/DeformerGraph/Shaders/Private/DataInterfaceDebugDraw.ush
@@ -53,7 +53,7 @@ struct FDebugDraw
 	}
 };
 
-DI_IMPL_READ(ReadDebugDraw, FDebugDraw, void)
+DI_IMPL_READ(ReadDebugDraw, FDebugDraw, )
 {
 	FDebugDraw DebugDraw;
 	DebugDraw.LocalToWorld = DI_LOCAL(LocalToWorld);
diff --git a/Engine/Plugins/Animation/DeformerGraph/Shaders/Private/DataInterfaceJiggleSpring.ush b/Engine/Plugins/Animation/DeformerGraph/Shaders/Private/DataInterfaceJiggleSpring.ush
--- a/Engine/Plugins/Animation/DeformerGraph/Shaders/Private/DataInterfaceJiggleSpring.ush
+++ b/Engine/Plugins/Animation/DeformerGraph/Shaders/Private/DataInterfaceJiggleSpring.ush
@@ -11,7 +11,7 @@ Buffer<float> DI_LOCAL(StiffnessWeightsBuffer);
 Buffer<float> DI_LOCAL(DampingWeightsBuffer);
 
 
-DI_IMPL_READ(ReadNumVertices, uint, void)
+DI_IMPL_READ(ReadNumVertices, uint, )
 {
 	return DI_LOCAL(NumVertices);
 }
diff --git a/Engine/Plugins/Animation/DeformerGraph/Shaders/Private/DataInterfaceMorphTarget.ush b/Engine/Plugins/Animation/DeformerGraph/Shaders/Private/DataInterfaceMorphTarget.ush
--- a/Engine/Plugins/Animation/DeformerGraph/Shaders/Private/DataInterfaceMorphTarget.ush
+++ b/Engine/Plugins/Animation/DeformerGraph/Shaders/Private/DataInterfaceMorphTarget.ush
@@ -7,7 +7,7 @@ uint DI_LOCAL(InputStreamStart);
 
 Buffer<float> DI_LOCAL(MorphBuffer);
 
-DI_IMPL_READ(ReadNumVertices, uint, void)
+DI_IMPL_READ(ReadNumVertices, uint, )
 {
 	return DI_LOCAL(NumVertices);
 }
diff --git a/Engine/Plugins/Animation/DeformerGraph/Shaders/Private/DataInterfaceRawBuffer.ush b/Engine/Plugins/Animation/DeformerGraph/Shaders/Private/DataInterfaceRawBuffer.ush
--- a/Engine/Plugins/Animation/DeformerGraph/Shaders/Private/DataInterfaceRawBuffer.ush
+++ b/Engine/Plugins/Animation/DeformerGraph/Shaders/Private/DataInterfaceRawBuffer.ush
@@ -10,7 +10,7 @@ StructuredBuffer<BUFFER_TYPE> DI_LOCAL(BufferSRV);
 #endif
 RWStructuredBuffer<BUFFER_TYPE> DI_LOCAL(BufferUAV);
 
-DI_IMPL_READ(ReadNumValues, uint, void)
+DI_IMPL_READ(ReadNumValues, uint, )
 {
 	return DI_LOCAL(BufferSize);
 }
diff --git a/Engine/Plugins/Animation/DeformerGraph/Shaders/Private/DataInterfaceSkeletalMeshRead.ush b/Engine/Plugins/Animation/DeformerGraph/Shaders/Private/DataInterfaceSkeletalMeshRead.ush
--- a/Engine/Plugins/Animation/DeformerGraph/Shaders/Private/DataInterfaceSkeletalMeshRead.ush
+++ b/Engine/Plugins/Animation/DeformerGraph/Shaders/Private/DataInterfaceSkeletalMeshRead.ush
@@ -30,12 +30,12 @@ Buffer<uint> DI_LOCAL(InputWeightLookupStream);
 Buffer<uint> DI_LOCAL(DuplicatedIndicesIndices);
 Buffer<uint> DI_LOCAL(DuplicatedIndices);
 
-DI_IMPL_READ(ReadNumVertices, uint, void)
+DI_IMPL_READ(ReadNumVertices, uint, )
 {
 	return DI_LOCAL(NumVertices);
 }
 
-DI_IMPL_READ(ReadNumTriangles, uint, void)
+DI_IMPL_READ(ReadNumTriangles, uint, )
 {
 	return DI_LOCAL(NumTriangles);
 }
diff --git a/Engine/Plugins/Animation/DeformerGraph/Shaders/Private/DataInterfaceSkeleton.ush b/Engine/Plugins/Animation/DeformerGraph/Shaders/Private/DataInterfaceSkeleton.ush
--- a/Engine/Plugins/Animation/DeformerGraph/Shaders/Private/DataInterfaceSkeleton.ush
+++ b/Engine/Plugins/Animation/DeformerGraph/Shaders/Private/DataInterfaceSkeleton.ush
@@ -18,7 +18,7 @@ Buffer<float4> DI_LOCAL(BoneMatrices);
 Buffer<uint> DI_LOCAL(InputWeightStream);
 Buffer<uint> DI_LOCAL(InputWeightLookupStream);
 
-DI_IMPL_READ(ReadNumVertices, uint, void)
+DI_IMPL_READ(ReadNumVertices, uint, )
 {
 	return DI_LOCAL(NumVertices);
 }
diff --git a/Engine/Plugins/Animation/DeformerGraph/Shaders/Private/DataInterfaceSkinnedMesh.ush b/Engine/Plugins/Animation/DeformerGraph/Shaders/Private/DataInterfaceSkinnedMesh.ush
--- a/Engine/Plugins/Animation/DeformerGraph/Shaders/Private/DataInterfaceSkinnedMesh.ush
+++ b/Engine/Plugins/Animation/DeformerGraph/Shaders/Private/DataInterfaceSkinnedMesh.ush
@@ -20,17 +20,17 @@ Buffer<float2> DI_LOCAL(UVInputBuffer);
 Buffer<uint> DI_LOCAL(DuplicatedIndicesIndices);
 Buffer<uint> DI_LOCAL(DuplicatedIndices);
 
-DI_IMPL_READ(ReadNumVertices, uint, void)
+DI_IMPL_READ(ReadNumVertices, uint, )
 {
 	return DI_LOCAL(NumVertices);
 }
 
-DI_IMPL_READ(ReadNumTriangles, uint, void)
+DI_IMPL_READ(ReadNumTriangles, uint, )
 {
 	return DI_LOCAL(NumTriangles);
 }
 
-DI_IMPL_READ(ReadNumUVChannels, uint, void)
+DI_IMPL_READ(ReadNumUVChannels, uint, )
 {
 	return DI_LOCAL(NumUVChannels);
 }
diff --git a/Engine/Plugins/Animation/DeformerGraph/Shaders/Private/DataInterfaceSkinnedMeshExec.ush b/Engine/Plugins/Animation/DeformerGraph/Shaders/Private/DataInterfaceSkinnedMeshExec.ush
--- a/Engine/Plugins/Animation/DeformerGraph/Shaders/Private/DataInterfaceSkinnedMeshExec.ush
+++ b/Engine/Plugins/Animation/DeformerGraph/Shaders/Private/DataInterfaceSkinnedMeshExec.ush
@@ -4,7 +4,7 @@
 
 int3 DI_LOCAL(NumThreads);
 
-DI_IMPL_READ(ReadNumThreads, int3, void)
+DI_IMPL_READ(ReadNumThreads, int3, )
 {
 	return DI_LOCAL(NumThreads);
 }
diff --git a/Engine/Plugins/Animation/DeformerGraph/Shaders/Private/DataInterfaceSkinnedMeshWrite.ush b/Engine/Plugins/Animation/DeformerGraph/Shaders/Private/DataInterfaceSkinnedMeshWrite.ush
--- a/Engine/Plugins/Animation/DeformerGraph/Shaders/Private/DataInterfaceSkinnedMeshWrite.ush
+++ b/Engine/Plugins/Animation/DeformerGraph/Shaders/Private/DataInterfaceSkinnedMeshWrite.ush
@@ -10,7 +10,7 @@ RWBuffer<float> DI_LOCAL(PositionBufferUAV);
 RWBuffer<TANGENT_RWBUFFER_FORMAT> DI_LOCAL(TangentBufferUAV);
 RWBuffer<UNORM float4> DI_LOCAL(ColorBufferUAV);
 
-DI_IMPL_READ(ReadNumVertices, uint, void)
+DI_IMPL_READ(ReadNumVertices, uint, )
 {
 	return DI_LOCAL(NumVertices);
 }
diff --git a/Engine/Plugins/Experimental/MLDeformer/MLDeformerFramework/Shaders/Private/MLDeformerGraphDataInterface.ush b/Engine/Plugins/Experimental/MLDeformer/MLDeformerFramework/Shaders/Private/MLDeformerGraphDataInterface.ush
--- a/Engine/Plugins/Experimental/MLDeformer/MLDeformerFramework/Shaders/Private/MLDeformerGraphDataInterface.ush
+++ b/Engine/Plugins/Experimental/MLDeformer/MLDeformerFramework/Shaders/Private/MLDeformerGraphDataInterface.ush
@@ -9,7 +9,7 @@ float DI_LOCAL(Weight);
 Buffer<float> DI_LOCAL(PositionDeltaBuffer);
 Buffer<uint> DI_LOCAL(VertexMapBuffer);
 
-DI_IMPL_READ(ReadNumVertices, uint, void)
+DI_IMPL_READ(ReadNumVertices, uint, )
 {
 	return DI_LOCAL(NumVertices);
 }
diff --git a/Engine/Plugins/Experimental/MLDeformer/MLDeformerFramework/Shaders/Private/MLDeformerGraphHeatMapDataInterface.ush b/Engine/Plugins/Experimental/MLDeformer/MLDeformerFramework/Shaders/Private/MLDeformerGraphHeatMapDataInterface.ush
--- a/Engine/Plugins/Experimental/MLDeformer/MLDeformerFramework/Shaders/Private/MLDeformerGraphHeatMapDataInterface.ush
+++ b/Engine/Plugins/Experimental/MLDeformer/MLDeformerFramework/Shaders/Private/MLDeformerGraphHeatMapDataInterface.ush
@@ -12,7 +12,7 @@ uint DI_LOCAL(GroundTruthBufferSize);
 StructuredBuffer<float3> DI_LOCAL(PositionGroundTruthBuffer);
 Buffer<uint> DI_LOCAL(VertexMapBuffer);
 
-DI_IMPL_READ(ReadNumVertices, uint, void)
+DI_IMPL_READ(ReadNumVertices, uint, )
 {
 	return DI_LOCAL(NumVertices);
 }
@@ -27,17 +27,17 @@ DI_IMPL_READ(ReadPositionGroundTruth, float3, uint VertexIndex)
 	return DI_LOCAL(PositionGroundTruthBuffer)[BufferIndex];
 }
 
-DI_IMPL_READ(ReadHeatMapMode, uint, void)
+DI_IMPL_READ(ReadHeatMapMode, uint, )
 {
 	return DI_LOCAL(HeatMapMode);
 }
 
-DI_IMPL_READ(ReadHeatMapMax, float, void)
+DI_IMPL_READ(ReadHeatMapMax, float, )
 {
 	return DI_LOCAL(HeatMapMax);
 }
 
-DI_IMPL_READ(ReadGroundTruthLerp, float, void)
+DI_IMPL_READ(ReadGroundTruthLerp, float, )
 {
 	return DI_LOCAL(GroundTruthLerp);
 }
