[MLDeformer]
- Play button in training mode now plays things, but doesn't calculate deltas, in order to speed things up.
- API refactors, now requiring less code to create new models by introducing some new base classes you can inherit from.
- Some tooltip fixes and improvements.
- Documented more of the API.
- Added a SetWeight and GetWeight blueprint function to control the weight of the deformer through blueprint.

diff --git a/Engine/Plugins/Experimental/MLDeformer/MLDeformerFramework/Source/MLDeformerFramework/Private/MLDeformerGeomCacheHelpers.cpp b/Engine/Plugins/Experimental/MLDeformer/MLDeformerFramework/Source/MLDeformerFramework/Private/MLDeformerGeomCacheHelpers.cpp
--- a/Engine/Plugins/Experimental/MLDeformer/MLDeformerFramework/Source/MLDeformerFramework/Private/MLDeformerGeomCacheHelpers.cpp
+++ b/Engine/Plugins/Experimental/MLDeformer/MLDeformerFramework/Source/MLDeformerFramework/Private/MLDeformerGeomCacheHelpers.cpp
@@ -2,18 +2,14 @@
 
 #include "MLDeformerGeomCacheHelpers.h"
 #include "MLDeformerModel.h"
-
 #include "GeometryCache.h"
 #include "GeometryCacheMeshData.h"
 #include "GeometryCacheTrack.h"
-
 #include "Animation/AnimSequence.h"
 #include "Engine/SkeletalMesh.h"
 #include "Rendering/SkeletalMeshModel.h"
 #include "Rendering/SkeletalMeshLODModel.h"
 
-#include "NeuralNetwork.h"
-
 #define LOCTEXT_NAMESPACE "MLDeformerGeomCacheHelpers"
 
 namespace UE::MLDeformer
@@ -236,9 +232,9 @@ namespace UE::MLDeformer
 		int32 InLODIndex,
 		float InSampleTime,
 		const TArray<UE::MLDeformer::FMLDeformerGeomCacheMeshMapping>& InMeshMappings,
-		const USkeletalMesh* SkelMesh,
+		const USkeletalMesh* InSkelMesh,
 		const UGeometryCache* InGeometryCache,
-		const FTransform& AlignmentTransform,
+		const FTransform& InAlignmentTransform,
 		TArray<FVector3f>& OutPositions)
 	{
 		if (InGeometryCache == nullptr)
@@ -246,12 +242,12 @@ namespace UE::MLDeformer
 			return;
 		}
 
-		if (!ensure(SkelMesh != nullptr))
+		if (!ensure(InSkelMesh != nullptr))
 		{
 			return;
 		}
 
-		const FSkeletalMeshModel* ImportedModel = SkelMesh->GetImportedModel();
+		const FSkeletalMeshModel* ImportedModel = InSkelMesh->GetImportedModel();
 		if (!ensure(ImportedModel != nullptr))
 		{
 			return;
@@ -283,7 +279,7 @@ namespace UE::MLDeformer
 				const int32 GeomCacheVertexIndex = MeshMapping.SkelMeshToTrackVertexMap[VertexIndex];
 				if (GeomCacheVertexIndex != INDEX_NONE && GeomCacheMeshData.Positions.IsValidIndex(GeomCacheVertexIndex))
 				{
-					const FVector3f GeomCacheVertexPos = (FVector3f)AlignmentTransform.TransformPosition((FVector)GeomCacheMeshData.Positions[GeomCacheVertexIndex]);
+					const FVector3f GeomCacheVertexPos = (FVector3f)InAlignmentTransform.TransformPosition((FVector)GeomCacheMeshData.Positions[GeomCacheVertexIndex]);
 					OutPositions[SkinnedVertexIndex] = GeomCacheVertexPos;
 				}
 			}
diff --git a/Engine/Plugins/Experimental/MLDeformer/MLDeformerFramework/Source/MLDeformerFramework/Private/MLDeformerGeomCacheModel.cpp b/Engine/Plugins/Experimental/MLDeformer/MLDeformerFramework/Source/MLDeformerFramework/Private/MLDeformerGeomCacheModel.cpp
new file mode 100644
--- /dev/null
+++ b/Engine/Plugins/Experimental/MLDeformer/MLDeformerFramework/Source/MLDeformerFramework/Private/MLDeformerGeomCacheModel.cpp
@@ -0,0 +1,77 @@
+// Copyright Epic Games, Inc. All Rights Reserved.
+#include "MLDeformerGeomCacheModel.h"
+#include "MLDeformerGeomCacheVizSettings.h"
+#include "MLDeformerGeomCacheHelpers.h"
+#include "MLDeformerComponent.h"
+#include "UObject/Object.h"
+#include "UObject/UObjectGlobals.h"
+#include "GeometryCache.h"
+
+#define LOCTEXT_NAMESPACE "MLDeformerGeomCacheModel"
+
+#if WITH_EDITOR
+void UMLDeformerGeomCacheModel::UpdateNumTargetMeshVertices()
+{
+	NumTargetMeshVerts = UE::MLDeformer::ExtractNumImportedGeomCacheVertices(GeometryCache);
+}
+
+UMLDeformerGeomCacheVizSettings* UMLDeformerGeomCacheModel::GetGeomCacheVizSettings() const
+{
+	return Cast<UMLDeformerGeomCacheVizSettings>(VizSettings);
+}
+
+void UMLDeformerGeomCacheModel::SetAssetEditorOnlyFlags()
+{
+	// Set the flags for the base class, which filters out the training anim sequence.
+	UMLDeformerModel::SetAssetEditorOnlyFlags();
+
+	// The training geometry cache is something we don't want to package.
+	if (GeometryCache)
+	{
+		GeometryCache->GetPackage()->SetPackageFlags(PKG_EditorOnly);
+	}
+
+	// Filter the viz settings specific assets.
+	UMLDeformerGeomCacheVizSettings* GeomCacheVizSettings = GetGeomCacheVizSettings();
+	if (GeomCacheVizSettings)
+	{
+		if (GeomCacheVizSettings->GetTestGroundTruth())
+		{
+			GeomCacheVizSettings->GetTestGroundTruth()->GetPackage()->SetPackageFlags(PKG_EditorOnly);
+		}
+	}
+}
+#endif // WITH_EDITOR
+
+#if WITH_EDITORONLY_DATA
+void UMLDeformerGeomCacheModel::SampleGroundTruthPositions(float SampleTime, TArray<FVector3f>& OutPositions)
+{
+	const UMLDeformerGeomCacheVizSettings* GeomCacheVizSettings = GetGeomCacheVizSettings();
+	check(GeomCacheVizSettings);
+
+	UGeometryCache* GeomCache = GeomCacheVizSettings->GetTestGroundTruth();
+	if (GeomCache == nullptr)
+	{
+		OutPositions.Reset();
+		return;
+	}
+
+	if (MeshMappings.IsEmpty())
+	{
+		TArray<FString> FailedImportedMeshNames;
+		TArray<FString> VertexMisMatchNames;
+		UE::MLDeformer::GenerateGeomCacheMeshMappings(SkeletalMesh, GeomCache, MeshMappings, FailedImportedMeshNames, VertexMisMatchNames);
+	}
+
+	UE::MLDeformer::SampleGeomCachePositions(
+		0,
+		SampleTime,
+		MeshMappings,
+		SkeletalMesh,
+		GeomCache,
+		AlignmentTransform,
+		OutPositions);
+}
+#endif
+
+#undef LOCTEXT_NAMESPACE
diff --git a/Engine/Plugins/Experimental/MLDeformer/MLDeformerFramework/Source/MLDeformerFramework/Private/MLDeformerModelInstance.cpp b/Engine/Plugins/Experimental/MLDeformer/MLDeformerFramework/Source/MLDeformerFramework/Private/MLDeformerModelInstance.cpp
--- a/Engine/Plugins/Experimental/MLDeformer/MLDeformerFramework/Source/MLDeformerFramework/Private/MLDeformerModelInstance.cpp
+++ b/Engine/Plugins/Experimental/MLDeformer/MLDeformerFramework/Source/MLDeformerFramework/Private/MLDeformerModelInstance.cpp
@@ -78,7 +78,6 @@ FString UMLDeformerModelInstance::CheckCompatibility(USkeletalMeshComponent* InS
 	{
 		ErrorText += InputInfo->GenerateCompatibilityErrorString(SkelMesh);
 		ErrorText += "\n";
-		check(!ErrorText.IsEmpty());
 		if (LogIssues)
 		{
 			UE_LOG(LogMLDeformer, Error, TEXT("ML Deformer '%s' isn't compatible with Skeletal Mesh '%s'.\nReason(s):\n%s"), 
@@ -183,7 +182,10 @@ int64 UMLDeformerModelInstance::SetBoneTransforms(float* OutputBuffer, int64 Out
 	const UMLDeformerInputInfo* InputInfo = Model->GetInputInfo();
 	const int32 AssetNumBones = InputInfo->GetNumBones();
 	int64 Index = StartIndex;
-	check((Index + AssetNumBones * 6) <= OutputBufferSize); // Make sure we don't write past the OutputBuffer. (6 because of two columns of the 3x3 rotation matrix)
+
+	// Make sure we don't write past the OutputBuffer. (6 because of two columns of the 3x3 rotation matrix)
+	checkf((Index + AssetNumBones * 6) <= OutputBufferSize, TEXT("Writing bones past the end of the input buffer."));
+
 	for (int32 BoneIndex = 0; BoneIndex < AssetNumBones; ++BoneIndex)
 	{
 		const FMatrix RotationMatrix = TrainingBoneTransforms[BoneIndex].GetRotation().ToMatrix();
@@ -206,8 +208,8 @@ int64 UMLDeformerModelInstance::SetCurveValues(float* OutputBuffer, int64 Output
 
 	// Write the weights into the output buffer.
 	int64 Index = StartIndex;
-	const int32 AssetNumCurves = InputInfo->GetNumCurves();
-	check((Index + AssetNumCurves) <= OutputBufferSize); // Make sure we don't write past the OutputBuffer.
+	const int32 AssetNumCurves = InputInfo->GetNumCurves();	
+	checkf((Index + AssetNumCurves) <= OutputBufferSize, TEXT("Writing curves past the end of the input buffer"));
 
 	// Write the curve weights to the output buffer.
 	UAnimInstance* AnimInstance = SkeletalMeshComponent->GetAnimInstance();
diff --git a/Engine/Plugins/Experimental/MLDeformer/MLDeformerFramework/Source/MLDeformerFramework/Private/MLDeformerMorphModel.cpp b/Engine/Plugins/Experimental/MLDeformer/MLDeformerFramework/Source/MLDeformerFramework/Private/MLDeformerMorphModel.cpp
new file mode 100644
--- /dev/null
+++ b/Engine/Plugins/Experimental/MLDeformer/MLDeformerFramework/Source/MLDeformerFramework/Private/MLDeformerMorphModel.cpp
@@ -0,0 +1,132 @@
+// Copyright Epic Games, Inc. All Rights Reserved.
+
+#include "MLDeformerMorphModel.h"
+#include "MLDeformerMorphModelInstance.h"
+#include "MLDeformerModelInstance.h"
+#include "MLDeformerComponent.h"
+#include "UObject/Object.h"
+#include "UObject/UObjectGlobals.h"
+#include "Components/SkinnedMeshComponent.h"
+#include "Components/SkeletalMeshComponent.h"
+#include "Rendering/MorphTargetVertexInfoBuffers.h"
+
+#define LOCTEXT_NAMESPACE "MLDeformerMorphModel"
+
+UMLDeformerMorphModel::UMLDeformerMorphModel(const FObjectInitializer& ObjectInitializer)
+	: Super(ObjectInitializer)
+{
+	MorphTargetSet = MakeShared<FExternalMorphSet>();
+	MorphTargetSet->Name = GetClass()->GetFName();
+}
+
+void UMLDeformerMorphModel::Serialize(FArchive& Archive)
+{
+	TRACE_CPUPROFILER_EVENT_SCOPE(UMLDeformerMorphModel::Serialize)
+
+	Super::Serialize(Archive);
+
+	// Check if we have initialized our compressed morph buffers.
+	bool bHasMorphData = false;
+	if (Archive.IsSaving())
+	{
+		bHasMorphData = MorphTargetSet.IsValid() ? MorphTargetSet->MorphBuffers.IsMorphCPUDataValid() : false;
+	}
+	Archive << bHasMorphData;
+
+	// Load or save the compressed morph buffers, if they exist.
+	if (bHasMorphData)
+	{
+		check(MorphTargetSet.IsValid());
+		Archive << MorphTargetSet->MorphBuffers;
+	}
+}
+
+UMLDeformerModelInstance* UMLDeformerMorphModel::CreateModelInstance(UMLDeformerComponent* Component)
+{
+	return NewObject<UMLDeformerMorphModelInstance>(Component);
+}
+
+void UMLDeformerMorphModel::PostMLDeformerComponentInit(UMLDeformerModelInstance* ModelInstance)
+{
+	if (ModelInstance->HasPostInitialized())
+	{
+		return;
+	}
+
+	TRACE_CPUPROFILER_EVENT_SCOPE(UMLDeformerMorphModel::PostMLDeformerComponentInit)
+
+	Super::PostMLDeformerComponentInit(ModelInstance);
+
+	// Register the external morph targets buffer to the skinned mesh component.
+	USkeletalMeshComponent* SkelMeshComponent = ModelInstance->GetSkeletalMeshComponent();
+	if (SkelMeshComponent && SkelMeshComponent->GetSkeletalMeshAsset())
+	{	
+		// If this triggers, please set the ExternalMorphSetID in your model's constructor to a value that represents your model type.
+		checkf(ExternalMorphSetID != -1, TEXT("Please configure your ExternalMorphSetID member to a unique ID."));
+
+		// Register the morph set. This overwrites the existing one for this model, if it already exists.
+		// Only add to LOD 0 for now.
+		const int32 LOD = 0;
+		SkelMeshComponent->AddExternalMorphSet(LOD, ExternalMorphSetID, MorphTargetSet);
+
+		// When we're in editor mode, keep the CPU data around, so we can re-initialize when needed.
+#if WITH_EDITOR
+		MorphTargetSet->MorphBuffers.SetEmptyMorphCPUDataOnInitRHI(false);
+#else
+		MorphTargetSet->MorphBuffers.SetEmptyMorphCPUDataOnInitRHI(true);
+#endif
+
+		// Release the render resources, but only in an editor build.
+		// The non-editor build shouldn't do this, as then it can't initialize again. The non-editor build assumes
+		// that the data doesn't change and we don't need to re-init.
+		// In the editor build we have to re-initialize the render resources as the morph targets can change after (re)training, so
+		// that is why we release them here, and intialize them again after.
+		FMorphTargetVertexInfoBuffers& MorphBuffers = MorphTargetSet->MorphBuffers;
+#if WITH_EDITOR
+		BeginReleaseResource(&MorphBuffers);
+#endif
+
+		// Reinitialize the GPU compressed buffers.
+		if (MorphBuffers.IsMorphCPUDataValid() && MorphBuffers.GetNumMorphs() > 0)
+		{
+			// In a non-editor build this will clear the CPU data.
+			// That also means it can't re-init the resources later on again.
+			BeginInitResource(&MorphBuffers);
+		}
+
+		// Update the weight information in the Skeletal Mesh.
+		SkelMeshComponent->RefreshExternalMorphTargetWeights();
+
+		ModelInstance->SetHasPostInitialized(true);
+	}
+}
+
+void UMLDeformerMorphModel::SetMorphTargetDeltas(const TArray<float>& Deltas)
+{
+	FloatArrayToVector3Array(Deltas, MorphTargetDeltas);
+}
+
+int32 UMLDeformerMorphModel::GetMorphTargetDeltaStartIndex(int32 BlendShapeIndex) const
+{
+	if (MorphTargetDeltas.Num() == 0)
+	{
+		return INDEX_NONE;
+	}
+
+	return GetNumBaseMeshVerts() * BlendShapeIndex;
+}
+
+void UMLDeformerMorphModel::BeginDestroy()
+{
+	if (MorphTargetSet.IsValid())
+	{
+		// Release and flush, waiting for the release to have completed, 
+		// If we don't do this we can get an error that we destroy a render resource that is still initialized,
+		// as the release happens in another thread.
+		ReleaseResourceAndFlush(&MorphTargetSet->MorphBuffers);
+		MorphTargetSet.Reset();
+	}
+	Super::BeginDestroy();
+}
+
+#undef LOCTEXT_NAMESPACE
diff --git a/Engine/Plugins/Experimental/MLDeformer/MLDeformerFramework/Source/MLDeformerFramework/Private/MLDeformerMorphModelInstance.cpp b/Engine/Plugins/Experimental/MLDeformer/MLDeformerFramework/Source/MLDeformerFramework/Private/MLDeformerMorphModelInstance.cpp
new file mode 100644
--- /dev/null
+++ b/Engine/Plugins/Experimental/MLDeformer/MLDeformerFramework/Source/MLDeformerFramework/Private/MLDeformerMorphModelInstance.cpp
@@ -0,0 +1,67 @@
+// Copyright Epic Games, Inc. All Rights Reserved.
+
+#include "MLDeformerMorphModelInstance.h"
+#include "MLDeformerMorphModel.h" 
+#include "NeuralNetwork.h"
+#include "Components/SkeletalMeshComponent.h"
+
+// Run the neural network, which calculates its outputs, which are the weights of our morph targets.
+void UMLDeformerMorphModelInstance::RunNeuralNetwork(float ModelWeight)
+{
+	TRACE_CPUPROFILER_EVENT_SCOPE(UNeuralMorphModelInstance::RunNeuralNetwork)
+
+	const UMLDeformerMorphModel* MorphModel = Cast<UMLDeformerMorphModel>(Model);
+	if (MorphModel == nullptr)
+	{
+		return;
+	}
+
+	const int32 ExternalMorphSetID = MorphModel->ExternalMorphSetID;
+	// If this check fails please set this member to some value larger than 0 in your model's constructor.
+	checkf(ExternalMorphSetID != -1, TEXT("Please set the ExternalMorphSetID member value to a unique value for your model type."));
+
+	// Grab the weight data for this morph set.
+	// This could potentially fail if we are applying this deformer to the wrong skeletal mesh component.
+	const int LOD = 0;	// For now we only support LOD 0, as we can't setup an ML Deformer per LOD yet.
+	FExternalMorphSetWeights* WeightData = SkeletalMeshComponent->GetExternalMorphWeights(LOD).MorphSets.Find(ExternalMorphSetID);
+	if (WeightData == nullptr)
+	{
+		return;
+	}
+
+	// If our model is active, we want to run the neural network and update the morph weights
+	// with the values that the neural net calculated for us.
+	if (ModelWeight > 0.0f)
+	{
+		// Perform the neural network inference, which updates the output tensor.
+		UMLDeformerModelInstance::RunNeuralNetwork(ModelWeight);
+
+		// Get the output tensor, read the values and use them as morph target weights inside the skeletal mesh component.
+		const UNeuralNetwork* NeuralNetwork = Model->GetNeuralNetwork();
+		const FNeuralTensor& OutputTensor = NeuralNetwork->GetOutputTensorForContext(NeuralNetworkInferenceHandle);
+		const int32 NumNetworkWeights = OutputTensor.Num();
+		const int32 NumMorphTargets = WeightData->Weights.Num();;
+
+		// If we have the expected amount of weights.
+		// +1 because we always have an extra morph target that represents the means, with fixed weight of 1.
+		if (NumMorphTargets == NumNetworkWeights + 1)
+		{
+			// Set the first morph target, which represents the means, to a weight of 1.0, as it always needs to be fully active.
+			WeightData->Weights[0] = 1.0f * ModelWeight;
+
+			// Update all generated morph target weights with the values calculated by our neural network.
+			for (int32 MorphIndex = 0; MorphIndex < NumNetworkWeights; ++MorphIndex)
+			{
+				WeightData->Weights[MorphIndex + 1] = OutputTensor.At<float>(MorphIndex) * ModelWeight;
+			}
+		}
+	}
+	else
+	{
+		// Reset all weights to zero.
+		for (float& Weight : WeightData->Weights)
+		{
+			Weight = 0.0f;
+		}
+	}
+}
diff --git a/Engine/Plugins/Experimental/MLDeformer/MLDeformerFramework/Source/MLDeformerFramework/Public/MLDeformerAsset.h b/Engine/Plugins/Experimental/MLDeformer/MLDeformerFramework/Source/MLDeformerFramework/Public/MLDeformerAsset.h
--- a/Engine/Plugins/Experimental/MLDeformer/MLDeformerFramework/Source/MLDeformerFramework/Public/MLDeformerAsset.h
+++ b/Engine/Plugins/Experimental/MLDeformer/MLDeformerFramework/Source/MLDeformerFramework/Public/MLDeformerAsset.h
@@ -10,7 +10,7 @@ class UMLDeformerModel;
 
 /**
  * The machine learning deformer asset class.
- * This class contains a Model property, through which all the magic happens.
+ * This class contains a Model property, through which most functionality happens.
  */
 UCLASS(BlueprintType, hidecategories=Object)
 class MLDEFORMERFRAMEWORK_API UMLDeformerAsset 
diff --git a/Engine/Plugins/Experimental/MLDeformer/MLDeformerFramework/Source/MLDeformerFramework/Public/MLDeformerComponent.h b/Engine/Plugins/Experimental/MLDeformer/MLDeformerFramework/Source/MLDeformerFramework/Public/MLDeformerComponent.h
--- a/Engine/Plugins/Experimental/MLDeformer/MLDeformerFramework/Source/MLDeformerFramework/Public/MLDeformerComponent.h
+++ b/Engine/Plugins/Experimental/MLDeformer/MLDeformerFramework/Source/MLDeformerFramework/Public/MLDeformerComponent.h
@@ -32,18 +32,42 @@ public:
 	void Deactivate() override;
 	// ~END UActorComponent overrides.
 
-	/** Setup the ML Deformer, by picking the deformer asset and skeletal mesh component. */
+	/** 
+	 * Setup the ML Deformer, by picking the deformer asset and skeletal mesh component. 
+	 * Call this when you want to assign a specific ML Deformer to a specific skeletal mesh component.
+	 * This is used in case there are multiple skeletal mesh components on your actor.
+	 * On default the first skeletal mesh component it finds will be used. This function allows you to change that behavior.
+	 * Keep in mind that the ML Deformer asset has to be trained using the same Skeletal Mesh as used by the skeletal mesh component you pick.
+	 * @param InDeformerAsset The ML Deformer asset to apply to the specified skeletal mesh component.	 
+	 * @param InSkelMeshComponent The skeletal mesh component to apply the specified deformer to.
+	 */
 	UFUNCTION(BlueprintCallable, Category = "MLDeformer")
 	void SetupComponent(UMLDeformerAsset* InDeformerAsset, USkeletalMeshComponent* InSkelMeshComponent);
 
+	/** 
+	 * Get the current ML Deformer weight. A value of 0 means it is fully disabled, while 1 means fully active.
+	 * Values can be anything between 0 and 1.
+	 */
+	UFUNCTION(BlueprintCallable, Category = "MLDeformer")
+	float GetWeight() const { return Weight; }
+
+	/**
+	 * Set the ML Deformer weight. This determines how active the deformer is. You can see it as a blend weight.
+	 * A value of 0 means it is inactive. Certain calculations will be skipped in that case.
+	 * A value of 1 means it is fully active.
+	 * Values between 0 and 1 blend between the two states.
+	 * Call this after you call SetupComponent.
+	 * @param NormalizedWeightValue The weight value.
+	 */
+	UFUNCTION(BlueprintCallable, Category = "MLDeformer")
+	void SetWeight(float NormalizedWeightValue) { Weight = FMath::Clamp<float>(NormalizedWeightValue, 0.0f, 1.0f);  }
+
 #if WITH_EDITOR
 	void PostEditChangeProperty(FPropertyChangedEvent& PropertyChangedEvent) override;
 #endif
 
 	UMLDeformerAsset* GetDeformerAsset() { return DeformerAsset; }
 	void SetDeformerAsset(UMLDeformerAsset* InDeformerAsset) { DeformerAsset = InDeformerAsset; }
-	float GetWeight() const { return Weight; }
-	void SetWeight(float NormalizedWeightValue) { Weight = NormalizedWeightValue;  }
 	UMLDeformerModelInstance* GetModelInstance() const { return ModelInstance; }
 	USkeletalMeshComponent* GetSkeletalMeshComponent() const { return SkelMeshComponent.Get(); }
 
@@ -52,6 +76,12 @@ protected:
 	virtual void TickComponent(float DeltaTime, enum ELevelTick TickType, FActorComponentTickFunction* ThisTickFunction) override;
 	// ~END AActorComponent overrides.
 
+	/** 
+	 * Initialize the component. 
+	 * This releases any existing deformer instance that is active, and creates a new one.
+	 * It then also calls PostMLDeformerComponentInit.
+	 * This method is called automatically by SetupComponent.
+	 */
 	void Init();
 
 	/** Bind to the MLDeformerModel's NeuralNetworkModifyDelegate. */
diff --git a/Engine/Plugins/Experimental/MLDeformer/MLDeformerFramework/Source/MLDeformerFramework/Public/MLDeformerCurveReference.h b/Engine/Plugins/Experimental/MLDeformer/MLDeformerFramework/Source/MLDeformerFramework/Public/MLDeformerCurveReference.h
--- a/Engine/Plugins/Experimental/MLDeformer/MLDeformerFramework/Source/MLDeformerFramework/Public/MLDeformerCurveReference.h
+++ b/Engine/Plugins/Experimental/MLDeformer/MLDeformerFramework/Source/MLDeformerFramework/Public/MLDeformerCurveReference.h
@@ -6,6 +6,7 @@
 #include "UObject/ObjectMacros.h"
 #include "MLDeformerCurveReference.generated.h"
 
+/** A reference to a curve implemented as a name. */
 USTRUCT()
 struct MLDEFORMERFRAMEWORK_API FMLDeformerCurveReference
 {
@@ -38,7 +39,7 @@ struct MLDEFORMERFRAMEWORK_API FMLDeformerCurveReference
 		return true;
 	}
 
-	/** Name of curve. */
+	/** The name of the curve. */
 	UPROPERTY(EditAnywhere, Category = AnimCurveReference)
 	FName CurveName;
 };
diff --git a/Engine/Plugins/Experimental/MLDeformer/MLDeformerFramework/Source/MLDeformerFramework/Public/MLDeformerGeomCacheHelpers.h b/Engine/Plugins/Experimental/MLDeformer/MLDeformerFramework/Source/MLDeformerFramework/Public/MLDeformerGeomCacheHelpers.h
--- a/Engine/Plugins/Experimental/MLDeformer/MLDeformerFramework/Source/MLDeformerFramework/Public/MLDeformerGeomCacheHelpers.h
+++ b/Engine/Plugins/Experimental/MLDeformer/MLDeformerFramework/Source/MLDeformerFramework/Public/MLDeformerGeomCacheHelpers.h
@@ -3,7 +3,6 @@
 #pragma once
 
 #include "CoreMinimal.h"
-#include "MLDeformerModule.h"
 
 class UGeometryCache;
 class USkeletalMesh;
@@ -12,30 +11,85 @@ class UAnimSequence;
 namespace UE::MLDeformer
 {
 #if WITH_EDITORONLY_DATA
-	/** Maps SkeletalMesh imported meshes to geometry cache tracks. */
+	/** A mapping between a geometry cache track and a mesh inside a USkeletalMesh. */
 	struct FMLDeformerGeomCacheMeshMapping
 	{
-		int32 MeshIndex = INDEX_NONE;	// The imported model's mesh info index.
-		int32 TrackIndex = INDEX_NONE;	// The geom cache track that this mesh is mapped to.
+		int32 MeshIndex = INDEX_NONE;	// The imported model's mesh info index, inside the skeletal mesh.
+		int32 TrackIndex = INDEX_NONE;	// The geometry cache track that this mesh is mapped to.
 		TArray<int32> SkelMeshToTrackVertexMap;	// This maps imported model individual meshes to the geomcache track's mesh data.
-		TArray<int32> ImportedVertexToRenderVertexMap; // Map the imported dcc vertex number to a render vertex. This is just one of the duplicates, which shares the same position.
+		TArray<int32> ImportedVertexToRenderVertexMap; // Map the imported DCC vertex number to a render vertex. This is just one of the duplicates, which shares the same position.
 	};
 
-	// Error checks.
+	/**
+	 * Check whether the geometry cache has any issues when used in combination with a given skeletal mesh.
+	 * Please note that this does not check the actual geometry data for matching vertex counts.
+	 * What it does check is if the geometry cache has geometry at all, and whether it has imported vertex numbers 
+	 * and whether flattened tracks was enabled, while there are multiple meshes in the skeletal mesh.
+	 * @param InSkeletalMesh The skeletal mesh that is being used.
+	 * @param InGeomCache The geometry cache that is being used.
+	 * @return The error text in case there are errors, or an empty text object in case there were no errors.
+	 */
 	MLDEFORMERFRAMEWORK_API FText GetGeomCacheErrorText(USkeletalMesh* InSkeletalMesh, UGeometryCache* InGeomCache);
+
+	/**
+	 * Check for any issues between ageometry cache and an anim sequence when used together.
+	 * This will check if the duration of both are the same or not.
+	 * @param InGeomCache The geometry cache object.
+	 * @param InAnimSequence The animation sequence object.
+	 * @return The error text in case there are errors, or an empty text object in case there were no errors.
+	 */
 	MLDEFORMERFRAMEWORK_API FText GetGeomCacheAnimSequenceErrorText(UGeometryCache* InGeomCache, UAnimSequence* InAnimSequence);
+
+	/**
+	 * Check for any mesh related issues between a skeletal mesh and geometry cache.
+	 * This will check if the Skeletal Mesh contains all meshes that are required for the geometry cache tracks.
+	 * It also checks for vertex count mismatches.
+	 * @param InSkelMesh The skeletal mesh object.
+	 * @param InGeomCache The geometry cache object.
+	 * @return The error text in case there are errors, or an empty text object in case there were no errors.
+	 */
 	MLDEFORMERFRAMEWORK_API FText GetGeomCacheMeshMappingErrorText(USkeletalMesh* InSkelMesh, UGeometryCache* InGeomCache);
 
-	// Geom cache operations.
+	/**
+	 * Get the number of imported geometry cache vertices.
+	 * Imported vertices are the vertex count as you see in the DCC, so a cube would have 8.
+	 * This is different from the render vertex count.
+	 * @param GeometryCache The geometry cache to get the imported vertex count for.
+	 * @return The number of vertices as seen in the DCC, so for example 8 for a cube.
+	 */
 	MLDEFORMERFRAMEWORK_API int32 ExtractNumImportedGeomCacheVertices(UGeometryCache* GeometryCache);
+
+	/**
+	 * Generate the mapping between geometry cache tracks and meshes inside a specific skeletal mesh.
+	 * This will check if the Skeletal Mesh contains all meshes that are required for the geometry cache tracks.
+	 * It also checks for vertex count mismatches.
+	 * This function is called by the GetGeomCacheMeshMappingErrorText function, in case you are just interested in the error checking part.
+	 * @param SkelMesh The skeletal mesh object.
+	 * @param GeomCache The geometry cache object.
+	 * @param OutMeshMapping The mapping that will be generated. This information maps each geometry cache map with a given compatible mesh inside a skeletal mesh.
+	 * @param OutFailedImportedMeshNames The array of geometry cache track names for which no mesh could be found inside the skeletal mesh object.
+	 * @param OutVertexMisMatchNames The array of geometry cache track names for which the vertex counts didn't match the related meshes inside the skeletal mesh.	 
+	 */
 	MLDEFORMERFRAMEWORK_API void GenerateGeomCacheMeshMappings(USkeletalMesh* SkelMesh, UGeometryCache* GeomCache, TArray<FMLDeformerGeomCacheMeshMapping>& OutMeshMappings, TArray<FString>& OutFailedImportedMeshNames, TArray<FString>& OutVertexMisMatchNames);
+
+	/**
+	 * Sample the vertex position data of a geometry cache, at a given time stamp.
+	 * This basically allows you to sample ground truth or training target vertex positions.
+	 * @param InLODIndex The LOD level to sample for.
+	 * @param InSampleTime The time to take the sample at, in seconds.
+	 * @param InMeshMappings The geometry cache to skeletal mesh mappings, which can be generated using the GenerateGeomCacheMeshMapings method.
+	 * @param InSkelMesh The skeletal mesh object.
+	 * @param InGeometryCache The geometry cache object.
+	 * @param InAlignmentTransform This is the transformation that will be applied on the sampled positions, as post process. This can be used to rotate or scale the position data.
+	 * @param OutPosiitons The resulting positions, as sampled at the specified time. This array will be resized internally.
+	 */
 	MLDEFORMERFRAMEWORK_API void SampleGeomCachePositions(
 		int32 InLODIndex,
 		float InSampleTime,
 		const TArray<FMLDeformerGeomCacheMeshMapping>& InMeshMappings,
-		const USkeletalMesh* SkelMesh,
+		const USkeletalMesh* InSkelMesh,
 		const UGeometryCache* InGeometryCache,
-		const FTransform& AlignmentTransform,
+		const FTransform& InAlignmentTransform,
 		TArray<FVector3f>& OutPositions);
 #endif	// #if WITH_EDITORONLY_DATA
 }	// namespace UE::MLDeformer
diff --git a/Engine/Plugins/Experimental/MLDeformer/MLDeformerFramework/Source/MLDeformerFramework/Public/MLDeformerGeomCacheModel.h b/Engine/Plugins/Experimental/MLDeformer/MLDeformerFramework/Source/MLDeformerFramework/Public/MLDeformerGeomCacheModel.h
new file mode 100644
--- /dev/null
+++ b/Engine/Plugins/Experimental/MLDeformer/MLDeformerFramework/Source/MLDeformerFramework/Public/MLDeformerGeomCacheModel.h
@@ -0,0 +1,84 @@
+// Copyright Epic Games, Inc. All Rights Reserved.
+#pragma once
+
+#include "CoreMinimal.h"
+#include "MLDeformerModel.h"
+#include "MLDeformerVizSettings.h"
+#include "MLDeformerInputInfo.h"
+#include "MLDeformerGeomCacheHelpers.h"
+#include "UObject/Object.h"
+#include "MLDeformerGeomCacheModel.generated.h"
+
+class UMLDeformerAsset;
+class UMLDeformerGeomCacheVizSettings;
+class UGeometryCache;
+
+/**
+ * A ML Deformer model that has a geometry cache as target mesh.
+ * Use this in combination with UMLDeformerGeomCacheVizSettings, FMLDeformerGeomCacheEditorModel, FMLDeformerGeomCacheVizSettingsDetails and FMLDeformerGeomCacheModelDetails.
+ */
+UCLASS()
+class MLDEFORMERFRAMEWORK_API UMLDeformerGeomCacheModel
+	: public UMLDeformerModel
+{
+	GENERATED_BODY()
+
+public:
+	// UMLDeformerModel overrides.
+	virtual FString GetDisplayName() const override { return "GeomCache Based Model"; }
+#if WITH_EDITORONLY_DATA
+	virtual bool HasTrainingGroundTruth() const override { return (GeometryCache.Get() != nullptr); }
+	virtual void SampleGroundTruthPositions(float SampleTime, TArray<FVector3f>& OutPositions) override;
+#endif
+#if WITH_EDITOR
+	virtual void UpdateNumTargetMeshVertices() override;
+	virtual void SetAssetEditorOnlyFlags() override;
+	// ~END UMLDeformerModel overrides.
+
+	/**
+	 * Get the visualization settings, already cast to a UMLDeformerGeomCacheVizSettings type.
+	 * @return A pointer to the visualization settings.
+	 */
+	UMLDeformerGeomCacheVizSettings* GetGeomCacheVizSettings() const;
+#endif
+
+#if WITH_EDITORONLY_DATA
+	/** 
+	 * Get the geometry cache that represents the target deformation.
+	 * This is our training target.
+	 * @return A pointer to the geometry cache, in read-only mode.
+	 */
+	const UGeometryCache* GetGeometryCache() const { return GeometryCache; }
+
+	/** 
+	 * Get the geometry cache that represents the target deformation.
+	 * This is our training target.
+	 * @return A pointer to the geometry cache.
+	 */
+	UGeometryCache* GetGeometryCache() { return GeometryCache; }
+
+	/**
+	 * Get the mapping between geometry cache tracks and meshes inside the skeletal mesh.
+	 * This lets us know what parts of the skeletal mesh are related to what geometry cache tracks.
+	 * Once we have that, we can calculate deltas between the two.
+	 * @return The geometry cache mesh to skeletal mesh mappings.
+	 */
+	TArray<UE::MLDeformer::FMLDeformerGeomCacheMeshMapping>& GetGeomCacheMeshMappings() { return MeshMappings; }
+
+	/**
+	 * Get the mapping between geometry cache tracks and meshes inside the skeletal mesh.
+	 * This lets us know what parts of the skeletal mesh are related to what geometry cache tracks.
+	 * Once we have that, we can calculate deltas between the two.
+	 * @return The geometry cache mesh to skeletal mesh mappings, in read-only mode.
+	 */
+	const TArray<UE::MLDeformer::FMLDeformerGeomCacheMeshMapping>& GetGeomCacheMeshMappings() const { return MeshMappings; }
+
+public:
+	/** The mappings between the geometry cache tracks and skeletal mesh imported meshes. */
+	TArray<UE::MLDeformer::FMLDeformerGeomCacheMeshMapping> MeshMappings;
+
+	/** The geometry cache that represents the target deformations. */
+	UPROPERTY(EditAnywhere, Category = "Target Mesh")
+	TObjectPtr<UGeometryCache> GeometryCache = nullptr;
+#endif // WITH_EDITORONLY_DATA
+};
diff --git a/Engine/Plugins/Experimental/MLDeformer/MLDeformerFramework/Source/MLDeformerFramework/Public/MLDeformerGeomCacheVizSettings.h b/Engine/Plugins/Experimental/MLDeformer/MLDeformerFramework/Source/MLDeformerFramework/Public/MLDeformerGeomCacheVizSettings.h
new file mode 100644
--- /dev/null
+++ b/Engine/Plugins/Experimental/MLDeformer/MLDeformerFramework/Source/MLDeformerFramework/Public/MLDeformerGeomCacheVizSettings.h
@@ -0,0 +1,39 @@
+// Copyright Epic Games, Inc. All Rights Reserved.
+
+#pragma once
+
+#include "CoreMinimal.h"
+#include "UObject/Object.h"
+#include "MLDeformerVizSettings.h"
+#include "MLDeformerGeomCacheVizSettings.generated.h"
+
+class UGeometryCache;
+
+/**
+ * The vizualization settings for a model that has a geometry cache.
+ * This can be used in combination with a UMLDeformerGeomCacheModel.
+ */
+UCLASS()
+class MLDEFORMERFRAMEWORK_API UMLDeformerGeomCacheVizSettings
+	: public UMLDeformerVizSettings
+{
+	GENERATED_BODY()
+
+public:
+#if WITH_EDITORONLY_DATA
+	// UMLDeformerModel overrides.
+	virtual bool HasTestGroundTruth() const override { return !GroundTruth.IsNull(); }
+	// ~END UMLDeformerModel overrides.
+
+	/**
+	 * Get the test ground truth geometry cache, which represents the ground truth version of the test animation sequence.
+	 * @return A pointer to the geometry cache object.
+	 */
+	UGeometryCache* GetTestGroundTruth() const { return GroundTruth; }
+
+public:
+	/** The geometry cache that represents the ground truth of the test anim sequence. */
+	UPROPERTY(EditAnywhere, Category = "Test Assets")
+	TObjectPtr<UGeometryCache> GroundTruth = nullptr;
+#endif // WITH_EDITORONLY_DATA
+};
diff --git a/Engine/Plugins/Experimental/MLDeformer/MLDeformerFramework/Source/MLDeformerFramework/Public/MLDeformerGraphDataInterface.h b/Engine/Plugins/Experimental/MLDeformer/MLDeformerFramework/Source/MLDeformerFramework/Public/MLDeformerGraphDataInterface.h
--- a/Engine/Plugins/Experimental/MLDeformer/MLDeformerFramework/Source/MLDeformerFramework/Public/MLDeformerGraphDataInterface.h
+++ b/Engine/Plugins/Experimental/MLDeformer/MLDeformerFramework/Source/MLDeformerFramework/Public/MLDeformerGraphDataInterface.h
@@ -19,6 +19,7 @@ class UNeuralNetwork;
 class USkeletalMeshComponent;
 class UMLDeformerModel;
 
+/** Defines the default base ML Deformer shader parameters. */
 #define MLDEFORMER_SHADER_PARAMETERS() \
 	SHADER_PARAMETER(uint32, NumVertices) \
 	SHADER_PARAMETER(uint32, InputStreamStart) \
@@ -108,6 +109,7 @@ class MLDEFORMERFRAMEWORK_API UMLDeformerGraphDataProvider
 	GENERATED_BODY()
 
 public:
+	/** The deformer component that this data provider works on. */
 	UPROPERTY(EditAnywhere, BlueprintReadWrite, Category = Binding)
 	TObjectPtr<UMLDeformerComponent> DeformerComponent = nullptr;
 
@@ -132,13 +134,13 @@ namespace UE::MLDeformer
 		// ~END FComputeDataProviderRenderProxy overrides.
 
 	protected:
-		FSkeletalMeshObject* SkeletalMeshObject;
+		FSkeletalMeshObject* SkeletalMeshObject = nullptr;
 		TObjectPtr<UNeuralNetwork> NeuralNetwork = nullptr;
 		FRHIShaderResourceView* VertexMapBufferSRV = nullptr;
 		FRDGBuffer* Buffer = nullptr;
 		FRDGBufferSRV* BufferSRV = nullptr;
-		float Weight = 1.0f;
 		int32 NeuralNetworkInferenceHandle = -1;
+		float Weight = 1.0f;
 		bool bCanRunNeuralNet = false;
 	};
 }	// namespace UE::MLDeformer
diff --git a/Engine/Plugins/Experimental/MLDeformer/MLDeformerFramework/Source/MLDeformerFramework/Public/MLDeformerInputInfo.h b/Engine/Plugins/Experimental/MLDeformer/MLDeformerFramework/Source/MLDeformerFramework/Public/MLDeformerInputInfo.h
--- a/Engine/Plugins/Experimental/MLDeformer/MLDeformerFramework/Source/MLDeformerFramework/Public/MLDeformerInputInfo.h
+++ b/Engine/Plugins/Experimental/MLDeformer/MLDeformerFramework/Source/MLDeformerFramework/Public/MLDeformerInputInfo.h
@@ -8,7 +8,6 @@
 
 class USkeletalMesh;
 class USkeletalMeshComponent;
-class UMLDeformerModel;
 
 /**
  * The neural network input information.
@@ -99,9 +98,28 @@ public:
 	 */
 	int32 GetNumCurves() const { return CurveNames.Num(); }
 
+	/**
+	 * Get the list of bone names that should be included as training inputs.
+	 * @return An array of strings, one for each bone name.
+	 */
 	TArray<FString>& GetBoneNameStrings() { return BoneNameStrings; }
+
+	/**
+	 * Get the list of curve names that should be included as training inputs.
+	 * @return An array of strings, one for each curve name.
+	 */
 	TArray<FString>& GetCurveNameStrings() { return CurveNameStrings; }
+
+	/**
+	 * Get the list of bone names that should be included as training inputs.
+	 * @return An array of FName objects, one for each bone.
+	 */
 	TArray<FName>& GetBoneNames() { return BoneNames; }
+
+	/**
+	 * Get the list of curve names that should be included as training inputs.
+	 * @return An array of FName objects, one for each curve.
+	 */
 	TArray<FName>& GetCurveNames() { return CurveNames; }
 
 	/**
@@ -135,19 +153,30 @@ public:
 	void ExtractBoneRotations(USkeletalMeshComponent* SkelMeshComponent, TArray<float>& OutRotations) const;
 
 	/**
-	 * Get the number of imported vertices in the base mesh.
+	 * Get the number of imported vertices in the base mesh, which is the linear skinned skeletal mesh.
 	 * @return The number of imported vertices in the base mesh.
 	 */
 	int32 GetNumBaseMeshVertices() const { return NumBaseMeshVertices; }
 
+	/**
+	 * Set the number of vertices of our base mesh, which is the linear skinned skeletal mesh.
+	 * The vertex count must be the same as in the DCC, so not the render mesh vertex count.
+	 * @param NumVerts The number of vertices of our base mesh.
+	 */
 	void SetNumBaseVertices(int32 NumVerts) { NumBaseMeshVertices = NumVerts; }
 
 	/**
-	 * Get the number of imported vertices in the target mesh.
+	 * Get the number of imported vertices in the target mesh, which is our training target/ground truth mesh.
+	 * The number of vertices is the vertex count as in the DCC.
 	 * @return The number of imported vertices in the target mesh.
 	 */
 	int32 GetNumTargetMeshVertices() const { return NumTargetMeshVertices; }
 
+	/**
+	 * Set the number of target mesh vertices. The target mesh is our training target/ground truth mesh.
+	 * The vertex count must be the same as in the DCC, so not the render mesh vertex count.
+	 * @param NumVerts The number of vertices in the target mesh.
+	 */
 	void SetNumTargetVertices(int32 NumVerts) { NumTargetMeshVertices = NumVerts; }
 
 protected:
diff --git a/Engine/Plugins/Experimental/MLDeformer/MLDeformerFramework/Source/MLDeformerFramework/Public/MLDeformerModel.h b/Engine/Plugins/Experimental/MLDeformer/MLDeformerFramework/Source/MLDeformerFramework/Public/MLDeformerModel.h
--- a/Engine/Plugins/Experimental/MLDeformer/MLDeformerFramework/Source/MLDeformerFramework/Public/MLDeformerModel.h
+++ b/Engine/Plugins/Experimental/MLDeformer/MLDeformerFramework/Source/MLDeformerFramework/Public/MLDeformerModel.h
@@ -23,14 +23,29 @@ class UNeuralNetwork;
 
 namespace UE::MLDeformer
 {
-	/** The vertex map on the GPU. */
+	/**
+	 * The vertex map, but in a GPU buffer. 
+	 * This map basically has a DCC vertex number for every render vertex.
+	 * So if a cube requires 32 render vertices, there will be 32 ints inside this buffer, and each item in this buffer
+	 * will in this specific example case contain a value between 0 and 7, as a cube has only 8 vertices.
+	 */
 	class FVertexMapBuffer : public FVertexBufferWithSRV
 	{
 	public:
+		/**
+		 * Initialize the GPU buffer based on some array with integers.
+		 * @param InVertexMap The array of ints we want to store on the GPU. The size of this array should be equal to the number of render vertices of the skeletal mesh.
+		 */
 		void Init(const TArray<int32>& InVertexMap)	{ VertexMap = InVertexMap; }
 
 	private:
+		/**
+		 * This does the actual render resource init, which means this creates and fills the buffer on the GPU.
+		 * After it successfully initializes, it will empty our VertexMap member array to not store the data in both GPU memory and main memory.
+		 */
 		void InitRHI() override;
+
+		/** The array of integers we want to store on the GPU. This buffer will be emptied after successfully calling InitRHI. */
 		TArray<int32> VertexMap;
 	};
 }	// namespace UE::MLDeformer
@@ -51,25 +66,108 @@ class MLDEFORMERFRAMEWORK_API UMLDeformerModel
 public:
 	virtual ~UMLDeformerModel() = default;
 
-	// Methods you might want to override.
+	/**
+	 * Initialize the ML Deformer model.
+	 * This will update the DeformerAsset and InputInfo properties. It internally calls CreateInputInfo if no InputInfo has been set yet.
+	 * @param InDeformerAsset The deformer asset that this model will be part of.
+	 */
 	virtual void Init(UMLDeformerAsset* InDeformerAsset);
+
+	/**
+	 * Initialize the data that should be stored on the GPU. 
+	 * This base class will store the VertexMap on the GPU by initializing the VertexMapBuffer member.
+	 */
 	virtual void InitGPUData();
+
+	/**
+	 * Create the input info for this model.
+	 * @return A pointer to the newly created input info.
+	 */
 	virtual UMLDeformerInputInfo* CreateInputInfo();
+
+	/**
+	 * Create a new instance of this model, to be used in combination with a specific component.
+	 * @param Component The ML Deformer component that will own this instance.
+	 * @return A pointer to the model instance.
+	 */
 	virtual UMLDeformerModelInstance* CreateModelInstance(UMLDeformerComponent* Component);
+
+	/**
+	 * Get the display name of this model.
+	 * This will also define with what name this model will appear inside the UI.
+	 * @return The name of the model.
+	 */
 	virtual FString GetDisplayName() const PURE_VIRTUAL(UMLDeformerModel::GetDisplayName, return FString(););
+
+	/**
+	 * Initialize this model, after the ML deformer component has been initialized.
+	 * This is a post-step after component initialization. It can be used to register things such as 
+	 * external morph target sets that are generated by this model.
+	 * @param ModelInstance The ML Deformer model instance that we are initializing with.
+	 */
 	virtual void PostMLDeformerComponentInit(UMLDeformerModelInstance* ModelInstance) {}
+
+	/**
+	 * Defines whether this model supports bone transforms as input or not.
+	 * On default this is set to return true as most models have bone rotations as inputs to the neural network.
+	 * @result Returns true when this model supports bones, or false otherwise.
+	 */
 	virtual bool DoesSupportBones() const { return true; }
+
+	/**
+	 * Defines whether this model supports curves as inputs or not. A curve is just a single float value.
+	 * On default this returns true.
+	 * @result Returns true when this model supports curves, or false otherwise.
+	 */
 	virtual bool DoesSupportCurves() const { return true; }
-	virtual bool IsNeuralNetworkOnGPU() const { return true; }	// GPU neural network.
+
+	/**
+	 * Check whether the neural network of this model should run on the GPU or not.
+	 * This is false on default, which makes it a CPU based neural network.
+	 * Some code internally that creates and initializes the neural network will use the return value of this method to mark it to be on CPU or GPU.
+	 * @return Returns true if the neural network of this model should run on the GPU. False is returned when it should run on the CPU.
+	 */
+	virtual bool IsNeuralNetworkOnGPU() const { return false; }	// CPU neural network.
+
 #if WITH_EDITORONLY_DATA
+	/**
+	 * Check whether this model currently has a training mesh setup or not.
+	 * For example if there is say a GeometryCache as target mesh, it could check whether that property is a nullptr or not.
+	 * @return Returns true when the training target mesh has been selected, otherwise false is returned.
+	 */
 	virtual bool HasTrainingGroundTruth() const { return false; }
+
+	/**
+	 * Sample the positions from the target (ground truth) mesh, at a specific time (in seconds).
+	 * @param SampleTime The time to sample the positions at, in seconds.
+	 * @param OutPositions The array that will receive the resulting vertex positions. This array will automatically be resized internally.
+	 */
 	virtual void SampleGroundTruthPositions(float SampleTime, TArray<FVector3f>& OutPositions) {}
 #endif
+
 #if WITH_EDITOR
-	virtual void UpdateNumBaseMeshVertices();
+	/**
+	 * Update the cached number of vertices of both base and target meshes.
+	 */ 
 	virtual void UpdateCachedNumVertices();
-	virtual void UpdateNumTargetMeshVertices() PURE_VIRTUAL(UMLDeformerModel::GetNumTargetMeshVertices; mTargetMeshVerts = 0;);
+
+	/**
+	 * Update all editor only assets with the editor only flag.
+	 * This prevents certain assets that are only used during training to be included in a packaged build.
+	 * Examples of these are a training geometry cache object, and the ground truth geometry cache of a test animation in the visualization settings.
+	 * If you do not set the flag for this, those will unfortunately be included, and they can be massive in size and aren't needed at runtime.
+	 * You can set these flags like this:
+	 * @code
+	 * TrainingObject->GetPackage()->SetPackageFlags(PKG_EditorOnly).
+	 * @endcode
+	 */ 
 	virtual void SetAssetEditorOnlyFlags();
+
+	/**
+	 * Extract the number of imported (DCC) vertices from a skeletal mesh.
+	 * @param SkeletalMesh The skeletal mesh to get the number of imported vertices for.
+	 * @return The number of imported vertices, which is the vertex count as seen in the DCC.
+	 */ 
 	static int32 ExtractNumImportedSkinnedVertices(const USkeletalMesh* SkeletalMesh);
 #endif
 
@@ -84,54 +182,211 @@ public:
 	virtual USkeleton* GetSkeleton(bool& bInvalidSkeletonIsError, const IPropertyHandle* PropertyHandle) override;
 	// ~END IBoneReferenceSkeletonProvider overrides.
 
+	/**
+	 * Get the ML deformer asset that this model is part of.
+	 * @return A pointer to the ML deformer asset.
+	 */
 	UMLDeformerAsset* GetDeformerAsset() const;
+
+	/**
+	 * Get the input information, which is information about the inputs to the deformer.
+	 * Inputs are things like bone transforms and curve values.
+	 * @return A pointer to the input info object.
+	 */
 	UMLDeformerInputInfo* GetInputInfo() const { return InputInfo.Get(); }
 
+	/**
+	 * Get the number of vertices in the base mesh (linear skinned skeletal mesh).
+	 * This is the number of vertices in the DCC, so not the render vertex count.
+	 * @return The number of imported vertices inside linear skinned skeletal mesh.
+	 */
 	int32 GetNumBaseMeshVerts() const { return NumBaseMeshVerts; }
+
+	/**
+	 * Get the number of vertices of the target mesh.
+	 * This is the number of vertices in the DCC, so not the render vertex count.
+	 * @return The number of imported vertices inside the target mesh.
+	 */
 	int32 GetNumTargetMeshVerts() const { return NumTargetMeshVerts; }
 
+	/**
+	 * The mapping that maps from render vertices into dcc vertices.
+	 * The length of this array is the same as the number of render vertices in the skeletal mesh.
+	 * For a cube with 32 render vertices, the item values would be between 0..7 as in the dcc the cube has 8 vertices.
+	 * @return A reference to the array that contains the DCC vertex number for each render vertex.
+	 */
 	const TArray<int32>& GetVertexMap() const { return VertexMap; }
+
+	/**
+	 * Get the GPU buffer of the VertexMap.
+	 * @return A reference to the GPU buffer resource object that holds the vertex map.
+	 * @see GetVertexMap
+	 */
 	const UE::MLDeformer::FVertexMapBuffer& GetVertexMapBuffer() const { return VertexMapBuffer; }
 
+	/**
+	 * Get the neural network that we have trained. This can return a nullptr when no network has been trained yet.
+	 * This network is used during inference.
+	 * @return A pointer to the neural network, or nullptr when the network has not yet been trained.
+	 */
 	UNeuralNetwork* GetNeuralNetwork() const { return NeuralNetwork.Get(); }
-	void SetNeuralNetwork(UNeuralNetwork* InNeuralNetwork);
 
+	/**
+	 * Set the neural network object that we use during inference.
+	 * @param InNeuralNetwork The new neural network to use inside this deformer model.
+	 */
+	void SetNeuralNetwork(UNeuralNetwork* InNeuralNetwork);
 
 #if WITH_EDITORONLY_DATA
 	// UObject overrides.
 	virtual void PostEditChangeProperty(FPropertyChangedEvent& PropertyChangedEvent) override;
 	// ~END UObject overrides.
 
+	/**
+	 * Initialize the vertex map.
+	 * This will grab the skeletal mesh imported model and copy over the imported mapping from there.
+	 * This will make sure the VertexMap member is initialized.
+	 * @see GetVertexMap
+	 */
 	void InitVertexMap();
 
+	/**
+	 * Check whether we should include bone transforms as input to the model during training or not.
+	 * @return Returns true when bone transfomations should be a part of the network inputs, during the training process.
+	 */
 	bool ShouldIncludeBonesInTraining() const { return bIncludeBones; }
+
+	/**
+	 * Check whether we should include curve values as input to the model during training or not.
+	 * Curve values are single floats.
+	 * @return Returns true when curve values should be a part of the network inputs, during the training process.
+	 */
 	bool ShouldIncludeCurvesInTraining() const { return bIncludeCurves; }
 
+	/**
+	 * The delegate that gets fired when a property value changes.
+	 * @return A reference to the delegate.
+	 */
 	FMLDeformerModelOnPostEditProperty& OnPostEditChangeProperty() { return PostEditPropertyDelegate; }
 
+	/**
+	 * Get the visualization settings for this model. These settings are only used in the editor.
+	 * Visualization settings contain settings for the left side of the ML Deformer asset editor, containing 
+	 * settings like what the mesh spacing is, whether to draw labels, what the test anim sequence is, etc.
+	 * @return A pointer to the visualization settings. You can cast this to the type specific for your model, in case you inherited from 
+	 *         the UMLDeformerVizSettings base class. This never return a nullptr.
+	 */
 	UMLDeformerVizSettings* GetVizSettings() const { return VizSettings; }
+
+	/**
+	 * Get the skeletal mesh that is used during training.
+	 * You typically want to apply the ML Deformer on this specific skeletal mesh in your game as well.
+	 * @return A pointer to the skeletal mesh.
+	 */
 	const USkeletalMesh* GetSkeletalMesh() const { return SkeletalMesh;  }
 	USkeletalMesh* GetSkeletalMesh() { return SkeletalMesh; }
+
+	/**
+	 * Get the animation sequence that is used during training.
+	 * Each frame of this anim sequence will contain a training pose.
+	 * @return A pointer to the animation sequence used for training.
+	 */
 	const UAnimSequence* GetAnimSequence() const { return AnimSequence;  }
 	UAnimSequence* GetAnimSequence() { return AnimSequence; }
+
+	/**
+	 * Get the maximum number of training frames to use during training.
+	 * Your training anim sequence might contain say 10000 frames, but for quickly iterating you might
+	 * want to train on only 2000 frames instead. You can do this by setting the maximum training frames to 2000.
+	 * @return The max number of frames to use during training.
+	 */
 	int32 GetTrainingFrameLimit() const { return MaxTrainingFrames; }
+
+	/**
+	 * Get the target mesh alignment tranformation.
+	 * This is a transformation that is applied to the vertex positions of the target mesh, before we calculate the deltas
+	 * between the linear skinned mesh and the target mesh.
+	 * This is useful when you imported target mesh that isn't scaled the same, or perhaps it is rotated 90 degrees over the x axis.
+	 * The alignment transform is then used to correct this and align both base and target mesh.
+	 * @return The alignment transformation. When set to Identity it will not do anything, which is its default.
+	 */
 	const FTransform& GetAlignmentTransform() const { return AlignmentTransform; }
+
+	/**
+	 * Get the list of bones that we configured to be included during training.
+	 * A bone reference is basically just a name of a bone.
+	 * This can be different from the bone list inside the InputInfo property though.
+	 * Please keep in mind that the BoneIncludeList is what is setup in the UI, but the InputInfo contains 
+	 * the actual list of what we trained on, so what we use during inference.
+	 * A user could change this bone list after they trained a model, in which case the input info and this bone include list
+	 * will contain different items.
+	 * @return The list of bones that should be included when performing the next training session.
+	 */
 	TArray<FBoneReference>& GetBoneIncludeList() { return BoneIncludeList; }
 	const TArray<FBoneReference>& GetBoneIncludeList() const { return BoneIncludeList; }
+
+	/**
+	 * Get the list of curves that we configured to be included during training.
+	 * A curve reference is basically just a name of a curve.
+	 * This can be different from the curve list inside the InputInfo property though.
+	 * Please keep in mind that the CurveIncludeList is what is setup in the UI, but the InputInfo contains 
+	 * the actual list of what we trained on, so what we use during inference.
+	 * A user could change this curve list after they trained a model, in which case the input info and this curve include list
+	 * will contain different items.
+	 * @return The list of curves that should be included when performing the next training session.
+	 */
 	TArray<FMLDeformerCurveReference>& GetCurveIncludeList() { return CurveIncludeList; }
 	const TArray<FMLDeformerCurveReference>& GetCurveIncludeList() const { return CurveIncludeList; }
+
+	/**
+	 * Get the delta cutoff length. Deltas that have a length larger than this length will be set to zero.
+	 * This can be useful when there are some vertices that due to incorrect data have a very long length.
+	 * Skipping those deltas will prevent issues.
+	 * @return The length after which deltas will be ignored. So anything delta length larger than this value will be ignored.
+	 */
 	float GetDeltaCutoffLength() const { return DeltaCutoffLength; }
 #endif	// #if WITH_EDITORONLY_DATA
 
 protected:
+	/** The deformer asset that this model is part of. */
 	TObjectPtr<UMLDeformerAsset> DeformerAsset = nullptr;
+
+	/** The delegate that gets fired when a property is being modified. */
 	FMLDeformerModelOnPostEditProperty PostEditPropertyDelegate;
 
+	/**
+	 * Set the training input information.
+	 * @param Input A pointer to the input information object.
+	 * @see CreateInputInfo
+	 */
 	void SetInputInfo(UMLDeformerInputInfo* Input) { InputInfo = Input; }
+
+	/**
+	 * Convert an array of floats to an array of Vector3's.
+	 * A requirement is that the number of items in the float array is a multiple of 3 (xyz).
+	 * The order of the float items must be like this: (x, y, z, x, y, z, x, y, z, ...).
+	 * The OutVectorArray will be automatically resized internally.
+	 * @param FloatArray The array of floats to build an array of Vector3's for.
+	 * @param OutVectorArray The array that will contain the vectors instead of floats.
+	 */
 	void FloatArrayToVector3Array(const TArray<float>& FloatArray, TArray<FVector3f>& OutVectorArray);
 
+	/**
+	 * Update the cached number of vertices in the base mesh. 
+	 * The number of vertices should be the number of DCC vertices, so not the number of render vertices.
+	 * This just updates the NumBaseMeshVerts property.
+	 */ 
+	virtual void UpdateNumBaseMeshVertices();
+
+	/**
+	 * Update the cached number of target mesh vertices. Every model needs to implement this.
+	 * The number of vertices should be the number of DCC vertices, so not the number of render vertices.
+	 * This just updates the NumTargetMeshVerts property.
+	 */ 
+	virtual void UpdateNumTargetMeshVertices() PURE_VIRTUAL(UMLDeformerModel::GetNumTargetMeshVertices; mTargetMeshVerts = 0;);
+
 public:
-	/** Cached number of skeltal mesh vertices. */
+	/** Cached number of skeletal mesh vertices. */
 	UPROPERTY()
 	int32 NumBaseMeshVerts = 0;
 
@@ -167,11 +422,11 @@ public:
 	UPROPERTY()
 	TObjectPtr<UMLDeformerVizSettings> VizSettings = nullptr;
 
-	/** Include bones as inputs to training process? */
+	/** Specifies whether bone transformations should be included as inputs during the training process. */
 	UPROPERTY(EditAnywhere, Category = "Inputs and Output")
 	bool bIncludeBones = true;
 
-	/** Include curves as inputs to the training process? */
+	/** Specifies whether curve values (a float per curve) should be included as inputs during the training process. */
 	UPROPERTY(EditAnywhere, Category = "Inputs and Output")
 	bool bIncludeCurves = false;
 
@@ -202,8 +457,10 @@ public:
 	UPROPERTY(EditAnywhere, BlueprintReadWrite, Category = "Inputs and Output", meta = (ClampMin = "1"))
 	int32 MaxTrainingFrames = 1000000;
 
-	/** Sometimes there can be some vertices that cause some issues that cause deltas to be very long. We can ignore these deltas by setting a cutoff value. 
-	  * Deltas that are longer than the cutoff value (in units), will be ignored and set to zero length. */
+	/**
+	 * Sometimes there can be some vertices that cause some issues that cause deltas to be very long. We can ignore these deltas by setting a cutoff value. 
+	 * Deltas that are longer than the cutoff value (in units), will be ignored and set to zero length. 
+	 */
 	UPROPERTY(EditAnywhere, BlueprintReadWrite, Category = "Inputs and Output", meta = (ClampMin = "0.01", ForceUnits="cm"))
 	float DeltaCutoffLength = 30.0f;
 #endif
diff --git a/Engine/Plugins/Experimental/MLDeformer/MLDeformerFramework/Source/MLDeformerFramework/Public/MLDeformerMorphModel.h b/Engine/Plugins/Experimental/MLDeformer/MLDeformerFramework/Source/MLDeformerFramework/Public/MLDeformerMorphModel.h
new file mode 100644
--- /dev/null
+++ b/Engine/Plugins/Experimental/MLDeformer/MLDeformerFramework/Source/MLDeformerFramework/Public/MLDeformerMorphModel.h
@@ -0,0 +1,74 @@
+// Copyright Epic Games, Inc. All Rights Reserved.
+#pragma once
+
+#include "CoreMinimal.h"
+#include "MLDeformerGeomCacheModel.h"
+#include "UObject/Object.h"
+#include "MLDeformerMorphModel.generated.h"
+
+class USkeletalMesh;
+class UMLDeformerAsset;
+class UMLDeformerModelInstance;
+struct FExternalMorphSet;
+
+UCLASS()
+class MLDEFORMERFRAMEWORK_API UMLDeformerMorphModel
+	: public UMLDeformerGeomCacheModel
+{
+	GENERATED_BODY()
+
+public:
+	UMLDeformerMorphModel(const FObjectInitializer& ObjectInitializer);
+
+	// UMLDeformerModel overrides.
+	virtual FString GetDisplayName() const override { return "Morph Base Model"; }
+	virtual void Serialize(FArchive& Archive) override;
+	virtual void PostMLDeformerComponentInit(UMLDeformerModelInstance* ModelInstance) override;
+	virtual bool IsNeuralNetworkOnGPU() const override { return false; }	// CPU neural network.
+	virtual UMLDeformerModelInstance* CreateModelInstance(UMLDeformerComponent* Component) override;
+	// ~END UMLDeformerModel overrides.
+
+	// UObject overrides.
+	void BeginDestroy() override;
+	// ~END UObject overrides.
+
+#if WITH_EDITORONLY_DATA
+	float GetMorphTargetDeltaThreshold() const { return MorphTargetDeltaThreshold; }
+#endif
+
+	UFUNCTION(BlueprintCallable, Category = "NeuralMorphModel")
+	void SetMorphTargetDeltas(const TArray<float>& Deltas);
+
+	const TArray<FVector3f>& GetMorphTargetDeltas() const { return MorphTargetDeltas; }
+	int32 GetMorphTargetDeltaStartIndex(int32 MorphTargetIndex) const;
+
+public:
+	/** The compressed morph target data, ready for the GPU. */
+	TSharedPtr<FExternalMorphSet> MorphTargetSet;
+
+	/**
+	 * The entire set of morph target deltas, 3 per vertex, for each morph target, as one flattened buffer.
+	 * So the size of this buffer is: (NumVertsPerMorphTarget * 3 * NumMorphTargets).
+	 */
+	TArray<FVector3f> MorphTargetDeltas;
+
+	/** 
+	 * The external morph set data type ID, specific to this model.
+	 * If you inherit your model from this base class, you should set this to some unique value, that represents your model.
+	 */
+	int32 ExternalMorphSetID = -1;
+
+#if WITH_EDITORONLY_DATA
+	/**
+	 * Morph target delta values that are smaller than or equal to this threshold will be zeroed out.
+	 * This essentially removes small deltas from morph targets, which will lower the memory usage at runtime, however when set too high it can also introduce visual artifacts.
+	 * A value of 0 will result in the highest quality morph targets, at the cost of higher runtime memory usage.
+	 */
+	UPROPERTY(EditAnywhere, BlueprintReadWrite, Category = "Morph Targets", meta = (ClampMin = "0.0", ClampMax = "1.0", ForceUnits="cm"))
+	float MorphTargetDeltaThreshold = 0.0025f;
+
+	/** The morph target error tolerance. Higher values result in larger compression, but could result in visual artifacts. */
+	UPROPERTY(EditAnywhere, BlueprintReadWrite, Category = "Morph Targets", meta = (ClampMin = "0.01", ClampMax = "500"))
+	float MorphTargetErrorTolerance = 50.0f;
+#endif // WITH_EDITORONLY_DATA
+};
diff --git a/Engine/Plugins/Experimental/MLDeformer/MLDeformerFramework/Source/MLDeformerFramework/Public/MLDeformerMorphModelInstance.h b/Engine/Plugins/Experimental/MLDeformer/MLDeformerFramework/Source/MLDeformerFramework/Public/MLDeformerMorphModelInstance.h
new file mode 100644
--- /dev/null
+++ b/Engine/Plugins/Experimental/MLDeformer/MLDeformerFramework/Source/MLDeformerFramework/Public/MLDeformerMorphModelInstance.h
@@ -0,0 +1,17 @@
+// Copyright Epic Games, Inc. All Rights Reserved.
+
+#pragma once
+
+#include "MLDeformerModelInstance.h"
+#include "MLDeformerMorphModelInstance.generated.h"
+
+
+UCLASS()
+class MLDEFORMERFRAMEWORK_API UMLDeformerMorphModelInstance
+	: public UMLDeformerModelInstance
+{
+	GENERATED_BODY()
+
+public:
+	virtual void RunNeuralNetwork(float ModelWeight) override;
+};
diff --git a/Engine/Plugins/Experimental/MLDeformer/MLDeformerFramework/Source/MLDeformerFramework/Public/MLDeformerMorphModelVizSettings.h b/Engine/Plugins/Experimental/MLDeformer/MLDeformerFramework/Source/MLDeformerFramework/Public/MLDeformerMorphModelVizSettings.h
new file mode 100644
--- /dev/null
+++ b/Engine/Plugins/Experimental/MLDeformer/MLDeformerFramework/Source/MLDeformerFramework/Public/MLDeformerMorphModelVizSettings.h
@@ -0,0 +1,45 @@
+// Copyright Epic Games, Inc. All Rights Reserved.
+
+#pragma once
+
+#include "CoreMinimal.h"
+#include "UObject/Object.h"
+#include "MLDeformerGeomCacheVizSettings.h"
+#include "MLDeformerMorphModelVizSettings.generated.h"
+
+/**
+ * The vizualization settings specific to this model.
+ */
+UCLASS()
+class MLDEFORMERFRAMEWORK_API UMLDeformerMorphModelVizSettings
+	: public UMLDeformerGeomCacheVizSettings
+{
+	GENERATED_BODY()
+
+public:
+#if WITH_EDITORONLY_DATA
+	/**
+	 * The morph target to visualize. The first one always being the means, so not a sparse target.
+	 * This only can be used after you trained, in the same editor session directly after training.
+	 */
+	UPROPERTY(EditAnywhere, Category = "Live Settings", meta = (ClampMin = "0"))
+	int32 MorphTargetNumber = 0;
+
+	/**
+	 * The morph target delta threshold. This is a preview of what deltas would be included in the selected morph target
+	 * when using a delta threshold during training that is equal to this value.
+	 * This only can be used after you trained, in the same editor session directly after training.
+	 */
+	UPROPERTY(EditAnywhere, Category = "Live Settings", meta = (ClampMin = "0.001", ClampMax = "1.0", ForceUnits="cm"))
+	float MorphTargetDeltaThreshold = 0.0025f;
+
+	/**
+	 * Specify whether we want to debug draw the morph targets.
+	 * Enabling this can help you see what the sparsity of the generated morph targets is.
+	 * It also takes the morph target delta threshold into account, so you can preview how that effects the sparsity of the morphs.
+	 * This only can be used after you trained, in the same editor session directly after training.
+	 */
+	UPROPERTY(EditAnywhere, Category = "Live Settings")
+	bool bDrawMorphTargets = false;
+#endif // WITH_EDITORONLY_DATA
+};
diff --git a/Engine/Plugins/Experimental/MLDeformer/MLDeformerFramework/Source/MLDeformerFramework/Public/MLDeformerVizSettings.h b/Engine/Plugins/Experimental/MLDeformer/MLDeformerFramework/Source/MLDeformerFramework/Public/MLDeformerVizSettings.h
--- a/Engine/Plugins/Experimental/MLDeformer/MLDeformerFramework/Source/MLDeformerFramework/Public/MLDeformerVizSettings.h
+++ b/Engine/Plugins/Experimental/MLDeformer/MLDeformerFramework/Source/MLDeformerFramework/Public/MLDeformerVizSettings.h
@@ -98,7 +98,7 @@ public:
 	UPROPERTY(EditAnywhere, Category = "Training Meshes", meta = (ClampMin = "0"))
 	uint32 TrainingFrameNumber = 0;
 
-	/** Draw the text labels above each actor? */
+	/** Specifies whether we should draw the labels, such as "Linear Skinned" and "ML Deformed" or not. */
 	UPROPERTY(EditAnywhere, Category = "Shared Settings")
 	bool bDrawLabels = true;
 
@@ -118,31 +118,41 @@ public:
 	UPROPERTY(EditAnywhere, Category = "Live Settings", meta = (ClampMin = "0"))
 	uint32 TestingFrameNumber = 0;
 
-	/** Show the heat map? This will visualize the active areas of the deformer. */
+	/** Specify whether the heatmap is enabled or not. */
 	UPROPERTY(EditAnywhere, Category = "Live Settings")
 	bool bShowHeatMap = false;
 
-	/** What should the heatmap visualize? */
+	/** Specifies what the heatmap colors should represent. */
 	UPROPERTY(EditAnywhere, Category = "Live Settings", meta = (EditCondition = "bShowHeatMap"))
 	EMLDeformerHeatMapMode HeatMapMode = EMLDeformerHeatMapMode::Activations;
 
-	/** How many centimeters does the most intense color of the heatmap represent? */
+	/**
+	 * Specifies how many centimeters the most intense color of the heatmap represents.
+	 * For example when set to 3, it means that everything that is red, is 3 cm or more.
+	 * In Activations heat map mode this means that everything is red for all vertices with a delta longer than 3 cm, so with a corrective delta longer than 3 cm.
+	 * In Ground Truth heat map mode this means that everything is red for all vertices where the difference between the ground truth and the ML Deformed
+	 * model is larger than 3 cm.
+	 */
 	UPROPERTY(EditAnywhere, Category = "Live Settings", meta = (EditCondition = "bShowHeatMap", ClampMin = "0.001", ForceUnits="cm"))
 	float HeatMapMax = 1.0f;
 
-	/** Lerp from ML deformed model to ground truth model when in heat map mode. */
+	/**
+	 * The Lerp factor from ML deformed model to ground truth model when in heat map mode. 
+	 * A value of 0 means that we look exactly the same as the ML Deformed output, while a value of 1 means
+	 * the output is exactly the same as the ground truth. This is the blend factor between those two.
+	 */
 	UPROPERTY(EditAnywhere, Category = "Live Settings", meta = (EditCondition = "HeatMapMode==EMLDeformerHeatMapMode::GroundTruth && bShowHeatMap", UIMin = "0.0", UIMax = "1.0", ClampMin = "0.0", ClampMax = "1.0"))
 	float GroundTruthLerp = 0.0f;
 
-	/** Draw the linear skinned actor? */
+	/** Specifies whether we draw the linear skinned model or not. */
 	UPROPERTY(EditAnywhere, Category = "Live Settings")
 	bool bDrawLinearSkinnedActor = true;
 
-	/** Draw the ML Deformed actor? */
+	/** Specifies whether we draw the ML Deformed model or not. */
 	UPROPERTY(EditAnywhere, Category = "Live Settings", DisplayName = "Draw ML Deformed Actor")
 	bool bDrawMLDeformedActor = true;
 
-	/** Draw the ground truth actor? */
+	/** Specifies whether we draw the ground truth model or not. */
 	UPROPERTY(EditAnywhere, Category = "Live Settings")
 	bool bDrawGroundTruthActor = true;
 
@@ -150,11 +160,11 @@ public:
 	UPROPERTY(EditAnywhere, Category = "Live Settings", meta = (ClampMin = "0.0", ClampMax = "1.0"))
 	float Weight = 1.0f;
 
-	/** Draw the vertex deltas? */
+	/** Specifiy whether we want to draw the vertex deltas or not. */
 	UPROPERTY(EditAnywhere, Category = "Training Meshes")
 	bool bDrawDeltas = true;
 
-	/** Enable this to draw the deltas in xray mode? */
+	/** Specify whether the vertex deltas are rendered, even if they are behind the current mesh. */
 	UPROPERTY(EditAnywhere, Category = "Training Meshes", meta = (EditCondition = "bDrawDeltas"))
 	bool bXRayDeltas = true;
 #endif
diff --git a/Engine/Plugins/Experimental/MLDeformer/MLDeformerFramework/Source/MLDeformerFrameworkEditor/Private/MLDeformerEditorModel.cpp b/Engine/Plugins/Experimental/MLDeformer/MLDeformerFramework/Source/MLDeformerFrameworkEditor/Private/MLDeformerEditorModel.cpp
--- a/Engine/Plugins/Experimental/MLDeformer/MLDeformerFramework/Source/MLDeformerFrameworkEditor/Private/MLDeformerEditorModel.cpp
+++ b/Engine/Plugins/Experimental/MLDeformer/MLDeformerFramework/Source/MLDeformerFrameworkEditor/Private/MLDeformerEditorModel.cpp
@@ -44,7 +44,7 @@ namespace UE::MLDeformer
 			Model->OnPostEditChangeProperty().Remove(PostEditPropertyDelegateHandle);
 		}
 
-		DeleteEditorActors();
+		ClearWorld();
 
 		FMLDeformerEditorModule& EditorModule = FModuleManager::GetModuleChecked<FMLDeformerEditorModule>("MLDeformerFrameworkEditor");
 		EditorModule.GetModelRegistry().RemoveEditorModelInstance(this);
@@ -72,6 +72,11 @@ namespace UE::MLDeformer
 		PostEditPropertyDelegateHandle = Model->OnPostEditChangeProperty().AddRaw(this, &FMLDeformerEditorModel::OnPostEditChangeProperty);
 	}
 
+	void FMLDeformerEditorModel::AddReferencedObjects(FReferenceCollector& Collector)
+	{
+		Collector.AddReferencedObject(EditorInputInfo);
+	}
+
 	void FMLDeformerEditorModel::UpdateEditorInputInfo()
 	{
 		InitInputInfo(EditorInputInfo);
@@ -224,8 +229,6 @@ namespace UE::MLDeformer
 				EditorActor->SetMeshOffsetFactor(EditorActor->IsTestActor() ? 2.0f : 1.0f);
 			}
 		}
-
-		OnPostCreateActors();
 	}
 
 	void FMLDeformerEditorModel::ClearWorld()
@@ -581,7 +584,7 @@ namespace UE::MLDeformer
 		}
 	}
 
-	void FMLDeformerEditorModel::HandleDefaultPropertyChanges(FPropertyChangedEvent& PropertyChangedEvent)
+	void FMLDeformerEditorModel::OnPropertyChanged(FPropertyChangedEvent& PropertyChangedEvent)
 	{
 		const FProperty* Property = PropertyChangedEvent.Property;
 		if (Property == nullptr)
@@ -669,40 +672,49 @@ namespace UE::MLDeformer
 			UpdateDeformerGraph();
 			GetEditor()->GetVizSettingsDetailsView()->ForceRefresh();
 		}
+		else
+		if (Property->GetFName() == GET_MEMBER_NAME_CHECKED(UMLDeformerModel, DeltaCutoffLength) ||
+			Property->GetFName() == GET_MEMBER_NAME_CHECKED(UMLDeformerModel, AlignmentTransform))
+		{
+			if (PropertyChangedEvent.ChangeType == EPropertyChangeType::ValueSet)
+			{
+				SetResamplingInputOutputsNeeded(true);
+				SampleDeltas();
+			}
+		}
 	}
 
 	FMLDeformerEditorActor* FMLDeformerEditorModel::GetVisualizationModeBaseActor() const
 	{
-		UMLDeformerVizSettings* VizSettings = Model->GetVizSettings();
-		const bool bTestData = VizSettings->GetVisualizationMode() == EMLDeformerVizMode::TestData;
-		FMLDeformerEditorActor* BaseActor = nullptr;
-		if (bTestData)
-		{
-			BaseActor = FindEditorActor(ActorID_Test_Base);
-		}
-		else
-		{
-			BaseActor = FindEditorActor(ActorID_Train_Base);
-		}
-		return BaseActor;
+		const UMLDeformerVizSettings* VizSettings = Model->GetVizSettings();
+		const bool bTestMode = VizSettings->GetVisualizationMode() == EMLDeformerVizMode::TestData;
+		return bTestMode ? FindEditorActor(ActorID_Test_Base) : FindEditorActor(ActorID_Train_Base);
 	}
 
 	void FMLDeformerEditorModel::OnPlayPressed()
 	{
 		// The play button shouldn't do anything in training mode.
-		UMLDeformerVizSettings* VizSettings = Model->GetVizSettings();
-		if (VizSettings->GetVisualizationMode() == EMLDeformerVizMode::TrainingData)
+		const UMLDeformerVizSettings* VizSettings = Model->GetVizSettings();
+		const FMLDeformerEditorActor* BaseActor = GetVisualizationModeBaseActor();
+		const bool bMustPause = (BaseActor && BaseActor->IsPlaying());
+		const bool bTestMode = (VizSettings->GetVisualizationMode() == EMLDeformerVizMode::TestData);
+
+		// Update the current frame numbers.
+		if (bMustPause)
 		{
-			return;
+			if (VizSettings->GetVisualizationMode() == EMLDeformerVizMode::TrainingData)
+			{
+				SetTrainingFrame(GetTrainingFrameAtTime(CalcTimelinePosition()));
+			}
+			else if (VizSettings->GetVisualizationMode() == EMLDeformerVizMode::TestData)
+			{
+				SetTestFrame(GetTestFrameAtTime(CalcTimelinePosition()));
+			}
 		}
 
-		FMLDeformerEditorActor* BaseActor = GetVisualizationModeBaseActor();
-		const bool bMustPause = (BaseActor && BaseActor->GetSkeletalMeshComponent()) ? !BaseActor->GetSkeletalMeshComponent()->bPauseAnims : false;
-
-		const bool bTestData = (VizSettings->GetVisualizationMode() == EMLDeformerVizMode::TestData);
 		for (FMLDeformerEditorActor* EditorActor : EditorActors)
 		{
-			if (EditorActor && ((EditorActor->IsTestActor() && bTestData) || (EditorActor->IsTrainingActor() && !bTestData)))
+			if (EditorActor && ((EditorActor->IsTestActor() && bTestMode) || (EditorActor->IsTrainingActor() && !bTestMode)))
 			{
 				EditorActor->SetPlaySpeed(VizSettings->AnimPlaySpeed);
 				EditorActor->Pause(bMustPause);
@@ -712,8 +724,7 @@ namespace UE::MLDeformer
 
 	bool FMLDeformerEditorModel::IsPlayingAnim() const
 	{
-		FMLDeformerEditorActor* BaseActor = GetVisualizationModeBaseActor();
-
+		const FMLDeformerEditorActor* BaseActor = GetVisualizationModeBaseActor();
 		if (BaseActor)
 		{
 			UDebugSkelMeshComponent* SkeletalMeshComponent = BaseActor->GetSkeletalMeshComponent();
@@ -728,7 +739,6 @@ namespace UE::MLDeformer
 	bool FMLDeformerEditorModel::IsPlayingForward() const
 	{
 		FMLDeformerEditorActor* BaseActor = GetVisualizationModeBaseActor();
-
 		if (BaseActor)
 		{
 			return BaseActor->GetPlaySpeed() > 0.0f;
@@ -1237,10 +1247,19 @@ namespace UE::MLDeformer
 		const UMLDeformerVizSettings* VizSettings = Model->GetVizSettings();
 		if (VizSettings->GetVisualizationMode() == EMLDeformerVizMode::TrainingData)
 		{
-			// Draw the deltas for the current frame.
 			const TArray<float>& VertexDeltas = Sampler->GetVertexDeltas();
 			const TArray<FVector3f>& LinearSkinnedPositions = Sampler->GetSkinnedVertexPositions();
-			if (VizSettings->GetDrawVertexDeltas() && (VertexDeltas.Num() / 3) == LinearSkinnedPositions.Num())
+			bool bDrawDeltas = (VizSettings->GetDrawVertexDeltas() && (VertexDeltas.Num() / 3) == LinearSkinnedPositions.Num());
+
+			// Disable drawing deltas when we're playing the training anim sequence, as that can get too slow.
+			const FMLDeformerEditorActor* BaseActor = FindEditorActor(ActorID_Train_Base);
+			if (BaseActor)
+			{
+				bDrawDeltas &= !BaseActor->IsPlaying();
+			}
+
+			// Draw the deltas for the current frame.
+			if (bDrawDeltas)
 			{
 				const FLinearColor DeltasColor = FMLDeformerEditorStyle::Get().GetColor("MLDeformer.Deltas.Color");
 				const FLinearColor DebugVectorsColor = FMLDeformerEditorStyle::Get().GetColor("MLDeformer.DebugVectors.Color");
@@ -1432,7 +1451,7 @@ namespace UE::MLDeformer
 
 			if (SkelMeshComponent)
 			{
-				UMLDeformerVizSettings* VizSettings = Model->GetVizSettings();
+				const UMLDeformerVizSettings* VizSettings = Model->GetVizSettings();
 				UMeshDeformer* MeshDeformer = IsTrained() ? VizSettings->GetDeformerGraph() : nullptr;	
 				const bool bUseHeatMapDeformer = VizSettings->GetShowHeatMap();
 				SkelMeshComponent->SetMeshDeformer(bUseHeatMapDeformer ? HeatMapDeformerGraph.Get() : MeshDeformer);
@@ -1547,13 +1566,13 @@ namespace UE::MLDeformer
 		return (Model->GetNeuralNetwork() != nullptr);
 	}
 
-	void FMLDeformerEditorModel::TriggerInputAssetChanged(bool RefreshVizSettings)
+	void FMLDeformerEditorModel::TriggerInputAssetChanged(bool bRefreshVizSettings)
 	{
 		OnInputAssetsChanged();
 		OnPostInputAssetChanged();
 		Model->SetAssetEditorOnlyFlags();
 		GetEditor()->GetModelDetailsView()->ForceRefresh();
-		if (RefreshVizSettings)
+		if (bRefreshVizSettings)
 		{
 			GetEditor()->GetVizSettingsDetailsView()->ForceRefresh();
 		}
@@ -1737,11 +1756,6 @@ namespace UE::MLDeformer
 		return FString(TEXT("/MLDeformerFramework/Deformers/DG_MLDeformerModel.DG_MLDeformerModel"));
 	}
 
-	FString FMLDeformerEditorModel::GetTrainedNetworkOnnxFile() const
-	{
-		return FString(FPaths::ProjectIntermediateDir() + TEXT("MLDeformerNetwork.onnx"));
-	}
-
 	const UAnimSequence* FMLDeformerEditorModel::GetAnimSequence() const
 	{
 		if (GetModel())
@@ -1823,7 +1837,7 @@ namespace UE::MLDeformer
 
 	FFrameNumber FMLDeformerEditorModel::GetTickResScrubPosition() const
 	{
-		double Resolution = GetTickResolution();
+		const double Resolution = GetTickResolution();
 		const UMLDeformerVizSettings* VizSettings = GetModel()->GetVizSettings();
 		double TrainingFrameNumber = 0.0;
 		if (VizSettings->GetVisualizationMode() == EMLDeformerVizMode::TrainingData)
@@ -1845,28 +1859,43 @@ namespace UE::MLDeformer
 
 	void FMLDeformerEditorModel::SetScrubPosition(FFrameTime NewScrubPostion)
 	{
-		double FrameNumber = NewScrubPostion.AsDecimal() * (double)GetFrameRate() / (double)GetTickResolution();
+		const double FrameNumber = NewScrubPostion.AsDecimal() * (double)GetFrameRate() / (double)GetTickResolution();
 		ScrubPosition.FrameNumber = (int32)FrameNumber;
-	
-		SetTrainingFrame(ScrubPosition.FrameNumber.Value);
-		SetTestFrame(ScrubPosition.FrameNumber.Value);
+
+		const UMLDeformerVizSettings* VizSettings = GetModel()->GetVizSettings();
+		if (VizSettings->GetVisualizationMode() == EMLDeformerVizMode::TrainingData)
+		{
+			SetTrainingFrame(ScrubPosition.FrameNumber.Value);
+		}
+		else if (VizSettings->GetVisualizationMode() == EMLDeformerVizMode::TestData)
+		{
+			SetTestFrame(ScrubPosition.FrameNumber.Value);
+		}
 	}
 
 	void FMLDeformerEditorModel::SetScrubPosition(FFrameNumber NewScrubPostion)
 	{
-		double Resolution = GetTickResolution();
-		double ClampedValue = FMath::Clamp((double)NewScrubPostion.Value, PlaybackRange.GetLowerBoundValue() * Resolution, PlaybackRange.GetUpperBoundValue() * Resolution);
-		int32 FrameValue = ClampedValue * (double)GetFrameRate() / Resolution;
-		SetTrainingFrame(FrameValue);
-		SetTestFrame(FrameValue);
+		const double Resolution = GetTickResolution();
+		const double ClampedValue = FMath::Clamp((double)NewScrubPostion.Value, PlaybackRange.GetLowerBoundValue() * Resolution, PlaybackRange.GetUpperBoundValue() * Resolution);
+		const int32 FrameValue = ClampedValue * (double)GetFrameRate() / Resolution;
+
+		const UMLDeformerVizSettings* VizSettings = GetModel()->GetVizSettings();
+		if (VizSettings->GetVisualizationMode() == EMLDeformerVizMode::TrainingData)
+		{
+			SetTrainingFrame(FrameValue);
+		}
+		else if (VizSettings->GetVisualizationMode() == EMLDeformerVizMode::TestData)
+		{
+			SetTestFrame(FrameValue);
+		}
 	}
 
-	void  FMLDeformerEditorModel::SetDisplayFrames(bool bInDisplayFrames)
+	void FMLDeformerEditorModel::SetDisplayFrames(bool bInDisplayFrames)
 	{
 		bDisplayFrames = bInDisplayFrames;
 	}
 
-	bool  FMLDeformerEditorModel::IsDisplayingFrames() const
+	bool FMLDeformerEditorModel::IsDisplayingFrames() const
 	{
 		return bDisplayFrames;
 	}
@@ -1913,6 +1942,26 @@ namespace UE::MLDeformer
 		return FMath::RoundToInt32((double)GetTicksPerFrame() * GetFrameRate());
 	}
 
+	void FMLDeformerEditorModel::OnPostEditChangeProperty(FPropertyChangedEvent& PropertyChangedEvent)
+	{
+		const FProperty* Property = PropertyChangedEvent.Property;
+		if (Property == nullptr)
+		{
+			return;
+		}
+
+		OnPropertyChanged(PropertyChangedEvent);
+	}
+
+	FString FMLDeformerEditorModel::GetTrainedNetworkOnnxFile() const
+	{
+		const FString ClassName = Model->GetClass()->GetFName().ToString();
+
+		// This returns something like:
+		// "<intermediatedir>/YourModel/YourModel.onnx"
+		return FString::Printf(TEXT("%s%s/%s.onnx"), *FPaths::ProjectIntermediateDir(), *ClassName, *ClassName);
+	}
+
 }	// namespace UE::MLDeformer
 
 #undef LOCTEXT_NAMESPACE
diff --git a/Engine/Plugins/Experimental/MLDeformer/MLDeformerFramework/Source/MLDeformerFrameworkEditor/Private/MLDeformerGeomCacheEditorModel.cpp b/Engine/Plugins/Experimental/MLDeformer/MLDeformerFramework/Source/MLDeformerFrameworkEditor/Private/MLDeformerGeomCacheEditorModel.cpp
new file mode 100644
--- /dev/null
+++ b/Engine/Plugins/Experimental/MLDeformer/MLDeformerFramework/Source/MLDeformerFrameworkEditor/Private/MLDeformerGeomCacheEditorModel.cpp
@@ -0,0 +1,276 @@
+// Copyright Epic Games, Inc. All Rights Reserved.
+
+#include "MLDeformerGeomCacheEditorModel.h"
+#include "MLDeformerGeomCacheActor.h"
+#include "MLDeformerEditorToolkit.h"
+#include "MLDeformerEditorStyle.h"
+#include "MLDeformerGeomCacheModel.h"
+#include "MLDeformerGeomCacheVizSettings.h"
+#include "MLDeformerInputInfo.h"
+#include "MLDeformerGeomCacheSampler.h"
+#include "MLDeformerGeomCacheHelpers.h"
+#include "MLDeformerGeomCacheModel.h"
+#include "GeometryCache.h"
+#include "GeometryCacheComponent.h"
+#include "Animation/DebugSkelMeshComponent.h"
+#include "Animation/MorphTarget.h"
+#include "Rendering/SkeletalMeshModel.h"
+#include "Rendering/SkeletalMeshLODModel.h"
+#include "Engine/SkeletalMesh.h"
+
+#define LOCTEXT_NAMESPACE "MLDeformerGeomCacheEditorModel"
+
+namespace UE::MLDeformer
+{
+	FMLDeformerSampler* FMLDeformerGeomCacheEditorModel::CreateSampler() const
+	{
+		FMLDeformerGeomCacheSampler* NewSampler = new FMLDeformerGeomCacheSampler();
+		NewSampler->OnGetGeometryCache().BindLambda([this] { return GetGeomCacheModel()->GetGeometryCache(); });
+		return NewSampler;
+	}
+
+	FMLDeformerEditorActor* FMLDeformerGeomCacheEditorModel::CreateEditorActor(const FMLDeformerEditorActor::FConstructSettings& Settings) const
+	{
+		return new FMLDeformerGeomCacheActor(Settings);
+	}
+
+	UMLDeformerGeomCacheVizSettings* FMLDeformerGeomCacheEditorModel::GetGeomCacheVizSettings() const 
+	{ 
+		return Cast<UMLDeformerGeomCacheVizSettings>(Model->GetVizSettings());
+	}
+
+	FMLDeformerGeomCacheSampler* FMLDeformerGeomCacheEditorModel::GetGeomCacheSampler() const
+	{ 
+		return static_cast<FMLDeformerGeomCacheSampler*>(Sampler);
+	}
+
+	UMLDeformerGeomCacheModel* FMLDeformerGeomCacheEditorModel::GetGeomCacheModel() const
+	{ 
+		return Cast<UMLDeformerGeomCacheModel>(Model);
+	}
+
+	FMLDeformerGeomCacheActor* FMLDeformerGeomCacheEditorModel::FindGeomCacheEditorActor(int32 ID) const
+	{
+		return static_cast<FMLDeformerGeomCacheActor*>(FindEditorActor(ID));
+	}
+
+	double FMLDeformerGeomCacheEditorModel::GetTrainingTimeAtFrame(int32 FrameNumber) const
+	{
+		// Try to get the frame from the geometry cache.
+		const FMLDeformerGeomCacheActor* EditorActor = static_cast<FMLDeformerGeomCacheActor*>(FindEditorActor(ActorID_Train_GroundTruth));
+		if (EditorActor && EditorActor->GetGeometryCacheComponent() && EditorActor->GetGeometryCacheComponent()->GeometryCache.Get())
+		{
+			return EditorActor->GetGeometryCacheComponent()->GetTimeAtFrame(FrameNumber);
+		}
+
+		return FMLDeformerEditorModel::GetTrainingTimeAtFrame(FrameNumber);
+	}
+
+	int32 FMLDeformerGeomCacheEditorModel::GetTrainingFrameAtTime(double TimeInSeconds) const
+	{
+		const FMLDeformerGeomCacheActor* EditorActor = static_cast<FMLDeformerGeomCacheActor*>(FindEditorActor(ActorID_Train_GroundTruth));
+		if (EditorActor && EditorActor->GetGeometryCacheComponent() && EditorActor->GetGeometryCacheComponent()->GeometryCache.Get())
+		{
+			return EditorActor->GetGeometryCacheComponent()->GetFrameAtTime(TimeInSeconds);
+		}
+
+		return FMLDeformerEditorModel::GetTrainingFrameAtTime(TimeInSeconds);
+	}
+
+	double FMLDeformerGeomCacheEditorModel::GetTestTimeAtFrame(int32 FrameNumber) const
+	{
+		// Try to get the frame from the geometry cache.
+		const FMLDeformerGeomCacheActor* EditorActor = static_cast<FMLDeformerGeomCacheActor*>(FindEditorActor(ActorID_Test_GroundTruth));
+		if (EditorActor && EditorActor->GetGeometryCacheComponent() && EditorActor->GetGeometryCacheComponent()->GeometryCache.Get())
+		{
+			return EditorActor->GetGeometryCacheComponent()->GetTimeAtFrame(FrameNumber);
+		}
+
+		return FMLDeformerEditorModel::GetTestTimeAtFrame(FrameNumber);
+	}
+
+	int32 FMLDeformerGeomCacheEditorModel::GetTestFrameAtTime(double TimeInSeconds) const
+	{
+		const FMLDeformerGeomCacheActor* EditorActor = static_cast<FMLDeformerGeomCacheActor*>(FindEditorActor(ActorID_Test_GroundTruth));
+		if (EditorActor && EditorActor->GetGeometryCacheComponent() && EditorActor->GetGeometryCacheComponent()->GeometryCache.Get())
+		{
+			return EditorActor->GetGeometryCacheComponent()->GetFrameAtTime(TimeInSeconds);
+		}
+
+		return FMLDeformerEditorModel::GetTestFrameAtTime(TimeInSeconds);
+	}
+
+	int32 FMLDeformerGeomCacheEditorModel::GetNumTrainingFrames() const
+	{
+		const UGeometryCache* GeometryCache = GetGeomCacheModel()->GetGeometryCache();
+		if (GeometryCache == nullptr)
+		{
+			return 0;
+		}
+		const int32 StartFrame = GeometryCache->GetStartFrame();
+		const int32 EndFrame = GeometryCache->GetEndFrame();
+		return (EndFrame - StartFrame) + 1;
+	}
+
+	void FMLDeformerGeomCacheEditorModel::UpdateIsReadyForTrainingState()
+	{
+		bIsReadyForTraining = false;
+
+		// Do some basic checks first, like if there is a skeletal mesh, ground truth, anim sequence, and if there are frames.
+		if (!IsEditorReadyForTrainingBasicChecks())
+		{
+			return;
+		}
+
+		// Now make sure the assets are compatible.
+		UMLDeformerGeomCacheModel* GeomCacheModel = GetGeomCacheModel();
+		UGeometryCache* GeomCache = GeomCacheModel->GetGeometryCache();
+		USkeletalMesh* SkeletalMesh = GeomCacheModel->GetSkeletalMesh();
+		if (!GetGeomCacheErrorText(SkeletalMesh, GeomCache).IsEmpty())
+		{
+			return;
+		}
+
+		// Make sure every skeletal imported mesh has some geometry track.
+		const int32 NumGeomCacheTracks = GeomCache ? GeomCache->Tracks.Num() : 0;
+		int32 NumSkelMeshes = 0;
+		if (SkeletalMesh)
+		{
+			FSkeletalMeshModel* ImportedModel = SkeletalMesh->GetImportedModel();
+			if (ImportedModel)
+			{
+				NumSkelMeshes = ImportedModel->LODModels[0].ImportedMeshInfos.Num();
+			}
+		}
+
+		// Allow the special case where there is just one mesh and track.
+		if (NumGeomCacheTracks != 1 || NumSkelMeshes != 1)
+		{
+			if (!GetGeomCacheSampler()->GetFailedImportedMeshNames().IsEmpty())
+			{
+				return;
+			}
+		}
+
+		bIsReadyForTraining = true;
+	}
+
+	void FMLDeformerGeomCacheEditorModel::CreateTrainingGroundTruthActor(UWorld* World)
+	{
+		UGeometryCache* GeomCache = GetGeomCacheModel()->GetGeometryCache();
+		const FLinearColor LabelColor = FMLDeformerEditorStyle::Get().GetColor("MLDeformer.TargetMesh.LabelColor");
+		const FLinearColor WireframeColor = FMLDeformerEditorStyle::Get().GetColor("MLDeformer.TargetMesh.WireframeColor");
+		CreateGeomCacheActor(
+			World,
+			ActorID_Train_GroundTruth,
+			"Train GroundTruth",
+			GeomCache,
+			LabelColor,
+			WireframeColor,
+			LOCTEXT("TrainGroundTruthActorLabelText", "Target Mesh"),
+			true);
+	}
+
+	void FMLDeformerGeomCacheEditorModel::CreateTestGroundTruthActor(UWorld* World)
+	{
+		UGeometryCache* GeomCache = GetGeomCacheVizSettings()->GetTestGroundTruth();
+		const FLinearColor LabelColor = FMLDeformerEditorStyle::Get().GetColor("MLDeformer.GroundTruth.LabelColor");
+		const FLinearColor WireframeColor = FMLDeformerEditorStyle::Get().GetColor("MLDeformer.GroundTruth.WireframeColor");
+		CreateGeomCacheActor(
+			World,
+			ActorID_Test_GroundTruth,
+			"Test GroundTruth",
+			GeomCache,
+			LabelColor,
+			WireframeColor,
+			LOCTEXT("TestGroundTruthActorLabelText", "Ground Truth"),
+			false);
+	}
+
+	void FMLDeformerGeomCacheEditorModel::CreateGeomCacheActor(UWorld* World, int32 ActorID, const FName& Name, UGeometryCache* GeomCache, FLinearColor LabelColor, FLinearColor WireframeColor, const FText& LabelText, bool bIsTrainingActor)
+	{
+		FActorSpawnParameters SpawnParams;
+		SpawnParams.Name = MakeUniqueObjectName(World, AActor::StaticClass(), Name);
+		AActor* Actor = World->SpawnActor<AActor>(SpawnParams);
+		Actor->SetFlags(RF_Transient);
+
+		// Create the Geometry Cache Component.
+		UGeometryCacheComponent* GeomCacheComponent = NewObject<UGeometryCacheComponent>(Actor);
+		GeomCacheComponent->SetGeometryCache(GeomCache);
+		GeomCacheComponent->RegisterComponent();
+		GeomCacheComponent->SetOverrideWireframeColor(true);
+		GeomCacheComponent->SetWireframeOverrideColor(WireframeColor);
+		GeomCacheComponent->MarkRenderStateDirty();
+		GeomCacheComponent->SetVisibility(false);
+		Actor->SetRootComponent(GeomCacheComponent);
+
+		// Create the editor actor.
+		FMLDeformerEditorActor::FConstructSettings Settings;
+		Settings.Actor = Actor;
+		Settings.TypeID = ActorID;
+		Settings.LabelColor = LabelColor;
+		Settings.LabelText = LabelText;
+		Settings.bIsTrainingActor = bIsTrainingActor;
+		FMLDeformerGeomCacheActor* EditorActor = static_cast<FMLDeformerGeomCacheActor*>(CreateEditorActor(Settings));
+		EditorActor->SetGeometryCacheComponent(GeomCacheComponent);
+		EditorActors.Add(EditorActor);
+	}
+
+	ETrainingResult FMLDeformerGeomCacheEditorModel::Train()
+	{
+		UE_LOG(LogMLDeformer, Error, TEXT("Please implement your Train method."));
+		return ETrainingResult::Success;
+	}
+
+	void FMLDeformerGeomCacheEditorModel::OnPropertyChanged(FPropertyChangedEvent& PropertyChangedEvent)
+	{
+		const FProperty* Property = PropertyChangedEvent.Property;
+		if (Property == nullptr)
+		{
+			return;
+		}
+
+		// Handle base class property changes.
+		FMLDeformerEditorModel::OnPropertyChanged(PropertyChangedEvent);
+
+		// Properties specific to the geometry cache editor model.
+		if (Property->GetFName() == GET_MEMBER_NAME_CHECKED(UMLDeformerGeomCacheModel, GeometryCache) ||
+			Property->GetFName() == GET_MEMBER_NAME_CHECKED(UMLDeformerGeomCacheVizSettings, GroundTruth))
+		{
+			TriggerInputAssetChanged(true);
+		}
+	}
+
+	void FMLDeformerGeomCacheEditorModel::OnInputAssetsChanged()
+	{
+		// Update the skeletal mesh components of the training, test base, and ml deformed actor.
+		FMLDeformerEditorModel::OnInputAssetsChanged();
+
+		UMLDeformerGeomCacheVizSettings* VizSettings = GetGeomCacheVizSettings();
+		check(VizSettings);
+		const float TestAnimSpeed = VizSettings->GetAnimPlaySpeed();
+		UAnimSequence* TestAnimSequence = VizSettings->GetTestAnimSequence();
+
+		// Update the training geometry cache.
+		UGeometryCacheComponent* GeometryCacheComponent = FindGeomCacheEditorActor(ActorID_Train_GroundTruth)->GetGeometryCacheComponent();
+		check(GeometryCacheComponent);
+		GeometryCacheComponent->SetGeometryCache(GetGeomCacheModel()->GetGeometryCache());
+		GeometryCacheComponent->SetLooping(false);
+		GeometryCacheComponent->SetManualTick(true);
+		GeometryCacheComponent->SetPlaybackSpeed(TestAnimSpeed);
+		GeometryCacheComponent->Play();
+
+		// Update the test geometry cache (ground truth) component.
+		GeometryCacheComponent = FindGeomCacheEditorActor(ActorID_Test_GroundTruth)->GetGeometryCacheComponent();
+		check(GeometryCacheComponent);
+		GeometryCacheComponent->SetGeometryCache(VizSettings->GetTestGroundTruth());
+		GeometryCacheComponent->SetLooping(true);
+		GeometryCacheComponent->SetManualTick(true);
+		GeometryCacheComponent->SetPlaybackSpeed(TestAnimSpeed);
+		GeometryCacheComponent->Play();
+
+		GetGeomCacheModel()->MeshMappings.Reset();
+	}
+}	// namespace UE::MLDeformer
+
+#undef LOCTEXT_NAMESPACE
diff --git a/Engine/Plugins/Experimental/MLDeformer/MLDeformerFramework/Source/MLDeformerFrameworkEditor/Private/MLDeformerGeomCacheModelDetails.cpp b/Engine/Plugins/Experimental/MLDeformer/MLDeformerFramework/Source/MLDeformerFrameworkEditor/Private/MLDeformerGeomCacheModelDetails.cpp
new file mode 100644
--- /dev/null
+++ b/Engine/Plugins/Experimental/MLDeformer/MLDeformerFramework/Source/MLDeformerFrameworkEditor/Private/MLDeformerGeomCacheModelDetails.cpp
@@ -0,0 +1,82 @@
+// Copyright Epic Games, Inc. All Rights Reserved.
+
+#include "MLDeformerGeomCacheModelDetails.h"
+#include "MLDeformerGeomCacheEditorModel.h"
+#include "MLDeformerGeomCacheModel.h"
+#include "DetailCategoryBuilder.h"
+#include "DetailWidgetRow.h"
+#include "Widgets/Layout/SBox.h"
+#include "SWarningOrErrorBox.h"
+
+#define LOCTEXT_NAMESPACE "MLDeformerGeomCacheModelDetails"
+
+namespace UE::MLDeformer
+{
+	bool FMLDeformerGeomCacheModelDetails::UpdateMemberPointers(const TArray<TWeakObjectPtr<UObject>>& Objects)
+	{
+		if (!FMLDeformerModelDetails::UpdateMemberPointers(Objects))
+		{
+			return false;
+		}
+
+		GeomCacheModel = Cast<UMLDeformerGeomCacheModel>(Model);
+		check(GeomCacheModel);
+		GeomCacheEditorModel = static_cast<FMLDeformerGeomCacheEditorModel*>(EditorModel);
+
+		return (GeomCacheModel != nullptr && GeomCacheEditorModel != nullptr);
+	}
+
+	void FMLDeformerGeomCacheModelDetails::AddAnimSequenceErrors()
+	{
+		const FText WarningText = GetGeomCacheAnimSequenceErrorText(GeomCacheModel->GetGeometryCache(), Model->GetAnimSequence());
+		BaseMeshCategoryBuilder->AddCustomRow(FText::FromString("AnimSeqWarning"))
+			.Visibility(!WarningText.IsEmpty() ? EVisibility::Visible : EVisibility::Collapsed)
+			.WholeRowContent()
+			[
+				SNew(SBox)
+				.Padding(FMargin(0.0f, 4.0f))
+				[
+					SNew(SWarningOrErrorBox)
+					.MessageStyle(EMessageStyle::Warning)
+					.Message(WarningText)
+				]
+			];
+	}
+
+	void FMLDeformerGeomCacheModelDetails::AddTargetMesh()
+	{
+		TargetMeshCategoryBuilder->AddProperty(GET_MEMBER_NAME_CHECKED(UMLDeformerGeomCacheModel, GeometryCache), UMLDeformerGeomCacheModel::StaticClass());
+
+		const FText TargetMeshErrorText = GetGeomCacheErrorText(GeomCacheModel->GetSkeletalMesh(), GeomCacheModel->GetGeometryCache());
+		TargetMeshCategoryBuilder->AddCustomRow(FText::FromString("TargetMeshError"))
+			.Visibility(!TargetMeshErrorText.IsEmpty() ? EVisibility::Visible : EVisibility::Collapsed)
+			.WholeRowContent()
+			[
+				SNew(SBox)
+				.Padding(FMargin(0.0f, 4.0f))
+				[
+					SNew(SWarningOrErrorBox)
+					.MessageStyle(EMessageStyle::Error)
+					.Message(TargetMeshErrorText)
+				]
+			];
+
+		const FText ChangedErrorText = EditorModel->GetTargetAssetChangedErrorText();
+		TargetMeshCategoryBuilder->AddCustomRow(FText::FromString("TargetMeshChangedError"))
+			.Visibility(!ChangedErrorText.IsEmpty() ? EVisibility::Visible : EVisibility::Collapsed)
+			.WholeRowContent()
+			[
+				SNew(SBox)
+				.Padding(FMargin(0.0f, 4.0f))
+				[
+					SNew(SWarningOrErrorBox)
+					.MessageStyle(EMessageStyle::Error)
+					.Message(ChangedErrorText)
+				]
+			];
+
+		AddGeomCacheMeshMappingWarnings(TargetMeshCategoryBuilder, Model->GetSkeletalMesh(), GeomCacheModel->GetGeometryCache());
+	}
+}	// namespace UE::MLDeformer
+
+#undef LOCTEXT_NAMESPACE
diff --git a/Engine/Plugins/Experimental/MLDeformer/MLDeformerFramework/Source/MLDeformerFrameworkEditor/Private/MLDeformerGeomCacheVizSettingsDetails.cpp b/Engine/Plugins/Experimental/MLDeformer/MLDeformerFramework/Source/MLDeformerFrameworkEditor/Private/MLDeformerGeomCacheVizSettingsDetails.cpp
new file mode 100644
--- /dev/null
+++ b/Engine/Plugins/Experimental/MLDeformer/MLDeformerFramework/Source/MLDeformerFrameworkEditor/Private/MLDeformerGeomCacheVizSettingsDetails.cpp
@@ -0,0 +1,69 @@
+// Copyright Epic Games, Inc. All Rights Reserved.
+
+#include "MLDeformerGeomCacheVizSettingsDetails.h"
+#include "MLDeformerGeomCacheModel.h"
+#include "MLDeformerGeomCacheEditorModel.h"
+#include "MLDeformerEditorModule.h"
+#include "MLDeformerGeomCacheVizSettings.h"
+#include "DetailLayoutBuilder.h"
+#include "DetailCategoryBuilder.h"
+#include "DetailWidgetRow.h"
+#include "IDetailsView.h"
+#include "Modules/ModuleManager.h"
+#include "Widgets/Layout/SBox.h"
+#include "PropertyCustomizationHelpers.h"
+#include "PropertyHandle.h"
+#include "SWarningOrErrorBox.h"
+
+#define LOCTEXT_NAMESPACE "MLDeformerGeomCacheVizSettingsDetails"
+
+namespace UE::MLDeformer
+{
+	bool FMLDeformerGeomCacheVizSettingsDetails::UpdateMemberPointers(const TArray<TWeakObjectPtr<UObject>>& Objects)
+	{
+		if (!FMLDeformerVizSettingsDetails::UpdateMemberPointers(Objects))
+		{
+			return false;
+		}
+
+		GeomCacheModel = Cast<UMLDeformerGeomCacheModel>(Model);
+		GeomCacheVizSettings = Cast<UMLDeformerGeomCacheVizSettings>(VizSettings);
+		return (GeomCacheModel != nullptr && GeomCacheVizSettings != nullptr);
+	}
+
+	void FMLDeformerGeomCacheVizSettingsDetails::AddGroundTruth()
+	{
+		TestAssetsCategory->AddProperty(GET_MEMBER_NAME_CHECKED(UMLDeformerGeomCacheVizSettings, GroundTruth), UMLDeformerGeomCacheVizSettings::StaticClass());
+
+		// Show an error when the test anim sequence duration doesn't match the one of the ground truth.
+		const FText AnimErrorText = GetGeomCacheAnimSequenceErrorText(GeomCacheVizSettings->GetTestGroundTruth(), VizSettings->GetTestAnimSequence());
+		FDetailWidgetRow& GroundTruthAnimErrorRow = TestAssetsCategory->AddCustomRow(FText::FromString("GroundTruthAnimMismatchError"))
+			.Visibility(!AnimErrorText.IsEmpty() ? EVisibility::Visible : EVisibility::Collapsed)
+			.WholeRowContent()
+			[
+				SNew(SBox)
+				.Padding(FMargin(0.0f, 4.0f))
+				[
+					SNew(SWarningOrErrorBox)
+					.MessageStyle(EMessageStyle::Warning)
+					.Message(AnimErrorText)
+				]
+			];
+
+		const FText GeomErrorText = GetGeomCacheErrorText(Model->GetSkeletalMesh(), GeomCacheVizSettings->GetTestGroundTruth());
+		FDetailWidgetRow& GroundTruthGeomErrorRow = TestAssetsCategory->AddCustomRow(FText::FromString("GroundTruthGeomMismatchError"))
+			.Visibility(!GeomErrorText.IsEmpty() ? EVisibility::Visible : EVisibility::Collapsed)
+			.WholeRowContent()
+			[
+				SNew(SBox)
+				.Padding(FMargin(0.0f, 4.0f))
+				[
+					SNew(SWarningOrErrorBox)
+					.MessageStyle(EMessageStyle::Warning)
+					.Message(GeomErrorText)
+				]
+			];
+	}
+}	//namespace UE::MLDeformer
+
+#undef LOCTEXT_NAMESPACE
diff --git a/Engine/Plugins/Experimental/MLDeformer/MLDeformerFramework/Source/MLDeformerFrameworkEditor/Private/MLDeformerModelDetails.cpp b/Engine/Plugins/Experimental/MLDeformer/MLDeformerFramework/Source/MLDeformerFrameworkEditor/Private/MLDeformerModelDetails.cpp
--- a/Engine/Plugins/Experimental/MLDeformer/MLDeformerFramework/Source/MLDeformerFrameworkEditor/Private/MLDeformerModelDetails.cpp
+++ b/Engine/Plugins/Experimental/MLDeformer/MLDeformerFramework/Source/MLDeformerFrameworkEditor/Private/MLDeformerModelDetails.cpp
@@ -54,7 +54,7 @@ namespace UE::MLDeformer
 		BaseMeshCategoryBuilder = &DetailLayoutBuilder->EditCategory("Base Mesh", FText::GetEmpty(), ECategoryPriority::Important);
 		TargetMeshCategoryBuilder = &DetailLayoutBuilder->EditCategory("Target Mesh", FText::GetEmpty(), ECategoryPriority::Important);
 		InputOutputCategoryBuilder = &DetailLayoutBuilder->EditCategory("Inputs and Output", FText::GetEmpty(), ECategoryPriority::Important);
-		SettingsCategoryBuilder = &DetailLayoutBuilder->EditCategory("Training Settings", FText::GetEmpty(), ECategoryPriority::Important);
+		TrainingSettingsCategoryBuilder = &DetailLayoutBuilder->EditCategory("Training Settings", FText::GetEmpty(), ECategoryPriority::Important);
 	}
 
 	void FMLDeformerModelDetails::CustomizeDetails(IDetailLayoutBuilder& DetailBuilder)
@@ -309,7 +309,7 @@ namespace UE::MLDeformer
 		// Show a warning when no neural network has been set.
 		{		
 			UNeuralNetwork* NeuralNetwork = Model->GetNeuralNetwork();
-			FDetailWidgetRow& NeuralNetErrorRow = SettingsCategoryBuilder->AddCustomRow(FText::FromString("NeuralNetError"))
+			FDetailWidgetRow& NeuralNetErrorRow = TrainingSettingsCategoryBuilder->AddCustomRow(FText::FromString("NeuralNetError"))
 				.Visibility((NeuralNetwork == nullptr) ? EVisibility::Visible : EVisibility::Collapsed)
 				.WholeRowContent()
 				[
@@ -325,7 +325,7 @@ namespace UE::MLDeformer
 			// Check if our network is compatible with the skeletal mesh.
 			if (Model->GetSkeletalMesh() && NeuralNetwork)
 			{
-				FDetailWidgetRow& NeuralNetIncompatibleErrorRow = SettingsCategoryBuilder->AddCustomRow(FText::FromString("NeuralNetIncompatibleError"))
+				FDetailWidgetRow& NeuralNetIncompatibleErrorRow = TrainingSettingsCategoryBuilder->AddCustomRow(FText::FromString("NeuralNetIncompatibleError"))
 					.Visibility(!Model->GetInputInfo()->IsCompatible(Model->GetSkeletalMesh()) ? EVisibility::Visible : EVisibility::Collapsed)
 					.WholeRowContent()
 					[
diff --git a/Engine/Plugins/Experimental/MLDeformer/MLDeformerFramework/Source/MLDeformerFrameworkEditor/Private/MLDeformerMorphModelDetails.cpp b/Engine/Plugins/Experimental/MLDeformer/MLDeformerFramework/Source/MLDeformerFrameworkEditor/Private/MLDeformerMorphModelDetails.cpp
new file mode 100644
--- /dev/null
+++ b/Engine/Plugins/Experimental/MLDeformer/MLDeformerFramework/Source/MLDeformerFrameworkEditor/Private/MLDeformerMorphModelDetails.cpp
@@ -0,0 +1,43 @@
+// Copyright Epic Games, Inc. All Rights Reserved.
+
+#include "MLDeformerMorphModelDetails.h"
+#include "MLDeformerMorphModelEditorModel.h"
+#include "MLDeformerMorphModel.h"
+#include "DetailLayoutBuilder.h"
+#include "DetailCategoryBuilder.h"
+
+#define LOCTEXT_NAMESPACE "MLDeformerMorphModelDetails"
+
+namespace UE::MLDeformer
+{
+	bool FMLDeformerMorphModelDetails::UpdateMemberPointers(const TArray<TWeakObjectPtr<UObject>>& Objects)
+	{
+		if (!FMLDeformerGeomCacheModelDetails::UpdateMemberPointers(Objects))
+		{
+			return false;
+		}
+
+		MorphModel = Cast<UMLDeformerMorphModel>(Model);
+		check(MorphModel);
+		MorphModelEditorModel = static_cast<FMLDeformerMorphModelEditorModel*>(EditorModel);
+
+		return (MorphModel != nullptr && MorphModelEditorModel != nullptr);
+	}
+
+	void FMLDeformerMorphModelDetails::CreateCategories()
+	{
+		FMLDeformerGeomCacheModelDetails::CreateCategories();
+		MorphTargetCategoryBuilder = &DetailLayoutBuilder->EditCategory("Morph Targets", FText::GetEmpty(), ECategoryPriority::Important);
+	}
+
+	void FMLDeformerMorphModelDetails::CustomizeDetails(IDetailLayoutBuilder& DetailBuilder)
+	{
+		// Create all the detail categories and add the properties of the base class.
+		FMLDeformerGeomCacheModelDetails::CustomizeDetails(DetailBuilder);
+
+		MorphTargetCategoryBuilder->AddProperty(GET_MEMBER_NAME_CHECKED(UMLDeformerMorphModel, MorphTargetDeltaThreshold), UMLDeformerMorphModel::StaticClass());
+		MorphTargetCategoryBuilder->AddProperty(GET_MEMBER_NAME_CHECKED(UMLDeformerMorphModel, MorphTargetErrorTolerance), UMLDeformerMorphModel::StaticClass());
+	}
+}	// namespace UE::MLDeformer
+
+#undef LOCTEXT_NAMESPACE
diff --git a/Engine/Plugins/Experimental/MLDeformer/MLDeformerFramework/Source/MLDeformerFrameworkEditor/Private/MLDeformerMorphModelEditorModel.cpp b/Engine/Plugins/Experimental/MLDeformer/MLDeformerFramework/Source/MLDeformerFrameworkEditor/Private/MLDeformerMorphModelEditorModel.cpp
new file mode 100644
--- /dev/null
+++ b/Engine/Plugins/Experimental/MLDeformer/MLDeformerFramework/Source/MLDeformerFrameworkEditor/Private/MLDeformerMorphModelEditorModel.cpp
@@ -0,0 +1,128 @@
+// Copyright Epic Games, Inc. All Rights Reserved.
+
+#include "MLDeformerMorphModelEditorModel.h"
+#include "MLDeformerMorphModel.h"
+#include "MLDeformerMorphModelVizSettings.h"
+#include "MLDeformerGeomCacheActor.h"
+#include "MLDeformerEditorToolkit.h"
+#include "MLDeformerInputInfo.h"
+#include "MLDeformerGeomCacheHelpers.h"
+#include "GeometryCache.h"
+#include "GeometryCacheComponent.h"
+#include "Animation/DebugSkelMeshComponent.h"
+#include "Animation/MorphTarget.h"
+#include "Rendering/SkeletalMeshModel.h"
+#include "Rendering/SkeletalMeshLODModel.h"
+#include "Engine/SkeletalMesh.h"
+
+#define LOCTEXT_NAMESPACE "MLDeformerMorphModelEditorModel"
+
+namespace UE::MLDeformer
+{
+	FMLDeformerEditorModel* FMLDeformerMorphModelEditorModel::MakeInstance()
+	{
+		return new FMLDeformerMorphModelEditorModel();
+	}
+
+	void FMLDeformerMorphModelEditorModel::OnPropertyChanged(FPropertyChangedEvent& PropertyChangedEvent)
+	{
+		const FProperty* Property = PropertyChangedEvent.Property;
+		if (Property == nullptr)
+		{
+			return;
+		}
+
+		FMLDeformerGeomCacheEditorModel::OnPropertyChanged(PropertyChangedEvent);
+
+		if (Property->GetFName() == GET_MEMBER_NAME_CHECKED(UMLDeformerMorphModelVizSettings, MorphTargetNumber))
+		{
+			const UMLDeformerMorphModel* MorphModel = GetMorphModel();
+			UMLDeformerMorphModelVizSettings* MorphViz = GetMorphModelVizSettings();
+			const int32 NumMorphTargets = MorphModel->MorphTargetSet.IsValid() ? MorphModel->MorphTargetSet->MorphBuffers.GetNumMorphs() : 0;
+			MorphViz->MorphTargetNumber = (NumMorphTargets > 0) ? FMath::Min<int32>(MorphViz->MorphTargetNumber, NumMorphTargets - 1) : 0;
+		}
+	}
+
+	UMLDeformerMorphModel* FMLDeformerMorphModelEditorModel::GetMorphModel() const
+	{ 
+		return Cast<UMLDeformerMorphModel>(Model);
+	}
+
+	UMLDeformerMorphModelVizSettings* FMLDeformerMorphModelEditorModel::GetMorphModelVizSettings() const
+	{
+		return Cast<UMLDeformerMorphModelVizSettings>(GetMorphModel()->GetVizSettings());
+	}
+
+	FString FMLDeformerMorphModelEditorModel::GetDefaultDeformerGraphAssetPath() const
+	{
+		return FString(TEXT("/Script/OptimusCore.OptimusDeformer'/Optimus/Deformers/DG_LinearBlendSkin_Morph_Cloth_RecomputeNormals.DG_LinearBlendSkin_Morph_Cloth_RecomputeNormals'"));
+	}
+
+	FString FMLDeformerMorphModelEditorModel::GetHeatMapDeformerGraphPath() const
+	{
+		return FString(TEXT("/MLDeformerFramework/Deformers/DG_MLDeformerModel_GPUMorph_HeatMap.DG_MLDeformerModel_GPUMorph_HeatMap"));
+	}
+
+	void FMLDeformerMorphModelEditorModel::OnPreTraining()
+	{
+		// Backup the morph target deltas in case we abort training.
+		MorphTargetDeltasBackup = GetMorphModel()->MorphTargetDeltas;
+	}
+
+	void FMLDeformerMorphModelEditorModel::OnPostTraining(ETrainingResult TrainingResult, bool bUsePartiallyTrainedWhenAborted)
+	{
+		// We aborted and don't want to use partially trained results, we should restore the deltas that we just overwrote after training.
+		if (TrainingResult == ETrainingResult::Aborted && !bUsePartiallyTrainedWhenAborted)
+		{
+			// Restore the blendshape backup.
+			GetMorphModel()->MorphTargetDeltas = MorphTargetDeltasBackup;
+		}
+		else if (TrainingResult == ETrainingResult::Success || (TrainingResult == ETrainingResult::Aborted && bUsePartiallyTrainedWhenAborted))
+		{
+			const UMLDeformerMorphModel* MorphModel = GetMorphModel();
+			if (!MorphModel->MorphTargetDeltas.IsEmpty())
+			{
+				// Set deltas with a length equal or below a given threshold to zero, for better compression.
+				TArray<FVector3f> MorphTargetDeltas = MorphModel->MorphTargetDeltas;
+				ZeroDeltasByThreshold(MorphTargetDeltas, MorphModel->MorphTargetDeltaThreshold);
+
+				// Build morph targets inside the engine, using the engine's compression scheme.
+				// Add one as we included the means now as extra morph target.
+				InitEngineMorphTargets(MorphTargetDeltas);
+			}
+		}
+
+		// This internally calls InitGPUData() which updates the GPU buffer with the deltas.
+		FMLDeformerGeomCacheEditorModel::OnPostTraining(TrainingResult, bUsePartiallyTrainedWhenAborted);
+	}
+
+	void FMLDeformerMorphModelEditorModel::InitEngineMorphTargets(const TArray<FVector3f>& Deltas)
+	{
+		const int32 LOD = 0;
+
+		// Turn the delta buffer in a set of engine morph targets.
+		UMLDeformerMorphModel* MorphModel = GetMorphModel();
+		TArray<UMorphTarget*> MorphTargets;	// These will be garbage collected.
+		CreateEngineMorphTargets(MorphTargets, Deltas, TEXT("MLDeformerMorph_"), LOD, MorphModel->MorphTargetDeltaThreshold);
+
+		// Now compress the morph targets to GPU friendly buffers.
+		check(MorphModel->MorphTargetSet.IsValid());
+		FMorphTargetVertexInfoBuffers& MorphBuffers = MorphModel->MorphTargetSet->MorphBuffers;
+		CompressEngineMorphTargets(MorphBuffers, MorphTargets, LOD, MorphModel->MorphTargetErrorTolerance);
+	}
+
+	void FMLDeformerMorphModelEditorModel::Render(const FSceneView* View, FViewport* Viewport, FPrimitiveDrawInterface* PDI)
+	{
+		FMLDeformerGeomCacheEditorModel::Render(View, Viewport, PDI);
+
+		// Debug draw the selected morph target.
+		const UMLDeformerMorphModelVizSettings* VizSettings = GetMorphModelVizSettings();
+		if (VizSettings->bDrawMorphTargets && VizSettings->GetVisualizationMode() == EMLDeformerVizMode::TestData)
+		{
+			const FVector DrawOffset = -VizSettings->GetMeshSpacingOffsetVector();
+			DrawMorphTarget(PDI, GetMorphModel()->MorphTargetDeltas, VizSettings->MorphTargetDeltaThreshold, VizSettings->MorphTargetNumber, DrawOffset);
+		}
+	}
+}	// namespace UE::MLDeformer
+
+#undef LOCTEXT_NAMESPACE
diff --git a/Engine/Plugins/Experimental/MLDeformer/MLDeformerFramework/Source/MLDeformerFrameworkEditor/Private/MLDeformerMorphModelVizSettingsDetails.cpp b/Engine/Plugins/Experimental/MLDeformer/MLDeformerFramework/Source/MLDeformerFrameworkEditor/Private/MLDeformerMorphModelVizSettingsDetails.cpp
new file mode 100644
--- /dev/null
+++ b/Engine/Plugins/Experimental/MLDeformer/MLDeformerFramework/Source/MLDeformerFrameworkEditor/Private/MLDeformerMorphModelVizSettingsDetails.cpp
@@ -0,0 +1,56 @@
+// Copyright Epic Games, Inc. All Rights Reserved.
+
+#include "MLDeformerMorphModelVizSettingsDetails.h"
+#include "MLDeformerMorphModel.h"
+#include "MLDeformerMorphModelEditorModel.h"
+#include "MLDeformerModule.h"
+#include "MLDeformerEditorModule.h"
+#include "MLDeformerGeomCacheHelpers.h"
+#include "MLDeformerMorphModelVizSettings.h"
+#include "DetailLayoutBuilder.h"
+#include "DetailCategoryBuilder.h"
+#include "DetailWidgetRow.h"
+#include "IDetailsView.h"
+#include "IDetailGroup.h"
+#include "Modules/ModuleManager.h"
+#include "Widgets/Layout/SBox.h"
+#include "PropertyCustomizationHelpers.h"
+#include "DetailLayoutBuilder.h"
+#include "SWarningOrErrorBox.h"
+
+#define LOCTEXT_NAMESPACE "MLDeformerMorphModelVizSettingsDetails"
+
+namespace UE::MLDeformer
+{
+	bool FMLDeformerMorphModelVizSettingsDetails::UpdateMemberPointers(const TArray<TWeakObjectPtr<UObject>>& Objects)
+	{
+		if (!FMLDeformerGeomCacheVizSettingsDetails::UpdateMemberPointers(Objects))
+		{
+			return false;
+		}
+
+		MorphModel = Cast<UMLDeformerMorphModel>(Model);
+		MorphModelVizSettings = Cast<UMLDeformerMorphModelVizSettings>(VizSettings);
+		return (MorphModel != nullptr && MorphModelVizSettings != nullptr);
+	}
+
+	bool FMLDeformerMorphModelVizSettingsDetails::IsMorphTargetsEnabled() const
+	{
+		return MorphModelVizSettings->bDrawMorphTargets && !MorphModel->MorphTargetDeltas.IsEmpty();
+	}
+
+	void FMLDeformerMorphModelVizSettingsDetails::AddAdditionalSettings()
+	{
+		IDetailGroup& MorphsGroup = LiveSettingsCategory->AddGroup("Morph Targets", LOCTEXT("MorphTargetsLabel", "Morph Targets"), false, true);
+		MorphsGroup.AddPropertyRow(DetailLayoutBuilder->GetProperty(GET_MEMBER_NAME_CHECKED(UMLDeformerMorphModelVizSettings, bDrawMorphTargets), UMLDeformerMorphModelVizSettings::StaticClass()))
+			.EditCondition(!MorphModel->MorphTargetDeltas.IsEmpty(), nullptr);
+
+		MorphsGroup.AddPropertyRow(DetailLayoutBuilder->GetProperty(GET_MEMBER_NAME_CHECKED(UMLDeformerMorphModelVizSettings, MorphTargetNumber), UMLDeformerMorphModelVizSettings::StaticClass()))
+			.EditCondition(TAttribute<bool>(this, &FMLDeformerMorphModelVizSettingsDetails::IsMorphTargetsEnabled), nullptr);
+
+		MorphsGroup.AddPropertyRow(DetailLayoutBuilder->GetProperty(GET_MEMBER_NAME_CHECKED(UMLDeformerMorphModelVizSettings, MorphTargetDeltaThreshold), UMLDeformerMorphModelVizSettings::StaticClass()))
+			.EditCondition(TAttribute<bool>(this, &FMLDeformerMorphModelVizSettingsDetails::IsMorphTargetsEnabled), nullptr);
+	}
+}	//namespace UE::MLDeformer
+
+#undef LOCTEXT_NAMESPACE
diff --git a/Engine/Plugins/Experimental/MLDeformer/MLDeformerFramework/Source/MLDeformerFrameworkEditor/Private/MLDeformerVizSettingsDetails.cpp b/Engine/Plugins/Experimental/MLDeformer/MLDeformerFramework/Source/MLDeformerFrameworkEditor/Private/MLDeformerVizSettingsDetails.cpp
--- a/Engine/Plugins/Experimental/MLDeformer/MLDeformerFramework/Source/MLDeformerFrameworkEditor/Private/MLDeformerVizSettingsDetails.cpp
+++ b/Engine/Plugins/Experimental/MLDeformer/MLDeformerFramework/Source/MLDeformerFrameworkEditor/Private/MLDeformerVizSettingsDetails.cpp
@@ -31,11 +31,6 @@
 
 namespace UE::MLDeformer
 {
-	TSharedRef<IDetailCustomization> FMLDeformerVizSettingsDetails::MakeInstance()
-	{
-		return MakeShareable(new FMLDeformerVizSettingsDetails());
-	}
-
 	bool FMLDeformerVizSettingsDetails::UpdateMemberPointers(const TArray<TWeakObjectPtr<UObject>>& Objects)
 	{
 		Model = nullptr;
diff --git a/Engine/Plugins/Experimental/MLDeformer/MLDeformerFramework/Source/MLDeformerFrameworkEditor/Public/MLDeformerEditorModel.h b/Engine/Plugins/Experimental/MLDeformer/MLDeformerFramework/Source/MLDeformerFrameworkEditor/Public/MLDeformerEditorModel.h
--- a/Engine/Plugins/Experimental/MLDeformer/MLDeformerFramework/Source/MLDeformerFrameworkEditor/Public/MLDeformerEditorModel.h
+++ b/Engine/Plugins/Experimental/MLDeformer/MLDeformerFramework/Source/MLDeformerFrameworkEditor/Public/MLDeformerEditorModel.h
@@ -17,7 +17,6 @@ class IPersonaPreviewScene;
 class FEditorViewportClient;
 class FSceneView;
 class FViewport;
-
 class FPrimitiveDrawInterface;
 class UWorld;
 class UMeshDeformer;
@@ -45,7 +44,7 @@ enum class ETrainingResult : uint8
 	/** The input or output data to the network has issues, which means we cannot train. */
 	FailOnData,
 
-	/** The python script has some error (see output log). */
+	/** The python script has some error (see output log). Or it's missing some required model training class. */
 	FailPythonError
 };
 
@@ -55,7 +54,13 @@ namespace UE::MLDeformer
 	class FMLDeformerSampler;
 	class FMLDeformerEditorModel;
 
-	/** The base class for the editor side of an UMLDeformerModel. */
+	/**
+	 * The base class for the editor side of an UMLDeformerModel.
+	 * The editor model class handles most of the editor interactions, such as property changes, triggering the training process,
+	 * creation of the viewport actors, and many more things.
+	 * In most cases you will want to inherit your own editor model from either FMLDeformerGeomCacheEditorModel or FMLDeformerMorphModelEditorModel though.
+	 * These classes will already implement a lot of base functionality in case you work with geometry caches as ground truth data etc.
+	 */
 	class MLDEFORMERFRAMEWORKEDITOR_API FMLDeformerEditorModel
 		: public TSharedFromThis<FMLDeformerEditorModel>
 		, public FGCObject
@@ -71,147 +76,672 @@ namespace UE::MLDeformer
 			UMLDeformerModel* Model = nullptr;
 		};
 
+		/**
+		 * The destructor, which unregisters the instance of this editor model from the model registry.
+		 * It also deletes the editor actors and unregisters from the OnPostEditChangeProperty delegate.
+		 */
 		virtual ~FMLDeformerEditorModel();
 
+		// FGCObject overrides.
+		virtual void AddReferencedObjects(FReferenceCollector& Collector) override;
+		virtual FString GetReferencerName() const override { return TEXT("FMLDeformerEditorModel"); }
+		// ~END FGCObject overrides.
+
 		// Required overrides.
+		/**
+		 * Get the number of training frames.
+		 * This must not be clamped yet to the maximum training frame limit. It must return the full amount of frames available in your training data.
+		 * For example this could return the number of frames inside your geometry cache.
+		 * @return The number of frames in the training data.
+		 */
 		virtual int32 GetNumTrainingFrames() const PURE_VIRTUAL(FMLDeformerEditorModel::GetNumTrainingFrames, return 0;);
+
+		/**
+		 * Launch the training. This gets executed when the Train button is pressed.
+		 * Training can succeed, or be aborted, and there can be training errors. The result is returned by this method.
+		 * You generally want to implement a class inherited from MLDeformerTrainingModel, and put a Train function inside of that, which you then execute through
+		 * the TrainModel function.
+		 * 
+		 * A code example:
+		 * 
+		 * @code{.cpp}
+		 * ETrainingResult FYourEditorModel::Train()
+		 * {
+		 *     return TrainModel<UYourTrainingModel>(this);
+		 * }
+		 * @endcode
+		 * 
+		 * That will internally call the Train function of your python class.
+		 * Now inside your Python class you do something like:
+		 * 
+		 * @code{.py}
+		 * @unreal.uclass()
+		 * class YourModelPythonTrainingModel(unreal.YourTrainingModel):
+		 *     @unreal.ufunction(override=True)
+		 *     def train(self):
+		 *         # ...do training here...
+		 *         return 0   # A value of 0 is success, 1 means aborted, see ETrainingResult.
+		 * @endcode
+		 * 
+		 * Calling the TrainModel function inside will then trigger the Python train method to be called.
+		 * 
+		 * @return The training result.
+		 */
 		virtual ETrainingResult Train() PURE_VIRTUAL(FMLDeformerEditorModel::Train, return ETrainingResult::Success;);
 
 		// Optional overrides.
+		/**
+		 * Initialize the model.
+		 * This will mainly create the delta sampler, editor input info and some other things.
+		 * @param Settings The initialization settings.
+		 */
 		virtual void Init(const InitSettings& Settings);
+
+		/**
+		 * Create the editor actors. These are the actors that will appear in the viewport of the asset editor.
+		 * Every actor has some specific ID. The base implementation of certain methods will assume there are specific actors, such as
+		 * a linear skinned actor, ML Deformed actor, and a ground truth one. You can find the ID values of those inside the FMLDeformerEditorActor.h file.
+		 * The ID values are things like: ActorID_Train_Base, ActorID_Train_GroundTruth, ActorID_Test_Base, ActorID_Test_MLDeformed and ActorID_Test_GroundTruth.
+		 * Nothing prevents you from adding more or less actors though.
+		 * @param InPersonaPreviewScene The persona scene.
+		 */
 		virtual void CreateActors(const TSharedRef<IPersonaPreviewScene>& InPersonaPreviewScene);
-		virtual void OnPostCreateActors() {}
+
+		/**
+		 * Clears the world, which basically means it removes all the editor actors and engine actors from the editor world.
+		 * It also resets the Persona preview scene.
+		 */
 		virtual void ClearWorld();
+
+		/**
+		 * Create an actor in the editor viewport.
+		 * Each actor will have a specific ID and a label that is displayed above it, with a specific color.
+		 * The ID values are things like: ActorID_Train_Base, ActorID_Train_GroundTruth, ActorID_Test_Base, ActorID_Test_MLDeformed and ActorID_Test_GroundTruth.
+		 * You can find the ID values of those inside the FMLDeformerEditorActor.h file.
+		 * @param Settings The creation settings for the editor actor you are creating.
+		 * @return A pointer to the editor actor object that was created.
+		 */
 		virtual FMLDeformerEditorActor* CreateEditorActor(const FMLDeformerEditorActor::FConstructSettings& Settings) const;
+
+		/**
+		 * Create the vertex delta sampler object.
+		 * You can create your own sampler in case you use anything else than say a Geometry Cache as ground truth target data.
+		 * A geometry cache based editor model would create a new FMLDeformerGeomCacheSampler for example.
+		 * @return A pointer to the newly created sampler.
+		 */
 		virtual FMLDeformerSampler* CreateSampler() const;
+
+		/**
+		 * Get the time in seconds, for a given frame in the training data.
+		 * @param FrameNumber The frame number to get the time in seconds for.
+		 * @return The time offset of this frame, in seconds.
+		 */
 		virtual double GetTrainingTimeAtFrame(int32 FrameNumber) const;
+
+		/**
+		 * Get the training data frame number for a specific time offset in seconds.
+		 * @param TimeInSeconds The time offset, in seconds, to get the frame number for.
+		 * @return The frame number for this time offset.
+		 */
 		virtual int32 GetTrainingFrameAtTime(double TimeInSeconds) const;
+
+		/**
+		 * Get the time in seconds, for a given frame in the test data.
+		 * @param FrameNumber The frame number to get the time in seconds for.
+		 * @return The time offset of this frame, in seconds.
+		 */
 		virtual double GetTestTimeAtFrame(int32 FrameNumber) const;
+
+		/**
+		 * Get the test data frame number for a specific time offset in seconds.
+		 * @param TimeInSeconds The time offset, in seconds, to get the frame number for.
+		 * @return The frame number for this time offset.
+		 */
 		virtual int32 GetTestFrameAtTime(double TimeInSeconds) const;
+
+		/**
+		 * Get the total number of frames in the test anim sequence.
+		 * @return The number of frames.
+		 */
 		virtual int32 GetNumTestFrames() const;
+
+		/**
+		 * Tick the editor viewport.
+		 * @param ViewportClient The editor viewport client object.
+		 * @param DeltaTime The delta time of this tick, in seconds.
+		 */
 		virtual void Tick(FEditorViewportClient* ViewportClient, float DeltaTime);
+
+		/**
+		 * Create the linear skinned actor that is used as training base.
+		 * This is the actor you see marked as "Training Base" in the training mode.
+		 * An editor actor needs to be created by this method.
+		 * The default implementation will directly set the persona actor and preview mesh and component to this actor.
+		 * @param InPersonaPreviewScene The personal preview scene.
+		 */
 		virtual void CreateTrainingLinearSkinnedActor(const TSharedRef<IPersonaPreviewScene>& InPersonaPreviewScene);
+
+		/**
+		 * Create the linear skinned test actor.
+		 * This is the "Linear Skinned" actor you see in testing mode.
+		 * An editor actor needs to be created by this method.
+		 * @param World The world to create the actor in.
+		 */
 		virtual void CreateTestLinearSkinnedActor(UWorld* World);
+
+		/**
+		 * Create the test ML Deformed actor.
+		 * This is the actor you see in test mode marked as "ML Deformed".
+		 * So this actor will have the ML Deformer component added to it as well.
+		 * An editor actor needs to be created by this method.
+		 * @param World The world to create the actor in.
+		 */
 		virtual void CreateTestMLDeformedActor(UWorld* World);
+
+		/**
+		 * Create the training ground truth actor.
+		 * This is the training target, so the actor that has the complex deformations, for example using a geometry cache.
+		 * An editor actor needs to be created by this method.
+		 * @param World The world to create the actor in.
+		 */
 		virtual void CreateTrainingGroundTruthActor(UWorld* World) {}
+
+		/**
+		 * Create the test ground truth actor.
+		 * This is the actor typically marked as "Ground Truth" in the editor.
+		 * It represents the test animation sequence ground truth deformation. So this could contain a geometry cache
+		 * equivalent to the animation sequence.
+		 * An editor actor needs to be created by this method.
+		 * @param World The world to create the actor in.
+		 */
 		virtual void CreateTestGroundTruthActor(UWorld* World) {}
+
+		/**
+		 * This is triggered when the training frame is changed by the user.
+		 * For example when they click the timeline in the editor, or enter a specific frame number.
+		 * This happens only in training mode.
+		 */
 		virtual void OnTrainingDataFrameChanged();
+
+		/**
+		 * Update the transforms of all the editor actors.
+		 * This should apply the mesh offsets to space the actors in the world/viewport.
+		 */
 		virtual void UpdateActorTransforms();
+
+		/**
+		 * Update the visibility status of the editor actors.
+		 * This should hide test actors when the editor is in training mode for example, and the other way around.
+		 */
 		virtual void UpdateActorVisibility();
+
+		/**
+		 * Update the visibility of the labels, as well as their positions.
+		 * For example labels of invisible actors should be made invisible as well.
+		 */
 		virtual void UpdateLabels();
+
+		/**
+		 * This is called when training data inputs change.
+		 * This includes changing things such as the Skeletal Mesh, training anim sequence and target mesh.
+		 * It will re-initialize the components. For example a skeletal mesh change will update the skeletal mesh component to the new mesh.
+		 */
 		virtual void OnInputAssetsChanged();
+
+		/**
+		 * This is executed after the input assets have changed.
+		 * On default this will refresh the deformer graphs, update the timeline ranges, select the right frames, modify the training button ready state, etc.
+		 */
 		virtual void OnPostInputAssetChanged();
-		virtual void HandleDefaultPropertyChanges(FPropertyChangedEvent& PropertyChangedEvent);
+
+		/**
+		 * This is called whenever a property value changes in the UI.
+		 * The default implementation handles all the shared property changes already.
+		 * You can overload this and call the base class method inside it, and then handle your own model specific property changes.
+		 * @param PropertyChangedEvent The event information object that contains info about the property change.
+		 */
+		virtual void OnPropertyChanged(FPropertyChangedEvent& PropertyChangedEvent);
+
+		/**
+		 * Executed when the user presses the play button in the timeline.
+		 * On default this will start playing the anim sequence related to whether we are in training or testing mode.
+		 */
 		virtual void OnPlayPressed();
+
+		/**
+		 * This is called just before we start the training process.
+		 * You can use it to backup any specific values that might be modified by the training process.
+		 * Sometimes you want to restore specific values when the user aborts the training, while the training already modified those values.
+		 */
 		virtual void OnPreTraining() {}
+
+		/**
+		 * This method is executed after training.
+		 * @param TrainingResult The result of the training process, such as whether it was a success or the user aborted, etc.
+		 * @param bUsePartiallyTrainedWhenAborted When the user aborts training, they can get offered to use the partially trained neural network. This boolean 
+		 *                                        specifies whether they picked "Yes" (true) or "No" (false) to that question.
+		 */
 		virtual void OnPostTraining(ETrainingResult TrainingResult, bool bUsePartiallyTrainedWhenAborted);
+
+		/**
+		 * This is executed when the training process got aborted by the user.
+		 * You can use this method to restore any values you had to backup during OnPreTraining for example.
+		 * @param bUsePartiallyTrainedWhenAborted When the user aborts training, they can get offered to use the partially trained neural network. This boolean 
+		 *                                        specifies whether they picked "Yes" (true) or "No" (false) to that question.
+		 */
 		virtual void OnTrainingAborted(bool bUsePartiallyTrainedData) {}
-		virtual void OnPostEditChangeProperty(FPropertyChangedEvent& PropertyChangedEvent) {}
+
+		/**
+		 * Check whether we are playing an animation or not.
+		 * This looks at whether we are in training or testing mode. It will look at the correct anim sequence based on this mode.
+		 * So if you are in training mode and the testing animation is playing, while the training sequence isn't playing, then it returns false.
+		 * It would in that case only return true when the training anim sequence is playing.
+		 * The same goes in testing mode, if we are in testing mode it will only return true when the test anim sequence is playing.
+		 * @return Returns true whether the mode specific animation sequence is playing.
+		 */
 		virtual bool IsPlayingAnim() const;
+
+		/**
+		 * Check whether our anim sequence is playing forward or backward.
+		 * Just like the IsPlayingAnim method, it is specific to the mode we are in.
+		 * @return Returns true when we are playing the animation forwards, and not backward.
+		 */
 		virtual bool IsPlayingForward() const;
+
+		/**
+		 * Get the current playback position of our training mode.
+		 * This will use the target mesh if there is any, and uses the anim sequence play offset in case no target mesh has been set yet.
+		 * So if you have a geometry cache or so specified as target mesh, it would extract the current time value from that.
+		 * @return The time offset, in seconds, of our training data.
+		 */
 		virtual double CalcTrainingTimelinePosition() const;
+
+		/**
+		 * Calculate the test mode timeline position.
+		 * This will use the ground truth test data if there is any, and uses the test anim sequence play offset in case no test ground truth has been set yet.
+		 * So if you have a geometry cache or so specified as test ground truth, it would extract the current time value from that.
+		 * @return The time offset, in seconds, of our test data.
+		 */
 		virtual double CalcTestTimelinePosition() const;
+
+		/**
+		 * This is executed when the timeline has been scrubbed.
+		 * Sometimes this is also executed without user interaction, but to update the same data when there was scrubbed.
+		 * @param NewScrubTime This is the new time value, in seconds.
+		 * @param bIsScrubbing Specifies whether the user is scrubbing or just clicked.
+		 */
 		virtual void OnTimeSliderScrubPositionChanged(double NewScrubTime, bool bIsScrubbing);
+
+		/**
+		 * Update the playback speed of the test data.
+		 * This will iterate over all the editor actors, check whether they are test mode actors, and if so
+		 * it will update the play speed. For example a linear skinned mesh will update its skeletal mesh component's play speed, while
+		 * a geometry cache ground truth actor would update the play speed set on the geometry cache component.
+		 */
 		virtual void UpdateTestAnimPlaySpeed();
+
+		/**
+		 * Clamp the current training frame index to be in a valid range.
+		 * A valid range would be between 0 and the number of frames in the training data.
+		 * So if the training data contains 100 frames, so index 0..99, and we set the TrainingFrameNumber member to 500, this method
+		 * will turn that into 99.
+		 */
 		virtual void ClampCurrentTrainingFrameIndex();
+
+		/**
+		 * Clamp the current test frame index to be in a valid range.
+		 * A valid range would be between 0 and the number of frames in the test data (anim sequence / ground truth).
+		 * So if the test data contains 100 frames, so index 0..99, and we set the TestingFrameNumber member to 500, this method
+		 * will turn that into 99.
+		 */
 		virtual void ClampCurrentTestFrameIndex();
+
+		/**
+		 * Get the total number of frames used for training, but limited by the maximum number of training frames we have setup.
+		 * So if the training data contains 10000 frames, but we set to train on a maximum of 2000 frames, then this will return 2000.
+		 * If the number of training frames is 10000 but we set the max number of training frames to 1 million, then it will still return only 10000.
+		 * @return The number of frames we should be training on.
+		 */
 		virtual int32 GetNumFramesForTraining() const;
+
+		/**
+		 * Go to a given frame in the training data.
+		 * Internally this will also call OnTimeSliderScrubPositionChanged, to go to the specific frame.
+		 * @param FrameNumber The training frame to go to. This will automatically be clamped to a valid range in case it goes out of bounds.
+		 */
 		virtual void SetTrainingFrame(int32 FrameNumber);
+
+		/**
+		 * Go to a given frame in the test data.
+		 * Internally this will also call OnTimeSliderScrubPositionChanged, to go to the specific frame.
+		 * @param FrameNumber The test frame to go to. This will automatically be clamped to a valid range in case it goes out of bounds.
+		 */
 		virtual void SetTestFrame(int32 FrameNumber);
+
+		/**
+		 * Render additional (debug) things in the viewport.
+		 * This will draw the vertex deltas when we are in training mode and aren't playing the training data.
+		 * @param View The scene view object.
+		 * @param Viewport The viewport object.
+		 * @param PDI The draw primitive interface used for drawing things.
+		 */
 		virtual void Render(const FSceneView* View, FViewport* Viewport, FPrimitiveDrawInterface* PDI);
+
+		/**
+		 * Update the state that is used to decide whether the Train button in the UI is enabled or not.
+		 * One criteria could be to check whether we have all required inputs, such as a skeletal mesh, target mesh and anim sequence.
+		 * This should set the bIsReadyForTraining bool to true when it is ready, or false if not.
+		 */
 		virtual void UpdateIsReadyForTrainingState() { bIsReadyForTraining = false; }
+
+		/**
+		 * Get the text that we should display as overlay on top of the asset editor viewport.
+		 * @return The text, or an empty text object if there is nothing to show.
+		 */
 		virtual FText GetOverlayText() const;
+
+		/**
+		 * Initialize the input information, which basically describes what the inputs to the neural network are going to be.
+		 * This contains things such as bone and curve names. The input info is used during both training and inference time.
+		 * It is used to see what bone transforms and curve values to feed into the network inputs.
+		 * Next to bones and curve reference it also stores the vertex counts of the base and target meshes, so we can later do some
+		 * compatibility checks when applying this ML Deformer onto some character.
+		 * @param InputInfo The input info object to initialize.
+		 */
 		virtual void InitInputInfo(UMLDeformerInputInfo* InputInfo);
+
+		/**
+		 * Call the SetupComponent on every ML Deformer component on our editor actors and update the compatibility status of the deformer components.
+		 * This gets triggered when some input assets change.
+		 */
 		virtual void RefreshMLDeformerComponents();
+
+		/**
+		 * Create the heat map material for this model.
+		 * This updates the HeatMapMaterial member.
+		 */
 		virtual void CreateHeatMapMaterial();
+
+		/**
+		 * Create the heat map deformer graph for this model.
+		 * This updates the HeatMapDeformerGraph member.
+		 */
 		virtual void CreateHeatMapDeformerGraph();
+
+		/**
+		 * Create the heat map assets for this model.
+		 * This calls CreateHeatMapMaterial and CreateHeatMapDeformerGraph on default.
+		 */
 		virtual void CreateHeatMapAssets();
+
+		/**
+		 * Specify whether the heat map material should be enabled or not.
+		 * This will activate the heat map visualization, or deactivate it.
+		 * @param bEnabled Set to true to enable heat map mode.
+		 */
 		virtual void SetHeatMapMaterialEnabled(bool bEnabled);
+
+		/**
+		 * Load the default Optimus deformer graph to use on the skeletal mesh when using this model.
+		 * @return The mesh deformer graph.
+		 */
 		virtual UMeshDeformer* LoadDefaultDeformerGraph();
+
+		/**
+		 * This makes sure that the deformer graph selected in the visualization settings is not a nullptr.
+		 * If the current graph hasn't been set, it will assign the default deformer graph to it.
+		 */
 		virtual void SetDefaultDeformerGraphIfNeeded();
+
+		/**
+		 * This applies the right mesh deformer to the editor actors that have an ML Deformer component on them.
+		 * Basically this just ensures that the heat map deformer graph is active when this is enabled. Otherwise the default
+		 * mesh deformer is used.
+		 */
 		virtual void UpdateDeformerGraph();
+
+		/**
+		 * Sample the vertex deltas between the training base model and target model.
+		 * This will initialize the sampler if needed and set the sampler space to post-skinning deltas, and then samples them using the sampler.
+		 * So this will update the internal state of the Sampler member. It will use the current training frame number to sample from.
+		 * If this number is somehow out of range it will automatically clamp it to be inside a valid range.
+		 */
 		virtual void SampleDeltas();
+
+		/**
+		 * Load the trained network from an onnx file.
+		 * The filename of the onnx file is determined by the GetTrainedNetworkOnnxFile method.
+		 */
 		virtual bool LoadTrainedNetwork() const;
+
+		/**
+		 * Get the onnx file name of the trained network.
+		 * This is the filename that your python script will save the onnx file to.
+		 * On default this will return the following filename: "<intermediatedir>/YourModel/YourModel.onnx".
+		 * Where "intermediatedir" is your project's Intermediate folder. And "YourModel" is the name of your runtime model class.
+		 * So if you have a model class named UFancyModel, on default this method will return "<intermediatedir>/FancyModel/FancyModel.onnx"
+		 */
+		virtual FString GetTrainedNetworkOnnxFile() const;
+
+		/**
+		 * Check whether our model is trained or not.
+		 * This is determined by looking whether the neural network pointer is nullptr or not.
+		 * @return Returns true when the model is trained already, or false if it hasn't been trained yet.
+		 */
 		virtual bool IsTrained() const;
+
+		/**
+		 * Get the editor actor that defines the timeline play position.
+		 * In training mode this is the training target actor, while in test mode it is the test ground truth actor.
+		 * @return A pointer to the ground truth actor depending on whether we're in test or training mode.
+		 */
 		virtual FMLDeformerEditorActor* GetTimelineEditorActor() const;
+
+		/**
+		 * Get the path to the heat map material asset.
+		 * @return The path to the heat map asset.
+		 */
 		virtual FString GetHeatMapMaterialPath() const;
+
+		/**
+		 * Get the path to the heat map deformer graph asset.
+		 * @return The path to the heat map deformer graph asset.
+		 */
 		virtual FString GetHeatMapDeformerGraphPath() const;
+
+		/**
+		 * Getthe path to the default deformer graph asset.
+		 * This is the deformer graph used when not in heat map mode.
+		 * @return The path to the default deformer graph asset.
+		 */
 		virtual FString GetDefaultDeformerGraphAssetPath() const;
-		virtual FString GetTrainedNetworkOnnxFile() const;
 
-		/** Get the current view range */
+		/** Get the current view range. */
 		TRange<double> GetViewRange() const;
 
-		/** Set the current view range */
+		/** Set the current view range. */
 		void SetViewRange(TRange<double> InRange);
 
-		/** Get the working range of the model's data */
+		/** Get the working range of the model's data. */
 		TRange<double> GetWorkingRange() const;
 
-		/** Get the playback range of the model's data */
+		/** Get the playback range of the model's data. */
 		TRange<FFrameNumber> GetPlaybackRange() const;
 
-		/** Get the current scrub position */
+		/** Get the current scrub position. */
 		FFrameNumber GetTickResScrubPosition() const;
 
+		/** Get how many ticks there are per frame. */
 		int32 GetTicksPerFrame() const;
-		/** Get the current scrub time */
+
+		/** Get the current scrub time. */
 		float GetScrubTime() const;
 
-		/** Set the current scrub position */
+		/** Set the current scrub position. */
 		void SetScrubPosition(FFrameTime NewScrubPostion);
 
-		/** Set the current scrub position */
+		/** Set the current scrub position. */
 		void SetScrubPosition(FFrameNumber NewScrubPostion);
 
-		/** Set if frames are displayed */
+		/** Set if frames are displayed. */
 		void SetDisplayFrames(bool bDisplayFrames);
 
+		/** Is the timeline displaying frames rather than seconds? */
 		bool IsDisplayingFrames() const;
 
+		/** Handle the timeline changes when our model changes. */
 		void HandleModelChanged();
+
+		/** Handle the timeline changes required when changing the visualization mode, so when we switch between training and testing mode. */
 		void HandleVizModeChanged(EMLDeformerVizMode Mode);
 
-		/** Handle the view range being changed */
+		/** Handle the view range being changed. */
 		void HandleViewRangeChanged(TRange<double> InRange);
 
-		/** Handle the working range being changed */
+		/** Handle the working range being changed. */
 		void HandleWorkingRangeChanged(TRange<double> InRange);
 
-		/** Get the framerate specified by the anim sequence */
+		/** Get the framerate specified by the anim sequence. */
 		double GetFrameRate() const;
 
-		/** Get the tick resolution we are displaying at */
+		/** Get the tick resolution we are displaying at. */
 		int32 GetTickResolution() const;
 
+		/** Get a pointer to the ML Deformer editor tookit. Cannot be a nullptr. */
 		FMLDeformerEditorToolkit* GetEditor() const { return Editor; }
+
+		/** Get a pointer to the runtime ML Deformer model. Cannot really be a nullptr. */
 		UMLDeformerModel* GetModel() const { return Model.Get(); }
+
+		/** Get a pointer to the world that our editor is displaying. Cannot be nullptr. */
 		UWorld* GetWorld() const;
+
+		/** Get all the created editor actors. */
 		const TArray<FMLDeformerEditorActor*>& GetEditorActors() const { return EditorActors; }
+
+		/** Find an editor actor with a specific ID, or return nullptr when not found. */
 		FMLDeformerEditorActor* FindEditorActor(int32 ActorTypeID) const;
+
+		/** Check whether we are ready to train. The Train button in the UI will be enabled or disabled based on this return value. */
 		bool IsReadyForTraining() const { return bIsReadyForTraining; }
+
+		/** Get the sampler we use to calculate vertex deltas when we are in training mdoe and enable deltas. */
 		FMLDeformerSampler* GetSampler() const { return Sampler; }
+
+		/**
+		 * Set whether we need to resample inputs or not. This can be used by the training code to determine if we need to resample inputs and outputs, or if
+		 * some cache can be used. Resampling is generally needed when input assets change, or other specific settings.
+		 * @param bNeeded Set to true to mark that inputs or other cached data need resampling.
+		 */
 		void SetResamplingInputOutputsNeeded(bool bNeeded) { bNeedToResampleInputOutputs = bNeeded; }
+
+		/**
+		 * Check whether we need to resample inputs and outputs.
+		 * @return Returns true if resampling of the inputs and outputs is needed, for example because input assets have changed.
+		 * @see SetResamplingInputOutputsNeeded
+		 */
 		bool GetResamplingInputOutputsNeeded() const { return bNeedToResampleInputOutputs; }
 
+		/**
+		 * Check whether the vertex counts of the skeletal mesh have changed compared to the one we have trained on.
+		 * @return An error string if there is an error, otherwise an empty text object is returned.
+		 */
 		FText GetBaseAssetChangedErrorText() const;
+
+		/**
+		 * Check whether the vertex map has changed or not.
+		 * This basically checks whether the vertex order has changed.
+		 * @return An error string if there is an error, otherwise an empty text object is returned.
+		 */
 		FText GetVertexMapChangedErrorText() const;
+
+		/**
+		 * Check for errors in the inputs. This basically checks if there are bones or curves to train on.
+		 * If there aren't any it will give some error.
+		 * @return An error string if there is an error, otherwise an empty text object is returned.
+		 */
 		FText GetInputsErrorText() const;
+
+		/**
+		 * Check for incompatible skeletons between the skeletal mesh and anim sequence skeleton.
+		 * @return An error string if there is an error, otherwise an empty text object is returned.
+		 */
 		FText GetIncompatibleSkeletonErrorText(USkeletalMesh* InSkelMesh, UAnimSequence* InAnimSeq) const;
+
+		/**
+		 * Check to see if the skeletal mesh has to be reimported because it misses some newly added data.
+		 * We have added imported mesh information to the Fbx importer. Older assets won't have this info, while this is required
+		 * to make the ML Deformer work properly.
+		 * @return An error string if there is an error, otherwise an empty text object is returned.
+		 */
 		FText GetSkeletalMeshNeedsReimportErrorText() const;
+
+		/**
+		 * Check whether the vertex count in the target mesh has changed or not.
+		 * If this is the case you would need to re-train the model.
+		 * @return An error string if there is an error, otherwise an empty text object is returned.
+		 */
 		FText GetTargetAssetChangedErrorText() const;
 
+		/**
+		 * Get the editor's input info object.
+		 * This is the information the user currently has setup in the UI, so not the input info of the runtime model.
+		 * The runtime model's input info contains the data inputs the model has been trained on, which can be different than what the user currently has setup in the UI.
+		 * @return A pointer to the editor's input info object.
+		 */
 		UMLDeformerInputInfo* GetEditorInputInfo() const { return EditorInputInfo.Get(); }	
+
+		/**
+		 * Update the editor's input info based on the current UI settings.
+		 */
 		void UpdateEditorInputInfo();
 
-		void TriggerInputAssetChanged(bool RefreshVizSettings=false);
+		/**
+		 * This will call OnInputAssetsChanged first, followed by OnPostInputAssetChanged.
+		 * It then will update the editor only flags for the assets, and then forces a refresh of the model details panel and optionally also the visualization settings panel.
+		 * @param RefreshVizSettings Set to true if you want the visualization settings details panel to also refresh.
+		 */
+		void TriggerInputAssetChanged(bool bRefreshVizSettings=false);
 
+		/**
+		 * Initialize the bone include list of the runtime model to the set of bones that are animated.
+		 * A bone is seen as animated if its animation data has keyframes that have a different value than the first keyframe.
+		 * So if the bone's animation data has 1000 keyframes, and they are all the same value, it won't see it as animated.
+		 */
 		void InitBoneIncludeListToAnimatedBonesOnly();
+
+		/**
+		 * Initialize the curve include list of the runtime model to the set of curves that are animated.
+		 * A curve is seen as animated if its animation data has keyframes that have a different value than the first keyframe.
+		 * So if the curve's animation data has 1000 keyframes, and they are all the same value, it won't see it as animated.
+		 */
 		void InitCurveIncludeListToAnimatedCurvesOnly();
+
+		/**
+		 * Get the number of curves on a specific skeletal mesh.
+		 * @param SkelMesh The skeletal mesh to get the curves count from.
+		 * @return The number of curves found on the skeletal mesh.
+		 */
 		int32 GetNumCurvesOnSkeletalMesh(USkeletalMesh* SkelMesh) const;
 
+		/**
+		 * Load a neural network from a given Onnx file.
+		 * @param Filename The onnx file name.
+		 * @return A pointer to the newly created neural network, or nullptr in case it failed to load.
+		 */
 		UNeuralNetwork* LoadNeuralNetworkFromOnnx(const FString& Filename) const;
+
+		/**
+		 * Get the currently desired training frame number.
+		 * You can call CheckTrainingFrameChanged in order to make the desired frame the actual frame.
+		 * @return The currenty desired training frame number, which might not the the same as the real training frame number.
+		 */
 		int32 GetCurrentTrainingFrame() const { return CurrentTrainingFrame; }
+
+		/**
+		 * Check whether our training data frame number changed.
+		 * We have a TrainingFrameNumber and a CurrentTrainingFrame. If they don't match it will call OnTrainingDataFrameChanged.
+		 */
 		void CheckTrainingDataFrameChanged();
 
 		/**
@@ -227,7 +757,25 @@ namespace UE::MLDeformer
 		void DrawMorphTarget(FPrimitiveDrawInterface* PDI, const TArray<FVector3f>& MorphDeltas, float DeltaThreshold, int32 MorphTargetIndex, const FVector& DrawOffset);
 
 	protected:
+		/**
+		 * Executed when a property changes by a change in the UI.
+		 * This internally calls OnPropertyChanged.
+		 */
+		void OnPostEditChangeProperty(FPropertyChangedEvent& PropertyChangedEvent);
+
+		/**
+		 * Delete all editor actors.
+		 * This does not actually delete the actors in the world.
+		 * You can call ClearWorld to include that.
+		 */
 		void DeleteEditorActors();
+
+		/**
+		 * Perform some basic checks to see if the editor can be ready to train the model.
+		 * This checks whether there is an anim sequence, target mesh, and skeletal mesh.
+		 * It can be called inside the UpdateIsReadyForTrainingState method, which can then do additional checks afterwards.
+		 * @return Returns true when at least the basic required inputs have been selected, false if not.
+		 */
 		bool IsEditorReadyForTrainingBasicChecks();
 
 		/** Zero all deltas with a length equal to, or smaller than the threshold value. */
@@ -252,12 +800,30 @@ namespace UE::MLDeformer
 		 */
 		void CompressEngineMorphTargets(FMorphTargetVertexInfoBuffers& OutMorphBuffers, const TArray<UMorphTarget*>& MorphTargets, int32 LOD = 0, float MorphErrorTolerance = 0.01f);
 
+		/**
+		 * Get the base actor depending on the currently active mode (train or test mode).
+		 * The base actor is the linear skinned one.
+		 * @return A pointer to the linear skinned actor in either test or training mode.
+		 */
 		FMLDeformerEditorActor* GetVisualizationModeBaseActor() const;
 
+		/**
+		 * Get the animation sequence, either the test or training one, depending on in which mode we are in.
+		 * This can return a nullptr if no anim sequence has been setup yet.
+		 * @return A pointer to the test sequence when we are in test mode, or the training anim sequence if we are in training mode.
+		 */
 		const UAnimSequence* GetAnimSequence() const;
 
+		/**
+		 * Calculate the current time line position offset, of the currently selected frame.
+		 * This depends on whether we are in test or training mode.
+		 * @return The time line offset, in seconds.
+		 */
 		double CalcTimelinePosition() const;
 
+		/**
+		 * Update the timeline related ranges, based on the length of the training or testing data.
+		 */
 		void UpdateRanges();
 
 	protected:
@@ -270,7 +836,7 @@ namespace UE::MLDeformer
 		/** A pointer to the editor toolkit. */
 		FMLDeformerEditorToolkit* Editor = nullptr;
 
-		/** A pointer to the sampler. */
+		/** A pointer to the sampler, which can sample target meshes to calculate deltas. */
 		FMLDeformerSampler* Sampler = nullptr;
 
 		/**
@@ -291,15 +857,6 @@ namespace UE::MLDeformer
 		/** The current training frame. */
 		int32 CurrentTrainingFrame = -1;
 
-		/**
-		 * Are we ready for training? 
-		 * The training button in the editor will be enabled or disabled based on this on default.
-		 */
-		bool bIsReadyForTraining = false;	
-
-		/** Do we need to resample all input/output data? */
-		bool bNeedToResampleInputOutputs = true;
-
 		/** The range we are currently viewing */
 		TRange<double> ViewRange;
 
@@ -309,11 +866,26 @@ namespace UE::MLDeformer
 		/** The playback range of this model for each timeframe */
 		TRange<double> PlaybackRange;
 
+		/** The current scrub position. */
 		FFrameTime ScrubPosition;
 
+		/**
+		 * Are we ready for training? 
+		 * The training button in the editor will be enabled or disabled based on this on default.
+		 */
+		bool bIsReadyForTraining = false;	
+
+		/** Do we need to resample all input/output data? */
+		bool bNeedToResampleInputOutputs = true;
+
+		/** Display frame numbers? */
 		bool bDisplayFrames = true;
 	};
 
+	/**
+	 * 
+	 * 
+	 */
 	template<class TrainingModelClass>
 	ETrainingResult TrainModel(FMLDeformerEditorModel* EditorModel)
 	{
diff --git a/Engine/Plugins/Experimental/MLDeformer/MLDeformerFramework/Source/MLDeformerFrameworkEditor/Public/MLDeformerGeomCacheEditorModel.h b/Engine/Plugins/Experimental/MLDeformer/MLDeformerFramework/Source/MLDeformerFrameworkEditor/Public/MLDeformerGeomCacheEditorModel.h
new file mode 100644
--- /dev/null
+++ b/Engine/Plugins/Experimental/MLDeformer/MLDeformerFramework/Source/MLDeformerFrameworkEditor/Public/MLDeformerGeomCacheEditorModel.h
@@ -0,0 +1,52 @@
+// Copyright Epic Games, Inc. All Rights Reserved.
+
+#pragma once
+
+#include "CoreMinimal.h"
+#include "MLDeformerEditorModel.h"
+
+class UMLDeformerGeomCacheModel;
+class UMLDeformerGeomCacheVizSettings;
+
+namespace UE::MLDeformer
+{
+	class FMLDeformerGeomCacheActor;
+	class FMLDeformerGeomCacheSampler;
+
+	/**
+	 * An editor model that is based on a geometry cache as training input.
+	 */
+	class MLDEFORMERFRAMEWORKEDITOR_API FMLDeformerGeomCacheEditorModel
+		: public FMLDeformerEditorModel
+	{
+	public:
+		// FGCObject overrides.
+		virtual FString GetReferencerName() const override { return TEXT("FMLDeformerGeomCacheEditorModel"); }
+		// ~END FGCObject overrides.
+
+		// FMLDeformerEditorModel overrides.
+		virtual FMLDeformerEditorActor* CreateEditorActor(const FMLDeformerEditorActor::FConstructSettings& Settings) const override;
+		virtual FMLDeformerSampler* CreateSampler() const override;
+		virtual void CreateTrainingGroundTruthActor(UWorld* World) override;
+		virtual void CreateTestGroundTruthActor(UWorld* World) override;
+		virtual int32 GetNumTrainingFrames() const override;
+		virtual double GetTrainingTimeAtFrame(int32 FrameNumber) const override;
+		virtual int32 GetTrainingFrameAtTime(double TimeInSeconds) const override;
+		virtual double GetTestTimeAtFrame(int32 FrameNumber) const override;
+		virtual int32 GetTestFrameAtTime(double TimeInSeconds) const override;
+		virtual void UpdateIsReadyForTrainingState() override;
+		virtual void OnPropertyChanged(FPropertyChangedEvent& PropertyChangedEvent) override;
+		virtual void OnInputAssetsChanged() override;
+		virtual ETrainingResult Train() override;
+		// ~END FMLDeformerEditorModel overrides.
+		
+		// Helpers.
+		FMLDeformerGeomCacheSampler* GetGeomCacheSampler() const;
+		UMLDeformerGeomCacheModel* GetGeomCacheModel() const;
+		UMLDeformerGeomCacheVizSettings* GetGeomCacheVizSettings() const;
+		FMLDeformerGeomCacheActor* FindGeomCacheEditorActor(int32 ID) const;
+
+	protected:
+		void CreateGeomCacheActor(UWorld* World, int32 ActorID, const FName& Name, UGeometryCache* GeomCache, FLinearColor LabelColor, FLinearColor WireframeColor, const FText& LabelText, bool bIsTrainingActor);
+	};
+}	// namespace UE::MLDeformer
diff --git a/Engine/Plugins/Experimental/MLDeformer/MLDeformerFramework/Source/MLDeformerFrameworkEditor/Public/MLDeformerGeomCacheModelDetails.h b/Engine/Plugins/Experimental/MLDeformer/MLDeformerFramework/Source/MLDeformerFrameworkEditor/Public/MLDeformerGeomCacheModelDetails.h
new file mode 100644
--- /dev/null
+++ b/Engine/Plugins/Experimental/MLDeformer/MLDeformerFramework/Source/MLDeformerFrameworkEditor/Public/MLDeformerGeomCacheModelDetails.h
@@ -0,0 +1,31 @@
+// Copyright Epic Games, Inc. All Rights Reserved.
+
+#pragma once
+
+#include "CoreMinimal.h"
+#include "MLDeformerModelDetails.h"
+#include "IDetailCustomization.h"
+
+class IDetailLayoutBuilder;
+class UMLDeformerGeomCacheModel;
+
+namespace UE::MLDeformer
+{
+	class FMLDeformerGeomCacheEditorModel;
+	class FMLDeformerEditorModel;
+
+	class MLDEFORMERFRAMEWORKEDITOR_API FMLDeformerGeomCacheModelDetails
+		: public FMLDeformerModelDetails
+	{
+	public:
+		// FMLDeformerModelDetails overrides.
+		virtual bool UpdateMemberPointers(const TArray<TWeakObjectPtr<UObject>>& Objects) override;
+		virtual void AddTargetMesh() override;
+		virtual void AddAnimSequenceErrors() override;
+		// ~END FMLDeformerModelDetails overrides.
+
+	protected:
+		TObjectPtr<UMLDeformerGeomCacheModel> GeomCacheModel = nullptr;
+		FMLDeformerGeomCacheEditorModel* GeomCacheEditorModel = nullptr;
+	};
+}	// namespace UE::MLDeformer
diff --git a/Engine/Plugins/Experimental/MLDeformer/MLDeformerFramework/Source/MLDeformerFrameworkEditor/Public/MLDeformerGeomCacheVizSettingsDetails.h b/Engine/Plugins/Experimental/MLDeformer/MLDeformerFramework/Source/MLDeformerFrameworkEditor/Public/MLDeformerGeomCacheVizSettingsDetails.h
new file mode 100644
--- /dev/null
+++ b/Engine/Plugins/Experimental/MLDeformer/MLDeformerFramework/Source/MLDeformerFrameworkEditor/Public/MLDeformerGeomCacheVizSettingsDetails.h
@@ -0,0 +1,28 @@
+// Copyright Epic Games, Inc. All Rights Reserved.
+
+#pragma once
+
+#include "CoreMinimal.h"
+#include "MLDeformerVizSettingsDetails.h"
+#include "IDetailCustomization.h"
+
+class IDetailLayoutBuilder;
+class UMLDeformerGeomCacheModel;
+class UMLDeformerGeomCacheVizSettings;
+
+namespace UE::MLDeformer
+{
+	class MLDEFORMERFRAMEWORKEDITOR_API FMLDeformerGeomCacheVizSettingsDetails
+		: public FMLDeformerVizSettingsDetails
+	{
+	public:
+		// FMLDeformerVizSettingsDetails overrides.
+		virtual bool UpdateMemberPointers(const TArray<TWeakObjectPtr<UObject>>& Objects) override;
+		virtual void AddGroundTruth() override;
+		// ~END FMLDeformerVizSettingsDetails overrides.
+
+	protected:
+		TObjectPtr<UMLDeformerGeomCacheModel> GeomCacheModel = nullptr;
+		TObjectPtr<UMLDeformerGeomCacheVizSettings> GeomCacheVizSettings = nullptr;
+	};
+}	// namespace UE::MLDeformer
diff --git a/Engine/Plugins/Experimental/MLDeformer/MLDeformerFramework/Source/MLDeformerFrameworkEditor/Public/MLDeformerModelDetails.h b/Engine/Plugins/Experimental/MLDeformer/MLDeformerFramework/Source/MLDeformerFrameworkEditor/Public/MLDeformerModelDetails.h
--- a/Engine/Plugins/Experimental/MLDeformer/MLDeformerFramework/Source/MLDeformerFrameworkEditor/Public/MLDeformerModelDetails.h
+++ b/Engine/Plugins/Experimental/MLDeformer/MLDeformerFramework/Source/MLDeformerFrameworkEditor/Public/MLDeformerModelDetails.h
@@ -55,6 +55,6 @@ namespace UE::MLDeformer
 		IDetailCategoryBuilder* BaseMeshCategoryBuilder = nullptr;
 		IDetailCategoryBuilder* TargetMeshCategoryBuilder = nullptr;
 		IDetailCategoryBuilder* InputOutputCategoryBuilder = nullptr;
-		IDetailCategoryBuilder* SettingsCategoryBuilder = nullptr;
+		IDetailCategoryBuilder* TrainingSettingsCategoryBuilder = nullptr;
 	};
 }	// namespace UE::MLDeformer
diff --git a/Engine/Plugins/Experimental/MLDeformer/MLDeformerFramework/Source/MLDeformerFrameworkEditor/Public/MLDeformerMorphModelDetails.h b/Engine/Plugins/Experimental/MLDeformer/MLDeformerFramework/Source/MLDeformerFrameworkEditor/Public/MLDeformerMorphModelDetails.h
new file mode 100644
--- /dev/null
+++ b/Engine/Plugins/Experimental/MLDeformer/MLDeformerFramework/Source/MLDeformerFrameworkEditor/Public/MLDeformerMorphModelDetails.h
@@ -0,0 +1,34 @@
+// Copyright Epic Games, Inc. All Rights Reserved.
+
+#pragma once
+
+#include "CoreMinimal.h"
+#include "MLDeformerGeomCacheModelDetails.h"
+#include "IDetailCustomization.h"
+
+class IDetailLayoutBuilder;
+class UMLDeformerMorphModel;
+
+namespace UE::MLDeformer
+{
+	class FMLDeformerMorphModelEditorModel;
+
+	class MLDEFORMERFRAMEWORKEDITOR_API FMLDeformerMorphModelDetails
+		: public FMLDeformerGeomCacheModelDetails
+	{
+	public:
+		// ILayoutDetails overrides.
+		virtual void CustomizeDetails(IDetailLayoutBuilder& DetailBuilder) override;
+		// ~END ILayoutDetails overrides.
+
+		// FMLDeformerModelDetails overrides.
+		virtual void CreateCategories() override;
+		virtual bool UpdateMemberPointers(const TArray<TWeakObjectPtr<UObject>>& Objects) override;
+		// ~END FMLDeformerModelDetails overrides.
+
+	protected:
+		TObjectPtr<UMLDeformerMorphModel> MorphModel = nullptr;
+		FMLDeformerMorphModelEditorModel* MorphModelEditorModel = nullptr;
+		IDetailCategoryBuilder* MorphTargetCategoryBuilder = nullptr;
+	};
+}	// namespace UE::MLDeformer
diff --git a/Engine/Plugins/Experimental/MLDeformer/MLDeformerFramework/Source/MLDeformerFrameworkEditor/Public/MLDeformerMorphModelEditorModel.h b/Engine/Plugins/Experimental/MLDeformer/MLDeformerFramework/Source/MLDeformerFrameworkEditor/Public/MLDeformerMorphModelEditorModel.h
new file mode 100644
--- /dev/null
+++ b/Engine/Plugins/Experimental/MLDeformer/MLDeformerFramework/Source/MLDeformerFrameworkEditor/Public/MLDeformerMorphModelEditorModel.h
@@ -0,0 +1,63 @@
+// Copyright Epic Games, Inc. All Rights Reserved.
+
+#pragma once
+
+#include "CoreMinimal.h"
+#include "MLDeformerGeomCacheEditorModel.h"
+#include "MLDeformerGeomCacheSampler.h"
+#include "MLDeformerEditorActor.h"
+#include "UObject/GCObject.h"
+
+class UMLDeformerAsset;
+class USkeletalMesh;
+class UGeometryCache;
+class UMLDeformerMorphModel;
+class UMLDeformerMorphModelVizSettings;
+
+namespace UE::MLDeformer
+{
+	class FMLDeformerGeomCacheActor;
+	class FNeuralMorphEditorModelActor;
+
+	class MLDEFORMERFRAMEWORKEDITOR_API FMLDeformerMorphModelEditorModel
+		: public FMLDeformerGeomCacheEditorModel
+	{
+	public:
+		// We need to implement this static MakeInstance method.
+		static FMLDeformerEditorModel* MakeInstance();
+
+		// FGCObject overrides.
+		virtual FString GetReferencerName() const override { return TEXT("FMLDeformerMorphModelEditorModel"); }
+		// ~END FGCObject overrides.
+
+		// FMLDeformerEditorModel overrides.
+		virtual void Render(const FSceneView* View, FViewport* Viewport, FPrimitiveDrawInterface* PDI) override;
+		virtual void OnPropertyChanged(FPropertyChangedEvent& PropertyChangedEvent) override;
+		virtual FString GetDefaultDeformerGraphAssetPath() const override;
+		virtual FString GetHeatMapDeformerGraphPath() const override;
+		virtual void OnPreTraining() override;
+		virtual void OnPostTraining(ETrainingResult TrainingResult, bool bUsePartiallyTrainedWhenAborted) override;
+		// ~END FMLDeformerEditorModel overrides.
+
+		// Helpers.
+		UMLDeformerMorphModel* GetMorphModel() const;
+		UMLDeformerMorphModelVizSettings* GetMorphModelVizSettings() const;
+
+	protected:
+		/**
+		 * Initialize a set of engine morph targets and compress them to GPU friendly buffers.
+		 * These morph targets are initialized from a set of deltas. Each morph target needs to have Model->GetNumBaseVerts() number of deltas.
+		 * All deltas are concatenated in one big array. So all deltas of all vertices for the second morph target are appended to the deltas for
+		 * the first morph target, etc. In other words, the layout is: [morph0_deltas][morph1_deltas][morph2_deltas][...].
+		 * @param Deltas The deltas for all morph targets concatenated. So the number of items in this array is a multiple of Model->GetNumBaseVerts().
+		 */
+		void InitEngineMorphTargets(const TArray<FVector3f>& Deltas);
+
+	protected:
+		/**
+		 * The entire set of morph target deltas, 3 per vertex, for each morph target, as one flattened buffer.
+		 * So the size of this buffer is: (NumVertsPerMorphTarget * NumMorphTargets).
+		 */
+		TArray<FVector3f> MorphTargetDeltasBackup;
+	};
+}	// namespace UE::MLDeformer
diff --git a/Engine/Plugins/Experimental/MLDeformer/MLDeformerFramework/Source/MLDeformerFrameworkEditor/Public/MLDeformerMorphModelVizSettingsDetails.h b/Engine/Plugins/Experimental/MLDeformer/MLDeformerFramework/Source/MLDeformerFrameworkEditor/Public/MLDeformerMorphModelVizSettingsDetails.h
new file mode 100644
--- /dev/null
+++ b/Engine/Plugins/Experimental/MLDeformer/MLDeformerFramework/Source/MLDeformerFrameworkEditor/Public/MLDeformerMorphModelVizSettingsDetails.h
@@ -0,0 +1,30 @@
+// Copyright Epic Games, Inc. All Rights Reserved.
+
+#pragma once
+
+#include "CoreMinimal.h"
+#include "MLDeformerGeomCacheVizSettingsDetails.h"
+#include "IDetailCustomization.h"
+
+class IDetailLayoutBuilder;
+class UMLDeformerMorphModel;
+class UMLDeformerMorphModelVizSettings;
+
+namespace UE::MLDeformer
+{
+	class MLDEFORMERFRAMEWORKEDITOR_API FMLDeformerMorphModelVizSettingsDetails
+		: public FMLDeformerGeomCacheVizSettingsDetails
+	{
+	public:
+		// FMLDeformerVizSettingsDetails overrides.
+		virtual bool UpdateMemberPointers(const TArray<TWeakObjectPtr<UObject>>& Objects) override;
+		virtual void AddAdditionalSettings() override;
+		// ~END FMLDeformerVizSettingsDetails overrides.
+
+	protected:
+		bool IsMorphTargetsEnabled() const;
+
+		TObjectPtr<UMLDeformerMorphModel> MorphModel = nullptr;
+		TObjectPtr<UMLDeformerMorphModelVizSettings> MorphModelVizSettings = nullptr;
+	};
+}	// namespace UE::MLDeformer
diff --git a/Engine/Plugins/Experimental/MLDeformer/MLDeformerFramework/Source/MLDeformerFrameworkEditor/Public/MLDeformerVizSettingsDetails.h b/Engine/Plugins/Experimental/MLDeformer/MLDeformerFramework/Source/MLDeformerFrameworkEditor/Public/MLDeformerVizSettingsDetails.h
--- a/Engine/Plugins/Experimental/MLDeformer/MLDeformerFramework/Source/MLDeformerFrameworkEditor/Public/MLDeformerVizSettingsDetails.h
+++ b/Engine/Plugins/Experimental/MLDeformer/MLDeformerFramework/Source/MLDeformerFrameworkEditor/Public/MLDeformerVizSettingsDetails.h
@@ -21,9 +21,6 @@ namespace UE::MLDeformer
 		: public IDetailCustomization
 	{
 	public:
-		/** Makes a new instance of this detail layout class for a specific detail view requesting it. */
-		static TSharedRef<IDetailCustomization> MakeInstance();
-
 		// ILayoutDetails overrides.
 		virtual void CustomizeDetails(IDetailLayoutBuilder& DetailBuilder) override;
 		// ~END ILayoutDetails overrides.
diff --git a/Engine/Plugins/Experimental/MLDeformer/NeuralMorphModel/Source/NeuralMorphModel/Private/NeuralMorphModel.cpp b/Engine/Plugins/Experimental/MLDeformer/NeuralMorphModel/Source/NeuralMorphModel/Private/NeuralMorphModel.cpp
--- a/Engine/Plugins/Experimental/MLDeformer/NeuralMorphModel/Source/NeuralMorphModel/Private/NeuralMorphModel.cpp
+++ b/Engine/Plugins/Experimental/MLDeformer/NeuralMorphModel/Source/NeuralMorphModel/Private/NeuralMorphModel.cpp
@@ -1,198 +1,36 @@
 // Copyright Epic Games, Inc. All Rights Reserved.
 #include "NeuralMorphModel.h"
 #include "NeuralMorphModelVizSettings.h"
-#include "NeuralMorphModelInstance.h"
-#include "MLDeformerGeomCacheHelpers.h"
-#include "MLDeformerModelInstance.h"
-#include "MLDeformerComponent.h"
-#include "UObject/Object.h"
-#include "UObject/UObjectGlobals.h"
-#include "Components/SkinnedMeshComponent.h"
-#include "Components/SkeletalMeshComponent.h"
-#include "GeometryCache.h"
-#include "Rendering/MorphTargetVertexInfoBuffers.h"
 
 #define LOCTEXT_NAMESPACE "NeuralMorphModel"
 
-// The morph target set ID for this model. This has to be unique for every different model.
-int32 UNeuralMorphModel::NeuralMorphsExternalMorphSetID = 0;
-
-UNeuralMorphModel::UNeuralMorphModel(const FObjectInitializer& ObjectInitializer)
-	: Super(ObjectInitializer)
+// Implement our module.
+namespace UE::NeuralMorphModel
 {
-#if WITH_EDITORONLY_DATA
-	VizSettings = ObjectInitializer.CreateEditorOnlyDefaultSubobject<UNeuralMorphModelVizSettings>(this, TEXT("VizSettings"));
-#endif
-
-	MorphTargetSet = MakeShared<FExternalMorphSet>();
-	MorphTargetSet->Name = FName(TEXT("NeuralBlendShapes"));
-}
-
-void UNeuralMorphModel::Serialize(FArchive& Archive)
-{
-	TRACE_CPUPROFILER_EVENT_SCOPE(UNeuralMorphModel::Serialize)
-
-	Super::Serialize(Archive);
-
-	// Check if we have initialized our compressed morph buffers.
-	bool bHasMorphData = false;
-	if (Archive.IsSaving())
-	{
-		bHasMorphData = MorphTargetSet.IsValid() ? MorphTargetSet->MorphBuffers.IsMorphCPUDataValid() : false;
-	}
-	Archive << bHasMorphData;
-
-	// Load or save the compressed morph buffers, if they exist.
-	if (bHasMorphData)
+	class NEURALMORPHMODEL_API FNeuralMorphModelModule
+		: public IModuleInterface
 	{
-		check(MorphTargetSet.IsValid());
-		Archive << MorphTargetSet->MorphBuffers;
-	}
-}
-
-UMLDeformerModelInstance* UNeuralMorphModel::CreateModelInstance(UMLDeformerComponent* Component)
-{
-	return NewObject<UNeuralMorphModelInstance>(Component);
+	};
 }
+IMPLEMENT_MODULE(UE::NeuralMorphModel::FNeuralMorphModelModule, NeuralMorphModel)
 
-#if WITH_EDITOR
-	void UNeuralMorphModel::UpdateNumTargetMeshVertices()
-	{
-		NumTargetMeshVerts = UE::MLDeformer::ExtractNumImportedGeomCacheVertices(GeometryCache);
-	}
-
-	void UNeuralMorphModel::SetAssetEditorOnlyFlags()
-	{
-		// Set the flags for the base class, which filters out the training anim sequence.
-		UMLDeformerModel::SetAssetEditorOnlyFlags();
+// Our log category for this model.
+NEURALMORPHMODEL_API DEFINE_LOG_CATEGORY(LogNeuralMorphModel)
 
-		// The training geometry cache is something we don't want to package.
-		if (GeometryCache)
-		{
-			GeometryCache->GetPackage()->SetPackageFlags(PKG_EditorOnly);
-		}
-
-		// Filter the viz settings specific assets.
-		UNeuralMorphModelVizSettings* NeuralVizSettings = Cast<UNeuralMorphModelVizSettings>(VizSettings);
-		if (NeuralVizSettings)
-		{
-			if (NeuralVizSettings->GetTestGroundTruth())
-			{
-				NeuralVizSettings->GetTestGroundTruth()->GetPackage()->SetPackageFlags(PKG_EditorOnly);
-			}
-		}
-	}
-#endif // WITH_EDITOR
+//////////////////////////////////////////////////////////////////////////////
 
+UNeuralMorphModel::UNeuralMorphModel(const FObjectInitializer& ObjectInitializer)
+	: Super(ObjectInitializer)
+{
+	// Create the visualization settings for this model.
+	// Never directly create one of the frameworks base classes such as the FMLDeformerMorphModelVizSettings as
+	// that can cause issues with detail customizations.
 #if WITH_EDITORONLY_DATA
-	void UNeuralMorphModel::SampleGroundTruthPositions(float SampleTime, TArray<FVector3f>& OutPositions)
-	{
-		UNeuralMorphModelVizSettings* VertexVizSettings = Cast<UNeuralMorphModelVizSettings>(VizSettings);
-		check(VertexVizSettings);
-
-		UGeometryCache* GeomCache = VertexVizSettings->GetTestGroundTruth();
-		if (GeomCache == nullptr)
-		{
-			OutPositions.Reset();
-			return;
-		}
-
-		if (MeshMappings.IsEmpty())
-		{
-			TArray<FString> FailedImportedMeshNames;
-			TArray<FString> VertexMisMatchNames;
-			GenerateGeomCacheMeshMappings(SkeletalMesh, GeomCache, MeshMappings, FailedImportedMeshNames, VertexMisMatchNames);
-		}
-
-		UE::MLDeformer::SampleGeomCachePositions(
-			0,
-			SampleTime,
-			MeshMappings,
-			SkeletalMesh,
-			GeomCache,
-			AlignmentTransform,
-			OutPositions);
-	}
+	VizSettings = ObjectInitializer.CreateEditorOnlyDefaultSubobject<UNeuralMorphModelVizSettings>(this, TEXT("VizSettings"));
 #endif
 
-void UNeuralMorphModel::PostMLDeformerComponentInit(UMLDeformerModelInstance* ModelInstance)
-{
-	if (ModelInstance->HasPostInitialized())
-	{
-		return;
-	}
-
-	TRACE_CPUPROFILER_EVENT_SCOPE(UNeuralMorphModel::PostMLDeformerComponentInit)
-
-	Super::PostMLDeformerComponentInit(ModelInstance);
-
-	// Register the external morph targets buffer to the render data of the skeletal mesh.
-	USkeletalMeshComponent* SkelMeshComponent = ModelInstance->GetSkeletalMeshComponent();
-	if (SkelMeshComponent && SkelMeshComponent->GetSkeletalMeshAsset())
-	{
-		// Register the morph set. This overwrites the existing one for this model, if it already exists.
-		// Only add to LOD 0 for now.
-		const int32 LOD = 0;
-		SkelMeshComponent->AddExternalMorphSet(LOD, UNeuralMorphModel::NeuralMorphsExternalMorphSetID, MorphTargetSet);
-
-		// When we're in editor mode, keep the CPU data around, so we can re-initialize when needed.
-		#if WITH_EDITOR
-			MorphTargetSet->MorphBuffers.SetEmptyMorphCPUDataOnInitRHI(false);
-		#else
-			MorphTargetSet->MorphBuffers.SetEmptyMorphCPUDataOnInitRHI(true);
-		#endif
-
-		// Release the render resources, but only in an editor build.
-		// The non-editor build shouldn't do this, as then it can't initialize again. The non-editor build assumes
-		// that the data doesn't change and we don't need to re-init.
-		// In the editor build we have to re-initialize the render resources as the morph targets can change after (re)training, so
-		// that is why we release them here, and intialize them again after.
-		FMorphTargetVertexInfoBuffers& MorphBuffers = MorphTargetSet->MorphBuffers;
-		#if WITH_EDITOR
-			BeginReleaseResource(&MorphBuffers);
-		#endif
-
-		// Reinitialize the GPU compressed buffers.
-		if (MorphBuffers.IsMorphCPUDataValid() && MorphBuffers.GetNumMorphs() > 0)
-		{
-			// In a non-editor build this will clear the CPU data.
-			// That also means it can't re-init the resources later on again.
-			BeginInitResource(&MorphBuffers);
-		}
-
-		// Update the weight information in the Skeletal Mesh.
-		SkelMeshComponent->RefreshExternalMorphTargetWeights();
-
-		ModelInstance->SetHasPostInitialized(true);
-	}
-}
-
-void UNeuralMorphModel::SetMorphTargetDeltas(const TArray<float>& Deltas)
-{
-	FloatArrayToVector3Array(Deltas, MorphTargetDeltas);
-}
-
-int32 UNeuralMorphModel::GetMorphTargetDeltaStartIndex(int32 BlendShapeIndex) const
-{
-	if (MorphTargetDeltas.Num() == 0)
-	{
-		return INDEX_NONE;
-	}
-
-	return GetNumBaseMeshVerts() * BlendShapeIndex;
-}
-
-void UNeuralMorphModel::BeginDestroy()
-{
-	if (MorphTargetSet.IsValid())
-	{
-		// Release and flush, waiting for the release to have completed, 
-		// If we don't do this we can get an error that we destroy a render resource that is still initialized,
-		// as the release happens in another thread.
-		ReleaseResourceAndFlush(&MorphTargetSet->MorphBuffers);
-		MorphTargetSet.Reset();
-	}
-	Super::BeginDestroy();
+	// The ID for the morph target data generated by this model.
+	ExternalMorphSetID = 0;
 }
 
 #undef LOCTEXT_NAMESPACE
diff --git a/Engine/Plugins/Experimental/MLDeformer/NeuralMorphModel/Source/NeuralMorphModel/Private/NeuralMorphModelInstance.cpp b/Engine/Plugins/Experimental/MLDeformer/NeuralMorphModel/Source/NeuralMorphModel/Private/NeuralMorphModelInstance.cpp
deleted file mode 100644
diff --git a/Engine/Plugins/Experimental/MLDeformer/NeuralMorphModel/Source/NeuralMorphModel/Private/NeuralMorphModelModule.cpp b/Engine/Plugins/Experimental/MLDeformer/NeuralMorphModel/Source/NeuralMorphModel/Private/NeuralMorphModelModule.cpp
deleted file mode 100644
diff --git a/Engine/Plugins/Experimental/MLDeformer/NeuralMorphModel/Source/NeuralMorphModel/Public/NeuralMorphModel.h b/Engine/Plugins/Experimental/MLDeformer/NeuralMorphModel/Source/NeuralMorphModel/Public/NeuralMorphModel.h
--- a/Engine/Plugins/Experimental/MLDeformer/NeuralMorphModel/Source/NeuralMorphModel/Public/NeuralMorphModel.h
+++ b/Engine/Plugins/Experimental/MLDeformer/NeuralMorphModel/Source/NeuralMorphModel/Public/NeuralMorphModel.h
@@ -2,7 +2,7 @@
 #pragma once
 
 #include "CoreMinimal.h"
-#include "MLDeformerModel.h"
+#include "MLDeformerMorphModel.h"
 #include "MLDeformerVizSettings.h"
 #include "MLDeformerInputInfo.h"
 #include "MLDeformerGeomCacheHelpers.h"
@@ -19,11 +19,21 @@ class USkeleton;
 class IPropertyHandle;
 struct FExternalMorphSet;
 
+NEURALMORPHMODEL_API DECLARE_LOG_CATEGORY_EXTERN(LogNeuralMorphModel, Log, All);
+
+/** 
+ * The mode of the model, either local or global. 
+ * In local mode the network contains a super simple neural network for each bone, while in 
+ * global mode all bones and curves are input to one larger fully connected network.
+ * The local mode has faster performance, while global mode can result in higher quality deformation.
+ * This model runs its neural network on the CPU, but uses comrpessed GPU based morph targets, which require shader model 5.
+ */
 UENUM()
 enum class ENeuralMorphMode : uint8
 {
 	/**
 	 * Each bone creates a set of morph targets and has its own small neural network.
+	 * The local mode can also create more localized morph targets and tends to use slightly less memory.
 	 * This mode is faster to process on the CPU side.
 	 */
 	Local,
@@ -31,14 +41,15 @@ enum class ENeuralMorphMode : uint8
 	/** 
 	 * There is one fully connected neural network that generates a set of morph targets.
 	 * This has a slightly higher CPU overhead, but could result in higher quality.
+	 * The Global mode is basically the same as the Vertex Delta Model, but runs the neural network on the CPU
+	 * and uses GPU compressed morph targets.
 	 */
 	Global
 };
 
-
 UCLASS()
 class NEURALMORPHMODEL_API UNeuralMorphModel 
-	: public UMLDeformerModel
+	: public UMLDeformerMorphModel
 {
 	GENERATED_BODY()
 
@@ -47,27 +58,9 @@ public:
 
 	// UMLDeformerModel overrides.
 	virtual FString GetDisplayName() const override { return "Neural Morph Model"; }
-	virtual void Serialize(FArchive& Archive) override;
-	virtual void PostMLDeformerComponentInit(UMLDeformerModelInstance* ModelInstance) override;
-	virtual bool IsNeuralNetworkOnGPU() const override { return false; }	// CPU neural network.
-	virtual UMLDeformerModelInstance* CreateModelInstance(UMLDeformerComponent* Component) override;
-#if WITH_EDITORONLY_DATA
-	virtual bool HasTrainingGroundTruth() const override { return (GeometryCache.Get() != nullptr); }
-	virtual void SampleGroundTruthPositions(float SampleTime, TArray<FVector3f>& OutPositions) override;
-#endif
-#if WITH_EDITOR
-	virtual void UpdateNumTargetMeshVertices() override;
-	virtual void SetAssetEditorOnlyFlags() override;
-#endif
 	// ~END UMLDeformerModel overrides.
 
-	// UObject overrides.
-	void BeginDestroy() override;
-	// ~END UObject overrides.
-
 #if WITH_EDITORONLY_DATA
-	const UGeometryCache* GetGeometryCache() const { return GeometryCache; }
-	UGeometryCache* GetGeometryCache() { return GeometryCache; }
 	int32 GetLocalNumHiddenLayers() const { return LocalNumHiddenLayers; }
 	int32 GetLocalNumNeuronsPerLayer() const { return LocalNumNeuronsPerLayer; }
 	int32 GetGlobalNumHiddenLayers() const { return GlobalNumHiddenLayers; }
@@ -76,59 +69,53 @@ public:
 	int32 GetBatchSize() const { return BatchSize; }
 	float GetLearningRate() const { return LearningRate; }
 	float GetRegularizationFactor() const { return RegularizationFactor;  }
-	float GetMorphTargetDeltaThreshold() const { return MorphTargetDeltaThreshold; }
-	TArray<UE::MLDeformer::FMLDeformerGeomCacheMeshMapping>& GetGeomCacheMeshMappings() { return MeshMappings; }
-	const TArray<UE::MLDeformer::FMLDeformerGeomCacheMeshMapping>& GetGeomCacheMeshMappings() const { return MeshMappings; }
-#endif
-
-	UFUNCTION(BlueprintCallable, Category = "NeuralMorphModel")
-	void SetMorphTargetDeltas(const TArray<float>& Deltas);
-
-	const TArray<FVector3f>& GetMorphTargetDeltas() const { return MorphTargetDeltas; }
-	int32 GetMorphTargetDeltaStartIndex(int32 MorphTargetIndex) const;
 
 public:
-	/** The compressed morph target data, ready for the GPU. */
-	TSharedPtr<FExternalMorphSet> MorphTargetSet;
-
-	/** The morph target set ID that is passed to FSkeletalMeshLODRenderData::AddExternalMorphBuffer(...). */
-	static int32 NeuralMorphsExternalMorphSetID;
-
-	/** 
-	 * The entire set of morph target deltas, 3 per vertex, for each morph target, as one flattened buffer. 
-	 * So the size of this buffer is: (NumVertsPerMorphTarget * 3 * NumMorphTargets).
+	/**
+	 * The mode that the neural network will operate in. 
+	 * Local mode means there is one tiny network per bone, while global mode has one network for all bones together.
+	 * The advantage of local mode is that it has higher performance, while global mode might result in better deformations.
 	 */
-	TArray<FVector3f> MorphTargetDeltas;
-
-#if WITH_EDITORONLY_DATA
-	TArray<UE::MLDeformer::FMLDeformerGeomCacheMeshMapping> MeshMappings;
-
-	/** The geometry cache that represents the complex mesh deformations. */
-	UPROPERTY(EditAnywhere, Category = "Target Mesh")
-	TObjectPtr<UGeometryCache> GeometryCache = nullptr;
-
-	/** The number of morph targets to generate per bone. */
 	UPROPERTY(EditAnywhere, BlueprintReadWrite, Category = "Training Settings")
 	ENeuralMorphMode Mode = ENeuralMorphMode::Local;
 
-	/** The number of morph targets to generate per bone. */
+	/** 
+	 * The number of morph targets to generate per bone.
+	 * Higher numbers result in better approximation of the target deformation, but also result in a higher memory footprint and slower performance.
+	 */
 	UPROPERTY(EditAnywhere, BlueprintReadWrite, DisplayName = "Num Morph Targets Per Bone", Category = "Training Settings", meta = (ClampMin = "1", ClampMax = "1000"))
 	int32 LocalNumMorphTargetsPerBone = 6;
 
-	/** The number of morph targets to generate. */
+	/** 
+	 * The number of morph targets to generate in total. 
+	 * Higher numbers result in better approximation of the target deformation, but also result in a higher memory footprint and slower performance.
+	 */
 	UPROPERTY(EditAnywhere, BlueprintReadWrite, DisplayName = "Num Morph Targets", Category = "Training Settings", meta = (ClampMin = "1", ClampMax = "1000"))
 	int32 GlobalNumMorphTargets = 128;
 
-	/** The number of iterations to train the model for. */
-	UPROPERTY(EditAnywhere, BlueprintReadWrite, Category = "Training Settings", meta = (ClampMin = "1", ClampMax = "500000"))
+	/** 
+	 * The number of iterations to train the model for. 
+	 * If you are quickly iterating then around 1000 to 3000 iterations should be enough.
+	 * If you want to generate final assets you might want to use a higher number of iterations, like 10k to 100k.
+	 * Once the loss doesn't go down anymore, you know that more iterations most likely won't help much.
+	 */
+	UPROPERTY(EditAnywhere, BlueprintReadWrite, Category = "Training Settings", meta = (ClampMin = "1", ClampMax = "1000000"))
 	int32 NumIterations = 2000;
 
-	/** The number of hidden layers that the neural network model will have.\nHigher numbers will slow down performance but can deal with more complex deformations. */
+	/** 
+	 * The number of hidden layers that the neural network model will have.
+	 * Higher numbers will slow down performance but can deal with more complex deformations. 
+	 * For the local model you most likely want to stick with a value of one or two.
+	 */
 	UPROPERTY(EditAnywhere, BlueprintReadWrite, AdvancedDisplay, DisplayName = "Num Hidden Layers", Category = "Training Settings", meta = (ClampMin = "1", ClampMax = "5"))
 	int32 LocalNumHiddenLayers = 1;
 
-	/** The number of units/neurons per hidden layer. Higher numbers will slow down performance but allow for more complex mesh deformations. */
-	UPROPERTY(EditAnywhere, BlueprintReadWrite, AdvancedDisplay, DisplayName = "Num Neurons Per Layer", Category = "Training Settings", meta = (ClampMin = "1"))
+	/** 
+	 * The number of units/neurons per hidden layer. 
+	 * Higher numbers will slow down performance but allow for more complex mesh deformations. 
+	 * For the local mode you probably want to keep this around the same value as the number of morph targets per bone.
+	 */
+	UPROPERTY(EditAnywhere, BlueprintReadWrite, AdvancedDisplay, DisplayName = "Num Neurons Per Layer", Category = "Training Settings", meta = (ClampMin = "1", ClampMax = "100"))
 	int32 LocalNumNeuronsPerLayer = 6;
 
 	/** The number of hidden layers that the neural network model will have.\nHigher numbers will slow down performance but can deal with more complex deformations. */
@@ -136,7 +123,7 @@ public:
 	int32 GlobalNumHiddenLayers = 2;
 
 	/** The number of units/neurons per hidden layer. Higher numbers will slow down performance but allow for more complex mesh deformations. */
-	UPROPERTY(EditAnywhere, BlueprintReadWrite, AdvancedDisplay, DisplayName = "Num Neurons Per Layer", Category = "Training Settings", meta = (ClampMin = "1"))
+	UPROPERTY(EditAnywhere, BlueprintReadWrite, AdvancedDisplay, DisplayName = "Num Neurons Per Layer", Category = "Training Settings", meta = (ClampMin = "1", ClampMax = "4096"))
 	int32 GlobalNumNeuronsPerLayer = 128;
 
 	/** The number of frames per batch when training the model. */
@@ -151,19 +138,7 @@ public:
 	 * The regularization factor. Higher values can help generate more sparse morph targets, but can also lead to visual artifacts. 
 	 * A value of 0 disables the regularization, and gives the highest quality, at the cost of higher runtime memory usage.
 	 */
-	UPROPERTY(EditAnywhere, BlueprintReadWrite, AdvancedDisplay, Category = "Training Settings", meta = (ClampMin = "0.0", ClampMax = "100.0"))
+	UPROPERTY(EditAnywhere, BlueprintReadWrite, AdvancedDisplay, Category = "Training Settings", meta = (ClampMin = "0.0", ClampMax = "10.0"))
 	float RegularizationFactor = 1.0f;
-
-	/** 
-	 * Morph target delta values that are smaller than or equal to this threshold will be zeroed out.
-	 * This essentially removes small deltas from morph targets, which will lower the memory usage at runtime, however when set too high it can also introduce visual artifacts.
-	 * A value of 0 will result in the highest quality morph targets, at the cost of higher runtime memory usage.
-	 */
-	UPROPERTY(EditAnywhere, BlueprintReadWrite, AdvancedDisplay, Category = "Training Settings", meta = (ClampMin = "0.0", ClampMax = "1.0"))
-	float MorphTargetDeltaThreshold = 0.0025f;
-
-	/** The morph target error tolerance. Higher values result in larger compression, but could result in visual artifacts. */
-	UPROPERTY(EditAnywhere, BlueprintReadWrite, AdvancedDisplay, Category = "Training Settings", meta = (ClampMin = "0.01", ClampMax = "1000"))
-	float MorphTargetErrorTolerance = 50.0f;
 #endif // WITH_EDITORONLY_DATA
 };
diff --git a/Engine/Plugins/Experimental/MLDeformer/NeuralMorphModel/Source/NeuralMorphModel/Public/NeuralMorphModelInstance.h b/Engine/Plugins/Experimental/MLDeformer/NeuralMorphModel/Source/NeuralMorphModel/Public/NeuralMorphModelInstance.h
deleted file mode 100644
diff --git a/Engine/Plugins/Experimental/MLDeformer/NeuralMorphModel/Source/NeuralMorphModel/Public/NeuralMorphModelModule.h b/Engine/Plugins/Experimental/MLDeformer/NeuralMorphModel/Source/NeuralMorphModel/Public/NeuralMorphModelModule.h
deleted file mode 100644
diff --git a/Engine/Plugins/Experimental/MLDeformer/NeuralMorphModel/Source/NeuralMorphModel/Public/NeuralMorphModelVizSettings.h b/Engine/Plugins/Experimental/MLDeformer/NeuralMorphModel/Source/NeuralMorphModel/Public/NeuralMorphModelVizSettings.h
--- a/Engine/Plugins/Experimental/MLDeformer/NeuralMorphModel/Source/NeuralMorphModel/Public/NeuralMorphModelVizSettings.h
+++ b/Engine/Plugins/Experimental/MLDeformer/NeuralMorphModel/Source/NeuralMorphModel/Public/NeuralMorphModelVizSettings.h
@@ -4,50 +4,17 @@
 
 #include "CoreMinimal.h"
 #include "UObject/Object.h"
-#include "MLDeformerVizSettings.h"
+#include "MLDeformerMorphModelVizSettings.h"
 #include "NeuralMorphModelVizSettings.generated.h"
 
-class UGeometryCache;
-
 /**
  * The vizualization settings specific to this model.
+ * Even if we have no new properties compared to the morph model, we still need to
+ * create this class in order to propertly register a detail customization for it in our editor module.
  */
 UCLASS()
 class NEURALMORPHMODEL_API UNeuralMorphModelVizSettings
-	: public UMLDeformerVizSettings
+	: public UMLDeformerMorphModelVizSettings
 {
 	GENERATED_BODY()
-
-public:
-#if WITH_EDITORONLY_DATA
-	virtual bool HasTestGroundTruth() const override { return (GroundTruth.Get() != nullptr); }
-	UGeometryCache* GetTestGroundTruth() const { return GroundTruth; }
-
-public:
-	/** The geometry cache that represents the ground truth of the test anim sequence. */
-	UPROPERTY(EditAnywhere, Category = "Test Assets")
-	TObjectPtr<UGeometryCache> GroundTruth = nullptr;
-
-	/** 
-	 * The morph target to visualize. The first one always being the means, so not a sparse target. 
-	 * This only can be used after you trained, in the same editor session directly after training.
-	 */
-	UPROPERTY(EditAnywhere, Category = "Live Settings", meta = (ClampMin = "0"))
-	int32 MorphTargetNumber = 0;
-
-	/**
-	 * The morph target delta threshold. This is a preview of what deltas would be included in the selected morph target
-	 * when using a delta threshold during training that is equal to this value.
-	 * This only can be used after you trained, in the same editor session directly after training.
-	 */
-	UPROPERTY(EditAnywhere, Category = "Live Settings", meta = (ClampMin = "0.001", ClampMax = "1.0"))
-	float MorphTargetDeltaThreshold = 0.01f;
-
-	/** 
-	 * Draw the morph targets as debug data? 
-	 * This only can be used after you trained, in the same editor session directly after training.
-	 */
-	UPROPERTY(EditAnywhere, Category = "Live Settings")
-	bool bDrawMorphTargets = false;
-#endif // WITH_EDITORONLY_DATA
 };
diff --git a/Engine/Plugins/Experimental/MLDeformer/NeuralMorphModel/Source/NeuralMorphModelEditor/Private/NeuralMorphEditorModel.cpp b/Engine/Plugins/Experimental/MLDeformer/NeuralMorphModel/Source/NeuralMorphModelEditor/Private/NeuralMorphEditorModel.cpp
--- a/Engine/Plugins/Experimental/MLDeformer/NeuralMorphModel/Source/NeuralMorphModelEditor/Private/NeuralMorphEditorModel.cpp
+++ b/Engine/Plugins/Experimental/MLDeformer/NeuralMorphModel/Source/NeuralMorphModelEditor/Private/NeuralMorphEditorModel.cpp
@@ -1,23 +1,10 @@
 // Copyright Epic Games, Inc. All Rights Reserved.
 
 #include "NeuralMorphEditorModel.h"
-#include "NeuralMorphModelModule.h"
 #include "NeuralMorphModel.h"
-#include "NeuralMorphModelVizSettings.h"
+#include "MLDeformerMorphModelVizSettings.h"
 #include "NeuralMorphTrainingModel.h"
-#include "MLDeformerGeomCacheActor.h"
 #include "MLDeformerEditorToolkit.h"
-#include "MLDeformerEditorStyle.h"
-#include "MLDeformerInputInfo.h"
-#include "MLDeformerGeomCacheSampler.h"
-#include "MLDeformerGeomCacheHelpers.h"
-#include "GeometryCache.h"
-#include "GeometryCacheComponent.h"
-#include "Animation/DebugSkelMeshComponent.h"
-#include "Animation/MorphTarget.h"
-#include "Rendering/SkeletalMeshModel.h"
-#include "Rendering/SkeletalMeshLODModel.h"
-#include "Engine/SkeletalMesh.h"
 
 #define LOCTEXT_NAMESPACE "NeuralMorphEditorModel"
 
@@ -25,29 +12,12 @@ namespace UE::NeuralMorphModel
 {
 	using namespace UE::MLDeformer;
 
-	FNeuralMorphEditorModel::~FNeuralMorphEditorModel()
-	{
-		Model->OnPostEditChangeProperty().Remove(PostEditPropertyDelegateHandle);
-	}
-
 	FMLDeformerEditorModel* FNeuralMorphEditorModel::MakeInstance()
 	{
 		return new FNeuralMorphEditorModel();
 	}
 
-	void FNeuralMorphEditorModel::AddReferencedObjects(FReferenceCollector& Collector)
-	{
-		Collector.AddReferencedObject(EditorInputInfo);
-	}
-
-	FMLDeformerSampler* FNeuralMorphEditorModel::CreateSampler() const 
-	{ 
-		FMLDeformerGeomCacheSampler* NewSampler = new FMLDeformerGeomCacheSampler();
-		NewSampler->OnGetGeometryCache().BindLambda([this] { return GetNeuralMorphModel()->GetGeometryCache(); });
-		return NewSampler;
-	}
-
-	void FNeuralMorphEditorModel::OnPostEditChangeProperty(FPropertyChangedEvent& PropertyChangedEvent)
+	void FNeuralMorphEditorModel::OnPropertyChanged(FPropertyChangedEvent& PropertyChangedEvent)
 	{
 		const FProperty* Property = PropertyChangedEvent.Property;
 		if (Property == nullptr)
@@ -55,25 +25,11 @@ namespace UE::NeuralMorphModel
 			return;
 		}
 
-		HandleDefaultPropertyChanges(PropertyChangedEvent);
+		// Process the base class property changes.
+		FMLDeformerMorphModelEditorModel::OnPropertyChanged(PropertyChangedEvent);
 
-		// When we change one of these properties below, restart animations etc.
-		if (Property->GetFName() == GET_MEMBER_NAME_CHECKED(UNeuralMorphModel, GeometryCache) || 
-			Property->GetFName() == GET_MEMBER_NAME_CHECKED(UNeuralMorphModelVizSettings, GroundTruth))
-		{
-			TriggerInputAssetChanged(true);
-		}
-		else
-		if (Property->GetFName() == GET_MEMBER_NAME_CHECKED(UNeuralMorphModel, DeltaCutoffLength) || 
-			Property->GetFName() == GET_MEMBER_NAME_CHECKED(UMLDeformerModel, AlignmentTransform))
-		{
-			if (PropertyChangedEvent.ChangeType == EPropertyChangeType::ValueSet)
-			{
-				SetResamplingInputOutputsNeeded(true);
-				SampleDeltas();
-			}
-		}
-		else if (Property->GetFName() == GET_MEMBER_NAME_CHECKED(UNeuralMorphModel, Mode))
+		// Handle property changes of this model.
+		if (Property->GetFName() == GET_MEMBER_NAME_CHECKED(UNeuralMorphModel, Mode))
 		{
 			if (PropertyChangedEvent.ChangeType == EPropertyChangeType::ValueSet)
 			{
@@ -81,306 +37,12 @@ namespace UE::NeuralMorphModel
 				GetEditor()->GetModelDetailsView()->ForceRefresh();
 			}
 		}
-		else if (Property->GetFName() == GET_MEMBER_NAME_CHECKED(UNeuralMorphModelVizSettings, MorphTargetNumber))
-		{
-			UNeuralMorphModel* NeuralModel = GetNeuralMorphModel();
-			UNeuralMorphModelVizSettings* NeuralViz = GetNeuralMorphModelVizSettings();
-			const int32 NumMorphTargets = NeuralModel->MorphTargetSet.IsValid() ? NeuralModel->MorphTargetSet->MorphBuffers.GetNumMorphs() : 0;
-			NeuralViz->MorphTargetNumber = (NumMorphTargets > 0) ? FMath::Min<int32>(NeuralViz->MorphTargetNumber, NumMorphTargets - 1) : 0;
-		}
-	}
-
-	void FNeuralMorphEditorModel::OnInputAssetsChanged()
-	{
-		// Update the skeletal mesh components of the training, test base, and ml deformed actor.
-		FMLDeformerEditorModel::OnInputAssetsChanged();
-
-		UNeuralMorphModelVizSettings* VizSettings = GetNeuralMorphModelVizSettings();
-		check(VizSettings);
-		const float TestAnimSpeed = VizSettings->GetAnimPlaySpeed();
-		UAnimSequence* TestAnimSequence = VizSettings->GetTestAnimSequence();
-
-		// Update the training geometry cache.
-		UGeometryCacheComponent* GeometryCacheComponent = FindNeuralMorphModelEditorActor(ActorID_Train_GroundTruth)->GetGeometryCacheComponent();
-		check(GeometryCacheComponent);
-		GeometryCacheComponent->SetGeometryCache(GetNeuralMorphModel()->GetGeometryCache());
-		GeometryCacheComponent->SetLooping(false);
-		GeometryCacheComponent->SetManualTick(true);
-		GeometryCacheComponent->SetPlaybackSpeed(TestAnimSpeed);
-		GeometryCacheComponent->Play();
-
-		// Update the test geometry cache (ground truth) component.
-		GeometryCacheComponent = FindNeuralMorphModelEditorActor(ActorID_Test_GroundTruth)->GetGeometryCacheComponent();
-		check(GeometryCacheComponent);
-		GeometryCacheComponent->SetGeometryCache(VizSettings->GetTestGroundTruth());
-		GeometryCacheComponent->SetLooping(true);
-		GeometryCacheComponent->SetManualTick(true);
-		GeometryCacheComponent->SetPlaybackSpeed(TestAnimSpeed);
-		GeometryCacheComponent->Play();
-
-		UNeuralMorphModel* NeuralMorphModel = GetNeuralMorphModel();
-		NeuralMorphModel->MeshMappings.Reset();
-	}
-
-	void FNeuralMorphEditorModel::CreateTrainingGroundTruthActor(UWorld* World)
-	{
-		UGeometryCache* GeomCache = GetNeuralMorphModel()->GetGeometryCache();
-		const FLinearColor LabelColor = FMLDeformerEditorStyle::Get().GetColor("MLDeformer.TargetMesh.LabelColor");
-		const FLinearColor WireframeColor = FMLDeformerEditorStyle::Get().GetColor("MLDeformer.TargetMesh.WireframeColor");
-		CreateGeomCacheActor(
-			World, 
-			ActorID_Train_GroundTruth, 
-			"Train GroundTruth", 
-			GeomCache, 
-			LabelColor, 
-			WireframeColor, 
-			LOCTEXT("TrainGroundTruthActorLabelText", "Target Mesh"),
-			true);
-	}
-
-	void FNeuralMorphEditorModel::CreateTestGroundTruthActor(UWorld* World)
-	{
-		UGeometryCache* GeomCache = GetNeuralMorphModelVizSettings()->GetTestGroundTruth();
-		const FLinearColor LabelColor = FMLDeformerEditorStyle::Get().GetColor("MLDeformer.GroundTruth.LabelColor");
-		const FLinearColor WireframeColor = FMLDeformerEditorStyle::Get().GetColor("MLDeformer.GroundTruth.WireframeColor");
-		CreateGeomCacheActor(
-			World, 
-			ActorID_Test_GroundTruth, 
-			"Test GroundTruth", 
-			GeomCache, 
-			LabelColor, 
-			WireframeColor,
-			LOCTEXT("TestGroundTruthActorLabelText", "Ground Truth"),
-			false);
-	}
-
-	void FNeuralMorphEditorModel::CreateGeomCacheActor(UWorld* World, int32 ActorID, const FName& Name, UGeometryCache* GeomCache, FLinearColor LabelColor, FLinearColor WireframeColor, const FText& LabelText, bool bIsTrainingActor)
-	{
-		FActorSpawnParameters SpawnParams;
-		SpawnParams.Name = MakeUniqueObjectName(World, AActor::StaticClass(), Name);
-		AActor* Actor = World->SpawnActor<AActor>(SpawnParams);
-		Actor->SetFlags(RF_Transient);
-
-		// Create the Geometry Cache Component.
-		UGeometryCacheComponent* GeomCacheComponent = NewObject<UGeometryCacheComponent>(Actor);
-		GeomCacheComponent->SetGeometryCache(GeomCache);
-		GeomCacheComponent->RegisterComponent();
-		GeomCacheComponent->SetOverrideWireframeColor(true);
-		GeomCacheComponent->SetWireframeOverrideColor(WireframeColor);
-		GeomCacheComponent->MarkRenderStateDirty();
-		GeomCacheComponent->SetVisibility(false);
-		Actor->SetRootComponent(GeomCacheComponent);
-
-		// Create the editor actor.
-		FMLDeformerEditorActor::FConstructSettings Settings;
-		Settings.Actor = Actor;
-		Settings.TypeID = ActorID;
-		Settings.LabelColor = LabelColor;
-		Settings.LabelText = LabelText;
-		Settings.bIsTrainingActor = bIsTrainingActor;
-		FMLDeformerGeomCacheActor* EditorActor = static_cast<FMLDeformerGeomCacheActor*>(CreateEditorActor(Settings));
-		EditorActor->SetGeometryCacheComponent(GeomCacheComponent);
-		EditorActors.Add(EditorActor);
-	} 
-
-	FMLDeformerEditorActor* FNeuralMorphEditorModel::CreateEditorActor(const FMLDeformerEditorActor::FConstructSettings& Settings) const
-	{
-		return new FMLDeformerGeomCacheActor(Settings);
-	}
-
-	FMLDeformerGeomCacheActor* FNeuralMorphEditorModel::FindNeuralMorphModelEditorActor(int32 TypeID) const
-	{
-		return static_cast<FMLDeformerGeomCacheActor*>(FindEditorActor(TypeID));
-	}
-
-	UNeuralMorphModelVizSettings* FNeuralMorphEditorModel::GetNeuralMorphModelVizSettings() const
-	{
-		return Cast<UNeuralMorphModelVizSettings>(GetNeuralMorphModel()->GetVizSettings());
-	}
-
-	double FNeuralMorphEditorModel::GetTrainingTimeAtFrame(int32 FrameNumber) const
-	{
-		// Try to get the frame from the geometry cache.
-		FMLDeformerGeomCacheActor* EditorActor = static_cast<FMLDeformerGeomCacheActor*>(FindEditorActor(ActorID_Train_GroundTruth));
-		if (EditorActor && EditorActor->GetGeometryCacheComponent() && EditorActor->GetGeometryCacheComponent()->GeometryCache.Get())
-		{
-			return EditorActor->GetGeometryCacheComponent()->GetTimeAtFrame(FrameNumber);
-		}
-
-		return FMLDeformerEditorModel::GetTrainingTimeAtFrame(FrameNumber);
-	}
-
-	int32 FNeuralMorphEditorModel::GetTrainingFrameAtTime(double TimeInSeconds) const
-	{
-		FMLDeformerGeomCacheActor* EditorActor = static_cast<FMLDeformerGeomCacheActor*>(FindEditorActor(ActorID_Train_GroundTruth));
-		if (EditorActor && EditorActor->GetGeometryCacheComponent() && EditorActor->GetGeometryCacheComponent()->GeometryCache.Get())
-		{
-			return EditorActor->GetGeometryCacheComponent()->GetFrameAtTime(TimeInSeconds);
-		}
-
-		return FMLDeformerEditorModel::GetTrainingFrameAtTime(TimeInSeconds);
-	}
-
-	double FNeuralMorphEditorModel::GetTestTimeAtFrame(int32 FrameNumber) const
-	{
-		// Try to get the frame from the geometry cache.
-		FMLDeformerGeomCacheActor* EditorActor = static_cast<FMLDeformerGeomCacheActor*>(FindEditorActor(ActorID_Test_GroundTruth));
-		if (EditorActor && EditorActor->GetGeometryCacheComponent() && EditorActor->GetGeometryCacheComponent()->GeometryCache.Get())
-		{
-			return EditorActor->GetGeometryCacheComponent()->GetTimeAtFrame(FrameNumber);
-		}
-
-		return FMLDeformerEditorModel::GetTestTimeAtFrame(FrameNumber);
-	}
-
-	int32 FNeuralMorphEditorModel::GetTestFrameAtTime(double TimeInSeconds) const
-	{
-		FMLDeformerGeomCacheActor* EditorActor = static_cast<FMLDeformerGeomCacheActor*>(FindEditorActor(ActorID_Test_GroundTruth));
-		if (EditorActor && EditorActor->GetGeometryCacheComponent() && EditorActor->GetGeometryCacheComponent()->GeometryCache.Get())
-		{
-			return EditorActor->GetGeometryCacheComponent()->GetFrameAtTime(TimeInSeconds);
-		}
-
-		return FMLDeformerEditorModel::GetTestFrameAtTime(TimeInSeconds);
-	}
-
-	int32 FNeuralMorphEditorModel::GetNumTrainingFrames() const
-	{
-		const UGeometryCache* GeometryCache = GetNeuralMorphModel()->GetGeometryCache();
-		if (GeometryCache == nullptr)
-		{
-			return 0;
-		}
-		const int32 StartFrame = GeometryCache->GetStartFrame();
-		const int32 EndFrame = GeometryCache->GetEndFrame();
-		return (EndFrame - StartFrame) + 1;
-	}
-
-	void FNeuralMorphEditorModel::UpdateIsReadyForTrainingState()
-	{
-		bIsReadyForTraining = false;
-
-		// Do some basic checks first, like if there is a skeletal mesh, ground truth, anim sequence, and if there are frames.
-		if (!IsEditorReadyForTrainingBasicChecks())
-		{
-			return;
-		}
-
-		// Now make sure the assets are compatible.
-		UNeuralMorphModel* NeuralMorphModel = GetNeuralMorphModel();
-		UGeometryCache* GeomCache = NeuralMorphModel->GetGeometryCache();
-		USkeletalMesh* SkeletalMesh = NeuralMorphModel->GetSkeletalMesh();
-		if (!GetGeomCacheErrorText(SkeletalMesh, GeomCache).IsEmpty())
-		{
-			return;
-		}
-
-		// Make sure every skeletal imported mesh has some geometry track.
-		const int32 NumGeomCacheTracks = GeomCache ? GeomCache->Tracks.Num() : 0;
-		int32 NumSkelMeshes = 0;
-		if (SkeletalMesh)
-		{
-			FSkeletalMeshModel* ImportedModel = SkeletalMesh->GetImportedModel();
-			if (ImportedModel)
-			{
-				NumSkelMeshes = ImportedModel->LODModels[0].ImportedMeshInfos.Num();		
-			}
-		}
-
-		// Allow the special case where there is just one mesh and track.
-		if (NumGeomCacheTracks != 1 || NumSkelMeshes != 1)
-		{
-			if (!GetGeomCacheSampler()->GetFailedImportedMeshNames().IsEmpty())
-			{
-				return;
-			}
-		}
-
-		bIsReadyForTraining = true;
 	}
 
 	ETrainingResult FNeuralMorphEditorModel::Train()
 	{
 		return TrainModel<UNeuralMorphTrainingModel>(this);
 	}
-
-	FString FNeuralMorphEditorModel::GetTrainedNetworkOnnxFile() const
-	{
-		return FString(FPaths::ProjectIntermediateDir() + TEXT("NeuralMorphModel/NeuralMorphModel.onnx"));
-	}
-
-	FString FNeuralMorphEditorModel::GetDefaultDeformerGraphAssetPath() const
-	{
-		return FString(TEXT("/Script/OptimusCore.OptimusDeformer'/Optimus/Deformers/DG_LinearBlendSkin_Morph_Cloth_RecomputeNormals.DG_LinearBlendSkin_Morph_Cloth_RecomputeNormals'"));
-	}
-
-	FString FNeuralMorphEditorModel::GetHeatMapDeformerGraphPath() const
-	{
-		return FString(TEXT("/MLDeformerFramework/Deformers/DG_MLDeformerModel_GPUMorph_HeatMap.DG_MLDeformerModel_GPUMorph_HeatMap"));
-	}
-
-	void FNeuralMorphEditorModel::OnPreTraining()
-	{
-		// Backup the morph target deltas in case we abort training.
-		UNeuralMorphModel* NeuralModel = GetNeuralMorphModel();
-		MorphTargetDeltasBackup = NeuralModel->MorphTargetDeltas;
-	}
-
-	void FNeuralMorphEditorModel::OnPostTraining(ETrainingResult TrainingResult, bool bUsePartiallyTrainedWhenAborted)
-	{
-		// We aborted and don't want to use partially trained results, we should restore the deltas that we just overwrote after training.
-		if (TrainingResult == ETrainingResult::Aborted && !bUsePartiallyTrainedWhenAborted)
-		{
-			// Restore the blendshape backup.
-			UNeuralMorphModel* NeuralModel = GetNeuralMorphModel();
-			NeuralModel->MorphTargetDeltas = MorphTargetDeltasBackup;
-		}
-		else if (TrainingResult == ETrainingResult::Success || (TrainingResult == ETrainingResult::Aborted && bUsePartiallyTrainedWhenAborted))
-		{
-			UNeuralMorphModel* NeuralModel = GetNeuralMorphModel();
-			if (!NeuralModel->MorphTargetDeltas.IsEmpty())
-			{
-				// Set deltas with a length equal or below a given threshold to zero, for better compression.
-				TArray<FVector3f> MorphTargetDeltas = NeuralModel->MorphTargetDeltas;
-				ZeroDeltasByThreshold(MorphTargetDeltas, NeuralModel->MorphTargetDeltaThreshold);
-
-				// Build morph targets inside the engine, using the engine's compression scheme.
-				// Add one as we included the means now as extra morph target.
-				InitEngineMorphTargets(MorphTargetDeltas);
-			}
-		}
-
-		// This internally calls InitGPUData() which updates the GPU buffer with the deltas.
-		FMLDeformerEditorModel::OnPostTraining(TrainingResult, bUsePartiallyTrainedWhenAborted);
-	}
-
-	void FNeuralMorphEditorModel::InitEngineMorphTargets(const TArray<FVector3f>& Deltas)
-	{
-		const int32 LOD = 0;
-
-		// Turn the delta buffer in a set of engine morph targets.
-		UNeuralMorphModel* NeuralModel = GetNeuralMorphModel();
-		TArray<UMorphTarget*> MorphTargets;	// These will be garbage collected.
-		CreateEngineMorphTargets(MorphTargets, Deltas, TEXT("NeuralMorph_"), LOD, NeuralModel->MorphTargetDeltaThreshold);
-
-		// Now compress the morph targets to GPU friendly buffers.
-		check(NeuralModel->MorphTargetSet.IsValid());
-		FMorphTargetVertexInfoBuffers& MorphBuffers = NeuralModel->MorphTargetSet->MorphBuffers;
-		CompressEngineMorphTargets(MorphBuffers, MorphTargets, LOD, NeuralModel->MorphTargetErrorTolerance);
-	}
-
-	void FNeuralMorphEditorModel::Render(const FSceneView* View, FViewport* Viewport, FPrimitiveDrawInterface* PDI)
-	{
-		FMLDeformerEditorModel::Render(View, Viewport, PDI);
-
-		// Debug draw the selected morph target.
-		UNeuralMorphModelVizSettings* VizSettings = GetNeuralMorphModelVizSettings();
-		if (VizSettings->bDrawMorphTargets && VizSettings->GetVisualizationMode() == EMLDeformerVizMode::TestData)
-		{
-			const FVector DrawOffset = -VizSettings->GetMeshSpacingOffsetVector();
-			DrawMorphTarget(PDI, GetNeuralMorphModel()->MorphTargetDeltas, VizSettings->MorphTargetDeltaThreshold, VizSettings->MorphTargetNumber, DrawOffset);
-		}
-	}
 }	// namespace UE::NeuralMorphModel
 
 #undef LOCTEXT_NAMESPACE
diff --git a/Engine/Plugins/Experimental/MLDeformer/NeuralMorphModel/Source/NeuralMorphModelEditor/Private/NeuralMorphModelDetails.cpp b/Engine/Plugins/Experimental/MLDeformer/NeuralMorphModel/Source/NeuralMorphModelEditor/Private/NeuralMorphModelDetails.cpp
--- a/Engine/Plugins/Experimental/MLDeformer/NeuralMorphModel/Source/NeuralMorphModelEditor/Private/NeuralMorphModelDetails.cpp
+++ b/Engine/Plugins/Experimental/MLDeformer/NeuralMorphModel/Source/NeuralMorphModelEditor/Private/NeuralMorphModelDetails.cpp
@@ -1,32 +1,9 @@
 // Copyright Epic Games, Inc. All Rights Reserved.
 
 #include "NeuralMorphModelDetails.h"
-#include "NeuralMorphEditorModel.h"
 #include "NeuralMorphModel.h"
-#include "MLDeformerModule.h"
-#include "MLDeformerAsset.h"
-#include "MLDeformerModel.h"
-#include "MLDeformerEditorModel.h"
-#include "MLDeformerEditorModule.h"
-#include "MLDeformerGeomCacheHelpers.h"
-#include "NeuralNetwork.h"
-#include "GeometryCache.h"
-#include "GeometryCacheTrack.h"
-#include "Animation/AnimSequence.h"
-#include "Engine/SkeletalMesh.h"
-#include "PropertyCustomizationHelpers.h"
-#include "AssetRegistry/AssetData.h"
-#include "Modules/ModuleManager.h"
-#include "DetailLayoutBuilder.h"
 #include "DetailCategoryBuilder.h"
 #include "DetailWidgetRow.h"
-#include "IDetailsView.h"
-#include "IDetailGroup.h"
-#include "Widgets/DeclarativeSyntaxSupport.h"
-#include "Widgets/Layout/SBox.h"
-#include "Widgets/Input/SButton.h"
-#include "DetailLayoutBuilder.h"
-#include "SWarningOrErrorBox.h"
 
 #define LOCTEXT_NAMESPACE "NeuralMorphModelDetails"
 
@@ -39,99 +16,36 @@ namespace UE::NeuralMorphModel
 		return MakeShareable(new FNeuralMorphModelDetails());
 	}
 
-	bool FNeuralMorphModelDetails::UpdateMemberPointers(const TArray<TWeakObjectPtr<UObject>>& Objects)
-	{
-		if (!FMLDeformerModelDetails::UpdateMemberPointers(Objects))
-		{
-			return false;
-		}
-
-		NeuralMorphModel = Cast<UNeuralMorphModel>(Model);
-		check(NeuralMorphModel);
-		NeuralMorphEditorModel = static_cast<FNeuralMorphEditorModel*>(EditorModel);
-
-		return (NeuralMorphModel != nullptr && NeuralMorphEditorModel != nullptr);
-	}
-
 	void FNeuralMorphModelDetails::CustomizeDetails(IDetailLayoutBuilder& DetailBuilder)
 	{
 		// Create all the detail categories and add the properties of the base class.
-		FMLDeformerModelDetails::CustomizeDetails(DetailBuilder);
+		FMLDeformerMorphModelDetails::CustomizeDetails(DetailBuilder);
+
+		UNeuralMorphModel* NeuralMorphModel = Cast<UNeuralMorphModel>(Model);
+		check(NeuralMorphModel);
 
-		// Training settings.
-		SettingsCategoryBuilder->AddProperty(GET_MEMBER_NAME_CHECKED(UNeuralMorphModel, Mode));
+		TrainingSettingsCategoryBuilder->AddProperty(GET_MEMBER_NAME_CHECKED(UNeuralMorphModel, Mode), UNeuralMorphModel::StaticClass());
 
-		SettingsCategoryBuilder->AddProperty(GET_MEMBER_NAME_CHECKED(UNeuralMorphModel, LocalNumMorphTargetsPerBone))
+		// Local mode settings.
+		TrainingSettingsCategoryBuilder->AddProperty(GET_MEMBER_NAME_CHECKED(UNeuralMorphModel, LocalNumMorphTargetsPerBone), UNeuralMorphModel::StaticClass())
 			.Visibility(NeuralMorphModel->Mode == ENeuralMorphMode::Local ? EVisibility::Visible : EVisibility::Collapsed);
-		SettingsCategoryBuilder->AddProperty(GET_MEMBER_NAME_CHECKED(UNeuralMorphModel, LocalNumHiddenLayers))
+		TrainingSettingsCategoryBuilder->AddProperty(GET_MEMBER_NAME_CHECKED(UNeuralMorphModel, LocalNumHiddenLayers), UNeuralMorphModel::StaticClass())
 			.Visibility(NeuralMorphModel->Mode == ENeuralMorphMode::Local ? EVisibility::Visible : EVisibility::Collapsed);
-		SettingsCategoryBuilder->AddProperty(GET_MEMBER_NAME_CHECKED(UNeuralMorphModel, LocalNumNeuronsPerLayer))
+		TrainingSettingsCategoryBuilder->AddProperty(GET_MEMBER_NAME_CHECKED(UNeuralMorphModel, LocalNumNeuronsPerLayer), UNeuralMorphModel::StaticClass())
 			.Visibility(NeuralMorphModel->Mode == ENeuralMorphMode::Local ? EVisibility::Visible : EVisibility::Collapsed);
 
-		SettingsCategoryBuilder->AddProperty(GET_MEMBER_NAME_CHECKED(UNeuralMorphModel, GlobalNumMorphTargets))
+		// Global mode settings.
+		TrainingSettingsCategoryBuilder->AddProperty(GET_MEMBER_NAME_CHECKED(UNeuralMorphModel, GlobalNumMorphTargets), UNeuralMorphModel::StaticClass())
 			.Visibility(NeuralMorphModel->Mode == ENeuralMorphMode::Global ? EVisibility::Visible : EVisibility::Collapsed);
-		SettingsCategoryBuilder->AddProperty(GET_MEMBER_NAME_CHECKED(UNeuralMorphModel, GlobalNumHiddenLayers))
+		TrainingSettingsCategoryBuilder->AddProperty(GET_MEMBER_NAME_CHECKED(UNeuralMorphModel, GlobalNumHiddenLayers), UNeuralMorphModel::StaticClass())
 			.Visibility(NeuralMorphModel->Mode == ENeuralMorphMode::Global ? EVisibility::Visible : EVisibility::Collapsed);
-		SettingsCategoryBuilder->AddProperty(GET_MEMBER_NAME_CHECKED(UNeuralMorphModel, GlobalNumNeuronsPerLayer))
+		TrainingSettingsCategoryBuilder->AddProperty(GET_MEMBER_NAME_CHECKED(UNeuralMorphModel, GlobalNumNeuronsPerLayer), UNeuralMorphModel::StaticClass())
 			.Visibility(NeuralMorphModel->Mode == ENeuralMorphMode::Global ? EVisibility::Visible : EVisibility::Collapsed);
 
-		SettingsCategoryBuilder->AddProperty(GET_MEMBER_NAME_CHECKED(UNeuralMorphModel, NumIterations));
-		SettingsCategoryBuilder->AddProperty(GET_MEMBER_NAME_CHECKED(UNeuralMorphModel, BatchSize));
-		SettingsCategoryBuilder->AddProperty(GET_MEMBER_NAME_CHECKED(UNeuralMorphModel, LearningRate));
-		SettingsCategoryBuilder->AddProperty(GET_MEMBER_NAME_CHECKED(UNeuralMorphModel, RegularizationFactor));
-		SettingsCategoryBuilder->AddProperty(GET_MEMBER_NAME_CHECKED(UNeuralMorphModel, MorphTargetErrorTolerance));
-	}
-	
-	void FNeuralMorphModelDetails::AddAnimSequenceErrors()
-	{
-		const FText WarningText = GetGeomCacheAnimSequenceErrorText(NeuralMorphModel->GetGeometryCache(), Model->GetAnimSequence());
-		BaseMeshCategoryBuilder->AddCustomRow(FText::FromString("AnimSeqWarning"))
-			.Visibility(!WarningText.IsEmpty() ? EVisibility::Visible : EVisibility::Collapsed)
-			.WholeRowContent()
-			[
-				SNew(SBox)
-				.Padding(FMargin(0.0f, 4.0f))
-				[
-					SNew(SWarningOrErrorBox)
-					.MessageStyle(EMessageStyle::Warning)
-					.Message(WarningText)
-				]
-			];
-	}
-
-	void FNeuralMorphModelDetails::AddTargetMesh()
-	{
-		TargetMeshCategoryBuilder->AddProperty(GET_MEMBER_NAME_CHECKED(UNeuralMorphModel, GeometryCache));
-
-		const FText TargetMeshErrorText = GetGeomCacheErrorText(NeuralMorphModel->GetSkeletalMesh(), NeuralMorphModel->GetGeometryCache());
-		TargetMeshCategoryBuilder->AddCustomRow(FText::FromString("TargetMeshError"))
-			.Visibility(!TargetMeshErrorText.IsEmpty() ? EVisibility::Visible : EVisibility::Collapsed)
-			.WholeRowContent()
-			[
-				SNew(SBox)
-				.Padding(FMargin(0.0f, 4.0f))
-				[
-					SNew(SWarningOrErrorBox)
-					.MessageStyle(EMessageStyle::Error)
-					.Message(TargetMeshErrorText)
-				]
-			];
-
-		const FText ChangedErrorText = EditorModel->GetTargetAssetChangedErrorText();
-		TargetMeshCategoryBuilder->AddCustomRow(FText::FromString("TargetMeshChangedError"))
-			.Visibility(!ChangedErrorText.IsEmpty() ? EVisibility::Visible : EVisibility::Collapsed)
-			.WholeRowContent()
-			[
-				SNew(SBox)
-				.Padding(FMargin(0.0f, 4.0f))
-				[
-					SNew(SWarningOrErrorBox)
-					.MessageStyle(EMessageStyle::Error)
-					.Message(ChangedErrorText)
-				]
-			];
-
-		AddGeomCacheMeshMappingWarnings(TargetMeshCategoryBuilder, Model->GetSkeletalMesh(), NeuralMorphModel->GetGeometryCache());
+		TrainingSettingsCategoryBuilder->AddProperty(GET_MEMBER_NAME_CHECKED(UNeuralMorphModel, NumIterations), UNeuralMorphModel::StaticClass());
+		TrainingSettingsCategoryBuilder->AddProperty(GET_MEMBER_NAME_CHECKED(UNeuralMorphModel, BatchSize), UNeuralMorphModel::StaticClass());
+		TrainingSettingsCategoryBuilder->AddProperty(GET_MEMBER_NAME_CHECKED(UNeuralMorphModel, LearningRate), UNeuralMorphModel::StaticClass());
+		TrainingSettingsCategoryBuilder->AddProperty(GET_MEMBER_NAME_CHECKED(UNeuralMorphModel, RegularizationFactor), UNeuralMorphModel::StaticClass());
 	}
 }	// namespace UE::NeuralMorphModel
 
diff --git a/Engine/Plugins/Experimental/MLDeformer/NeuralMorphModel/Source/NeuralMorphModelEditor/Private/NeuralMorphModelEditorModule.cpp b/Engine/Plugins/Experimental/MLDeformer/NeuralMorphModel/Source/NeuralMorphModelEditor/Private/NeuralMorphModelEditorModule.cpp
--- a/Engine/Plugins/Experimental/MLDeformer/NeuralMorphModel/Source/NeuralMorphModelEditor/Private/NeuralMorphModelEditorModule.cpp
+++ b/Engine/Plugins/Experimental/MLDeformer/NeuralMorphModel/Source/NeuralMorphModelEditor/Private/NeuralMorphModelEditorModule.cpp
@@ -1,6 +1,5 @@
 // Copyright Epic Games, Inc. All Rights Reserved.
 
-#include "NeuralMorphModelEditorModule.h"
 #include "NeuralMorphModelVizSettingsDetails.h"
 #include "NeuralMorphModelDetails.h"
 #include "NeuralMorphEditorModel.h"
@@ -11,12 +10,25 @@
 
 #define LOCTEXT_NAMESPACE "NeuralMorphModelEditorModule"
 
+namespace UE::NeuralMorphModel
+{
+	class FNeuralMorphModelEditorModule
+		: public IModuleInterface
+	{
+	public:
+		// IModuleInterface overrides.
+		void StartupModule() override;
+		void ShutdownModule() override;
+		// ~END IModuleInterface overrides.
+	};
+}
 IMPLEMENT_MODULE(UE::NeuralMorphModel::FNeuralMorphModelEditorModule, NeuralMorphModelEditor)
 
 namespace UE::NeuralMorphModel
 {
 	using namespace UE::MLDeformer;
 
+
 	void FNeuralMorphModelEditorModule::StartupModule()
 	{
 		// Register object detail customizations.
@@ -50,7 +62,6 @@ namespace UE::NeuralMorphModel
 			PropertyModule.NotifyCustomizationModuleChanged();
 		}
 	}
-
 }	// namespace UE::NeuralMorphModel
 
 #undef LOCTEXT_NAMESPACE
diff --git a/Engine/Plugins/Experimental/MLDeformer/NeuralMorphModel/Source/NeuralMorphModelEditor/Private/NeuralMorphModelVizSettingsDetails.cpp b/Engine/Plugins/Experimental/MLDeformer/NeuralMorphModel/Source/NeuralMorphModelEditor/Private/NeuralMorphModelVizSettingsDetails.cpp
deleted file mode 100644
diff --git a/Engine/Plugins/Experimental/MLDeformer/NeuralMorphModel/Source/NeuralMorphModelEditor/Private/NeuralMorphTrainingModel.cpp b/Engine/Plugins/Experimental/MLDeformer/NeuralMorphModel/Source/NeuralMorphModelEditor/Private/NeuralMorphTrainingModel.cpp
deleted file mode 100644
diff --git a/Engine/Plugins/Experimental/MLDeformer/NeuralMorphModel/Source/NeuralMorphModelEditor/Public/NeuralMorphEditorModel.h b/Engine/Plugins/Experimental/MLDeformer/NeuralMorphModel/Source/NeuralMorphModelEditor/Public/NeuralMorphEditorModel.h
--- a/Engine/Plugins/Experimental/MLDeformer/NeuralMorphModel/Source/NeuralMorphModelEditor/Public/NeuralMorphEditorModel.h
+++ b/Engine/Plugins/Experimental/MLDeformer/NeuralMorphModel/Source/NeuralMorphModelEditor/Public/NeuralMorphEditorModel.h
@@ -3,88 +3,31 @@
 #pragma once
 
 #include "CoreMinimal.h"
-#include "MLDeformerEditorModel.h"
-#include "MLDeformerGeomCacheSampler.h"
-#include "MLDeformerEditorActor.h"
-#include "NeuralMorphModel.h"
-#include "UObject/GCObject.h"
+#include "MLDeformerMorphModelEditorModel.h"
 
-namespace UE::MLDeformer
-{
-	class FMLDeformerGeomCacheActor;
-}
-
-class UMLDeformerAsset;
-class USkeletalMesh;
-class UGeometryCache;
 class UNeuralMorphModel;
-class UNeuralMorphModelVizSettings;
 
 namespace UE::NeuralMorphModel
 {
 	using namespace UE::MLDeformer;
 
-	class FNeuralMorphEditorModelActor;
-
 	class NEURALMORPHMODELEDITOR_API FNeuralMorphEditorModel
-		: public UE::MLDeformer::FMLDeformerEditorModel
+		: public UE::MLDeformer::FMLDeformerMorphModelEditorModel
 	{
 	public:
-		virtual ~FNeuralMorphEditorModel() override;
-
 		// We need to implement this static MakeInstance method.
 		static FMLDeformerEditorModel* MakeInstance();
 
 		// FGCObject overrides.
-		virtual void AddReferencedObjects(FReferenceCollector& Collector) override;
 		virtual FString GetReferencerName() const override { return TEXT("FNeuralMorphEditorModel"); }
 		// ~END FGCObject overrides.
 	
 		// FMLDeformerEditorModel overrides.
-		virtual FMLDeformerEditorActor* CreateEditorActor(const FMLDeformerEditorActor::FConstructSettings& Settings) const override;
-		virtual FMLDeformerSampler* CreateSampler() const;
-		virtual void CreateTrainingGroundTruthActor(UWorld* World) override;
-		virtual void CreateTestGroundTruthActor(UWorld* World) override;
-		virtual void OnInputAssetsChanged() override;
-		virtual int32 GetNumTrainingFrames() const override;
-		virtual double GetTrainingTimeAtFrame(int32 FrameNumber) const override;
-		virtual int32 GetTrainingFrameAtTime(double TimeInSeconds) const override;
-		virtual double GetTestTimeAtFrame(int32 FrameNumber) const override;
-		virtual int32 GetTestFrameAtTime(double TimeInSeconds) const override;
-		virtual void UpdateIsReadyForTrainingState() override;
 		virtual ETrainingResult Train() override;
-		virtual void Render(const FSceneView* View, FViewport* Viewport, FPrimitiveDrawInterface* PDI) override;
-		virtual FString GetTrainedNetworkOnnxFile() const override;
-		virtual FString GetDefaultDeformerGraphAssetPath() const override;
-		virtual FString GetHeatMapDeformerGraphPath() const override;
-		virtual void OnPreTraining() override;
-		virtual void OnPostTraining(ETrainingResult TrainingResult, bool bUsePartiallyTrainedWhenAborted) override;
-		virtual void OnPostEditChangeProperty(FPropertyChangedEvent& PropertyChangedEvent) override;
+		virtual void OnPropertyChanged(FPropertyChangedEvent& PropertyChangedEvent);
 		// ~END FMLDeformerEditorModel overrides.
 
 		// Some helpers that cast to this model's variants of some classes.
 		UNeuralMorphModel* GetNeuralMorphModel() const { return Cast<UNeuralMorphModel>(Model); }
-		UNeuralMorphModelVizSettings* GetNeuralMorphModelVizSettings() const;
-		FMLDeformerGeomCacheActor* FindNeuralMorphModelEditorActor(int32 TypeID) const;
-		FMLDeformerGeomCacheSampler* GetGeomCacheSampler() const { return static_cast<FMLDeformerGeomCacheSampler*>(Sampler); }
-
-	protected:
-		void CreateGeomCacheActor(UWorld* World, int32 ActorID, const FName& Name, UGeometryCache* GeomCache, FLinearColor LabelColor, FLinearColor WireframeColor, const FText& LabelText, bool bIsTrainingActor);
-
-		/**
-		 * Initialize a set of engine morph targets and compress them to GPU friendly buffers.
-		 * These morph targets are initialized from a set of deltas. Each morph target needs to have Model->GetNumBaseVerts() number of deltas.
-		 * All deltas are concatenated in one big array. So all deltas of all vertices for the second morph target are appended to the deltas for 
-		 * the first morph target, etc. In other words, the layout is: [morph0_deltas][morph1_deltas][morph2_deltas][...].
-		 * @param Deltas The deltas for all morph targets concatenated. So the number of items in this array is a multiple of Model->GetNumBaseVerts().
-		 */
-		void InitEngineMorphTargets(const TArray<FVector3f>& Deltas);
-
-	protected:
-		/**
-		 * The entire set of morph target deltas, 3 per vertex, for each morph target, as one flattened buffer.
-		 * So the size of this buffer is: (NumVertsPerMorphTarget * NumMorphTargets).
-		 */
-		TArray<FVector3f> MorphTargetDeltasBackup;
 	};
 }	// namespace UE::NeuralMorphModel
diff --git a/Engine/Plugins/Experimental/MLDeformer/NeuralMorphModel/Source/NeuralMorphModelEditor/Public/NeuralMorphModelDetails.h b/Engine/Plugins/Experimental/MLDeformer/NeuralMorphModel/Source/NeuralMorphModelEditor/Public/NeuralMorphModelDetails.h
--- a/Engine/Plugins/Experimental/MLDeformer/NeuralMorphModel/Source/NeuralMorphModelEditor/Public/NeuralMorphModelDetails.h
+++ b/Engine/Plugins/Experimental/MLDeformer/NeuralMorphModel/Source/NeuralMorphModelEditor/Public/NeuralMorphModelDetails.h
@@ -3,42 +3,25 @@
 #pragma once
 
 #include "CoreMinimal.h"
-#include "MLDeformerModelDetails.h"
+#include "MLDeformerMorphModelDetails.h"
 #include "IDetailCustomization.h"
-#include "AssetRegistry/AssetData.h"
-#include "Input/Reply.h"
 
 class IDetailLayoutBuilder;
 class UNeuralMorphModel;
 
-namespace UE::MLDeformer
-{
-	class FMLDeformerEditorModel;
-}
-
 namespace UE::NeuralMorphModel
 {
 	class FNeuralMorphEditorModel;
 
 	class NEURALMORPHMODELEDITOR_API FNeuralMorphModelDetails
-		: public UE::MLDeformer::FMLDeformerModelDetails
+		: public UE::MLDeformer::FMLDeformerMorphModelDetails
 	{
 	public:
 		/** Makes a new instance of this detail layout class for a specific detail view requesting it. */
 		static TSharedRef<IDetailCustomization> MakeInstance();
 
 		// ILayoutDetails overrides.
-		void CustomizeDetails(IDetailLayoutBuilder& DetailBuilder) override;
+		virtual void CustomizeDetails(IDetailLayoutBuilder& DetailBuilder) override;
 		// ~END ILayoutDetails overrides.
-
-		// FMLDeformerModelDetails overrides.
-		bool UpdateMemberPointers(const TArray<TWeakObjectPtr<UObject>>& Objects) override;
-		void AddTargetMesh() override;
-		void AddAnimSequenceErrors() override;
-		// ~END FMLDeformerModelDetails overrides.
-
-	protected:
-		TObjectPtr<UNeuralMorphModel> NeuralMorphModel = nullptr;
-		FNeuralMorphEditorModel* NeuralMorphEditorModel = nullptr;
 	};
 }	// namespace UE::NeuralMorphModel
diff --git a/Engine/Plugins/Experimental/MLDeformer/NeuralMorphModel/Source/NeuralMorphModelEditor/Public/NeuralMorphModelEditorModule.h b/Engine/Plugins/Experimental/MLDeformer/NeuralMorphModel/Source/NeuralMorphModelEditor/Public/NeuralMorphModelEditorModule.h
deleted file mode 100644
diff --git a/Engine/Plugins/Experimental/MLDeformer/NeuralMorphModel/Source/NeuralMorphModelEditor/Public/NeuralMorphModelVizSettingsDetails.h b/Engine/Plugins/Experimental/MLDeformer/NeuralMorphModel/Source/NeuralMorphModelEditor/Public/NeuralMorphModelVizSettingsDetails.h
--- a/Engine/Plugins/Experimental/MLDeformer/NeuralMorphModel/Source/NeuralMorphModelEditor/Public/NeuralMorphModelVizSettingsDetails.h
+++ b/Engine/Plugins/Experimental/MLDeformer/NeuralMorphModel/Source/NeuralMorphModelEditor/Public/NeuralMorphModelVizSettingsDetails.h
@@ -3,35 +3,23 @@
 #pragma once
 
 #include "CoreMinimal.h"
-#include "MLDeformerVizSettingsDetails.h"
-#include "IDetailCustomization.h"
-#include "AssetRegistry/AssetData.h"
-#include "PropertyHandle.h"
-
-class IDetailLayoutBuilder;
-class USkeleton;
-class UNeuralMorphModel;
-class UNeuralMorphModelVizSettings;
+#include "MLDeformerMorphModelVizSettingsDetails.h"
 
 namespace UE::NeuralMorphModel
 {
+	/**
+	 * The neural morph model visualization settings detail customization.
+	 * This is inherited from the Morph Model detail customization.
+	 * We do not have to do anything else. We still need this class in order to do our
+	 * detail customization correctly. Just using the Morph Model directly would cause issues.
+	 */
 	class NEURALMORPHMODELEDITOR_API FNeuralMorphModelVizSettingsDetails
-		: public UE::MLDeformer::FMLDeformerVizSettingsDetails
+		: public UE::MLDeformer::FMLDeformerMorphModelVizSettingsDetails
 	{
 	public:
-		/** Makes a new instance of this detail layout class for a specific detail view requesting it. */
-		static TSharedRef<IDetailCustomization> MakeInstance();
-
-		// FMLDeformerVizSettingsDetails overrides.
-		bool UpdateMemberPointers(const TArray<TWeakObjectPtr<UObject>>& Objects) override;
-		void AddGroundTruth() override;
-		void AddAdditionalSettings() override;
-		// ~END FMLDeformerVizSettingsDetails overrides.
-
-	protected:
-		bool IsMorphTargetsEnabled() const;
-
-		TObjectPtr<UNeuralMorphModel> NeuralMorphModel = nullptr;
-		TObjectPtr<UNeuralMorphModelVizSettings> NeuralMorphVizSettings = nullptr;
+		static TSharedRef<IDetailCustomization> MakeInstance()
+		{
+			return MakeShareable(new FNeuralMorphModelVizSettingsDetails());
+		}
 	};
 }	// namespace UE::NeuralMorphModel
diff --git a/Engine/Plugins/Experimental/MLDeformer/NeuralMorphModel/Source/NeuralMorphModelEditor/Public/NeuralMorphTrainingModel.h b/Engine/Plugins/Experimental/MLDeformer/NeuralMorphModel/Source/NeuralMorphModelEditor/Public/NeuralMorphTrainingModel.h
--- a/Engine/Plugins/Experimental/MLDeformer/NeuralMorphModel/Source/NeuralMorphModelEditor/Public/NeuralMorphTrainingModel.h
+++ b/Engine/Plugins/Experimental/MLDeformer/NeuralMorphModel/Source/NeuralMorphModelEditor/Public/NeuralMorphTrainingModel.h
@@ -6,13 +6,6 @@
 #include "MLDeformerTrainingModel.h"
 #include "NeuralMorphTrainingModel.generated.h"
 
-namespace UE::NeuralMorphModel
-{
-	class FNeuralMorphEditorModel;
-}
-
-class UNeuralMorphModel;
-
 UCLASS(Blueprintable)
 class NEURALMORPHMODELEDITOR_API UNeuralMorphTrainingModel
 	: public UMLDeformerTrainingModel
@@ -23,7 +16,4 @@ public:
 	/** Main training function, with implementation in python. */
 	UFUNCTION(BlueprintImplementableEvent, Category = "Training Model")
 	int32 Train() const;
-
-	UNeuralMorphModel* GetNeuralMorphModel() const;
-	UE::NeuralMorphModel::FNeuralMorphEditorModel* GetNeuralMorphEditorModel() const;
 };
diff --git a/Engine/Plugins/Experimental/MLDeformer/VertexDeltaModel/Source/VertexDeltaModel/Private/Tests/VertexDeltaModelTest.cpp b/Engine/Plugins/Experimental/MLDeformer/VertexDeltaModel/Source/VertexDeltaModel/Private/Tests/VertexDeltaModelTest.cpp
--- a/Engine/Plugins/Experimental/MLDeformer/VertexDeltaModel/Source/VertexDeltaModel/Private/Tests/VertexDeltaModelTest.cpp
+++ b/Engine/Plugins/Experimental/MLDeformer/VertexDeltaModel/Source/VertexDeltaModel/Private/Tests/VertexDeltaModelTest.cpp
@@ -5,7 +5,6 @@
 
 #if WITH_DEV_AUTOMATION_TESTS
 
-#include "VertexDeltaModelModule.h"
 #include "VertexDeltaModel.h"
 #include "VertexDeltaModelVizSettings.h"
 #include "MLDeformerAsset.h"
@@ -27,7 +26,6 @@ namespace UE::VertexDeltaModelTests
 	bool FVertexDeltaModelMainTest::RunTest(const FString& Parameters)
 	{
 		using namespace UE::MLDeformer;
-		using namespace UE::VertexDeltaModel;
 
 		UMLDeformerAsset* DeformerAsset = LoadObject<UMLDeformerAsset>(nullptr, TEXT("MLDeformerAsset'/VertexDeltaModel/Tests/Biceps/biceps_mld.biceps_mld'"));
 		UTEST_NOT_NULL(TEXT("MLDeformerAsset load"), DeformerAsset);
diff --git a/Engine/Plugins/Experimental/MLDeformer/VertexDeltaModel/Source/VertexDeltaModel/Private/VertexDeltaModel.cpp b/Engine/Plugins/Experimental/MLDeformer/VertexDeltaModel/Source/VertexDeltaModel/Private/VertexDeltaModel.cpp
--- a/Engine/Plugins/Experimental/MLDeformer/VertexDeltaModel/Source/VertexDeltaModel/Private/VertexDeltaModel.cpp
+++ b/Engine/Plugins/Experimental/MLDeformer/VertexDeltaModel/Source/VertexDeltaModel/Private/VertexDeltaModel.cpp
@@ -2,11 +2,21 @@
 
 #include "VertexDeltaModel.h"
 #include "VertexDeltaModelVizSettings.h"
-#include "MLDeformerGeomCacheHelpers.h"
-#include "GeometryCache.h"
-#include "UObject/Object.h"
 
-#define LOCTEXT_NAMESPACE "UVertexDeltaModel"
+#define LOCTEXT_NAMESPACE "VertexDeltaModel"
+
+VERTEXDELTAMODEL_API DEFINE_LOG_CATEGORY(LogVertexDeltaModel)
+
+// Implement the module.
+namespace UE::VertexDeltaModel
+{
+	class VERTEXDELTAMODEL_API FVertexDeltaModelModule
+		: public IModuleInterface
+	{
+	};
+}
+IMPLEMENT_MODULE(UE::VertexDeltaModel::FVertexDeltaModelModule, VertexDeltaModel)
+
 
 UVertexDeltaModel::UVertexDeltaModel(const FObjectInitializer& ObjectInitializer)
 	: Super(ObjectInitializer)
@@ -16,64 +26,4 @@ UVertexDeltaModel::UVertexDeltaModel(const FObjectInitializer& ObjectInitializer
 #endif
 }
 
-#if WITH_EDITOR
-	void UVertexDeltaModel::UpdateNumTargetMeshVertices()
-	{
-		NumTargetMeshVerts = UE::MLDeformer::ExtractNumImportedGeomCacheVertices(GeometryCache);
-	}
-
-	void UVertexDeltaModel::SetAssetEditorOnlyFlags()
-	{
-		// Set the flags for the base class, which filters out the training anim sequence.
-		UMLDeformerModel::SetAssetEditorOnlyFlags();
-
-		// The training geometry cache is something we don't want to package.
-		if (GeometryCache)
-		{
-			GeometryCache->GetPackage()->SetPackageFlags(PKG_EditorOnly);
-		}
-
-		// Filter the viz settings specific assets.
-		UVertexDeltaModelVizSettings* VertexDeltaVizSettings = Cast<UVertexDeltaModelVizSettings>(VizSettings);
-		if (VertexDeltaVizSettings)
-		{
-			if (VertexDeltaVizSettings->GetTestGroundTruth())
-			{
-				VertexDeltaVizSettings->GetTestGroundTruth()->GetPackage()->SetPackageFlags(PKG_EditorOnly);
-			}
-		}
-	}
-#endif // WITH_EDITOR
-
-#if WITH_EDITORONLY_DATA
-	void UVertexDeltaModel::SampleGroundTruthPositions(float SampleTime, TArray<FVector3f>& OutPositions)
-	{
-		UVertexDeltaModelVizSettings* VertexVizSettings = Cast<UVertexDeltaModelVizSettings>(VizSettings);
-		check(VertexVizSettings);
-
-		UGeometryCache* GeomCache = VertexVizSettings->GetTestGroundTruth();
-		if (GeomCache == nullptr)
-		{
-			OutPositions.Reset();
-			return;
-		}
-
-		if (MeshMappings.IsEmpty())
-		{
-			TArray<FString> FailedImportedMeshnames;
-			TArray<FString> VertexMisMatchNames;
-			GenerateGeomCacheMeshMappings(SkeletalMesh, GeomCache, MeshMappings, FailedImportedMeshnames, VertexMisMatchNames);
-		}
-
-		UE::MLDeformer::SampleGeomCachePositions(
-			0,
-			SampleTime,
-			MeshMappings,
-			SkeletalMesh,
-			GeomCache,
-			AlignmentTransform,
-			OutPositions);
-	}
-#endif
-
 #undef LOCTEXT_NAMESPACE
diff --git a/Engine/Plugins/Experimental/MLDeformer/VertexDeltaModel/Source/VertexDeltaModel/Private/VertexDeltaModelModule.cpp b/Engine/Plugins/Experimental/MLDeformer/VertexDeltaModel/Source/VertexDeltaModel/Private/VertexDeltaModelModule.cpp
deleted file mode 100644
diff --git a/Engine/Plugins/Experimental/MLDeformer/VertexDeltaModel/Source/VertexDeltaModel/Public/VertexDeltaModel.h b/Engine/Plugins/Experimental/MLDeformer/VertexDeltaModel/Source/VertexDeltaModel/Public/VertexDeltaModel.h
--- a/Engine/Plugins/Experimental/MLDeformer/VertexDeltaModel/Source/VertexDeltaModel/Public/VertexDeltaModel.h
+++ b/Engine/Plugins/Experimental/MLDeformer/VertexDeltaModel/Source/VertexDeltaModel/Public/VertexDeltaModel.h
@@ -2,25 +2,16 @@
 #pragma once
 
 #include "CoreMinimal.h"
-#include "MLDeformerModel.h"
-#include "MLDeformerVizSettings.h"
-#include "MLDeformerInputInfo.h"
-#include "MLDeformerGeomCacheHelpers.h"
+#include "MLDeformerGeomCacheModel.h"
+#include "VertexDeltaModelVizSettings.h"
 #include "UObject/Object.h"
 #include "VertexDeltaModel.generated.h"
 
-class USkeletalMesh;
-class UGeometryCache;
-class UAnimSequence;
-class UNeuralNetwork;
-class UMLDeformerAsset;
-class USkeleton;
-class IPropertyHandle;
-
+VERTEXDELTAMODEL_API DECLARE_LOG_CATEGORY_EXTERN(LogVertexDeltaModel, Log, All);
 
 UCLASS()
 class VERTEXDELTAMODEL_API UVertexDeltaModel 
-	: public UMLDeformerModel
+	: public UMLDeformerGeomCacheModel
 {
 	GENERATED_BODY()
 
@@ -29,37 +20,17 @@ public:
 
 	// UMLDeformerModel overrides.
 	virtual FString GetDisplayName() const override { return "Vertex Delta Model"; }
-#if WITH_EDITORONLY_DATA
-	virtual bool HasTrainingGroundTruth() const override { return (GeometryCache.Get() != nullptr); }
-	virtual void SampleGroundTruthPositions(float SampleTime, TArray<FVector3f>& OutPositions) override;
-#endif
-#if WITH_EDITOR
-	virtual void UpdateNumTargetMeshVertices() override;
-	virtual void SetAssetEditorOnlyFlags() override;
-#endif
+	virtual bool IsNeuralNetworkOnGPU() const override { return true; }	// GPU neural network.
 	// ~END UMLDeformerModel overrides.
 
 #if WITH_EDITORONLY_DATA
-	const UGeometryCache* GetGeometryCache() const { return GeometryCache; }
-	UGeometryCache* GetGeometryCache() { return GeometryCache; }
 	int32 GetNumHiddenLayers() const { return NumHiddenLayers; }
 	int32 GetNumNeuronsPerLayer() const { return NumNeuronsPerLayer; }
 	int32 GetNumIterations() const { return NumIterations; }
 	int32 GetBatchSize() const { return BatchSize; }
 	float GetLearningRate() const { return LearningRate; }
-	TArray<UE::MLDeformer::FMLDeformerGeomCacheMeshMapping>& GetGeomCacheMeshMappings() { return MeshMappings; }
-	const TArray<UE::MLDeformer::FMLDeformerGeomCacheMeshMapping>& GetGeomCacheMeshMappings() const { return MeshMappings; }
-#endif
 
 public:
-
-#if WITH_EDITORONLY_DATA
-	TArray<UE::MLDeformer::FMLDeformerGeomCacheMeshMapping> MeshMappings;
-
-	/** The geometry cache that represents the complex mesh deformations. */
-	UPROPERTY(EditAnywhere, Category = "Target Mesh")
-	TObjectPtr<UGeometryCache> GeometryCache = nullptr;
-
 	/** The number of hidden layers that the neural network model will have.\nHigher numbers will slow down performance but can deal with more complex deformations. */
 	UPROPERTY(EditAnywhere, BlueprintReadWrite, Category = "Training Settings", meta = (ClampMin = "1", ClampMax = "10"))
 	int32 NumHiddenLayers = 3;
diff --git a/Engine/Plugins/Experimental/MLDeformer/VertexDeltaModel/Source/VertexDeltaModel/Public/VertexDeltaModelModule.h b/Engine/Plugins/Experimental/MLDeformer/VertexDeltaModel/Source/VertexDeltaModel/Public/VertexDeltaModelModule.h
deleted file mode 100644
diff --git a/Engine/Plugins/Experimental/MLDeformer/VertexDeltaModel/Source/VertexDeltaModel/Public/VertexDeltaModelVizSettings.h b/Engine/Plugins/Experimental/MLDeformer/VertexDeltaModel/Source/VertexDeltaModel/Public/VertexDeltaModelVizSettings.h
--- a/Engine/Plugins/Experimental/MLDeformer/VertexDeltaModel/Source/VertexDeltaModel/Public/VertexDeltaModelVizSettings.h
+++ b/Engine/Plugins/Experimental/MLDeformer/VertexDeltaModel/Source/VertexDeltaModel/Public/VertexDeltaModelVizSettings.h
@@ -3,29 +3,16 @@
 #pragma once
 
 #include "CoreMinimal.h"
+#include "MLDeformerGeomCacheVizSettings.h"
 #include "UObject/Object.h"
-#include "MLDeformerVizSettings.h"
 #include "VertexDeltaModelVizSettings.generated.h"
 
-class UGeometryCache;
-
 /**
  * The vizualization settings specific to the the vertex delta model.
  */
 UCLASS()
 class VERTEXDELTAMODEL_API UVertexDeltaModelVizSettings 
-	: public UMLDeformerVizSettings
+	: public UMLDeformerGeomCacheVizSettings
 {
 	GENERATED_BODY()
-
-public:
-#if WITH_EDITORONLY_DATA
-	virtual bool HasTestGroundTruth() const override { return (GroundTruth.Get() != nullptr); }
-	UGeometryCache* GetTestGroundTruth() const { return GroundTruth; }
-
-public:
-	/** The geometry cache that represents the ground truth of the test anim sequence. */
-	UPROPERTY(EditAnywhere, Category = "Test Assets")
-	TObjectPtr<UGeometryCache> GroundTruth = nullptr;
-#endif // WITH_EDITORONLY_DATA
 };
diff --git a/Engine/Plugins/Experimental/MLDeformer/VertexDeltaModel/Source/VertexDeltaModelEditor/Private/VertexDeltaEditorModel.cpp b/Engine/Plugins/Experimental/MLDeformer/VertexDeltaModel/Source/VertexDeltaModelEditor/Private/VertexDeltaEditorModel.cpp
--- a/Engine/Plugins/Experimental/MLDeformer/VertexDeltaModel/Source/VertexDeltaModelEditor/Private/VertexDeltaEditorModel.cpp
+++ b/Engine/Plugins/Experimental/MLDeformer/VertexDeltaModel/Source/VertexDeltaModelEditor/Private/VertexDeltaEditorModel.cpp
@@ -28,264 +28,10 @@ namespace UE::VertexDeltaModel
 		return new FVertexDeltaEditorModel();
 	}
 
-	void FVertexDeltaEditorModel::AddReferencedObjects(FReferenceCollector& Collector)
-	{
-		Collector.AddReferencedObject(EditorInputInfo);
-	}
-
-	FMLDeformerSampler* FVertexDeltaEditorModel::CreateSampler() const 
-	{ 
-		FMLDeformerGeomCacheSampler* NewSampler = new FMLDeformerGeomCacheSampler();
-		NewSampler->OnGetGeometryCache().BindLambda([this] { return GetVertexDeltaModel()->GetGeometryCache(); });
-		return NewSampler;
-	}
-
-	void FVertexDeltaEditorModel::OnPostEditChangeProperty(FPropertyChangedEvent& PropertyChangedEvent)
-	{
-		const FProperty* Property = PropertyChangedEvent.Property;
-		if (Property == nullptr)
-		{
-			return;
-		}
-
-		HandleDefaultPropertyChanges(PropertyChangedEvent);
-
-		// When we change one of these properties below, restart animations etc.
-		if (Property->GetFName() == GET_MEMBER_NAME_CHECKED(UVertexDeltaModel, GeometryCache) || 
-			Property->GetFName() == GET_MEMBER_NAME_CHECKED(UVertexDeltaModelVizSettings, GroundTruth))
-		{
-			TriggerInputAssetChanged(true);
-		}
-		else
-		if (Property->GetFName() == GET_MEMBER_NAME_CHECKED(UVertexDeltaModel, DeltaCutoffLength) || 
-			Property->GetFName() == GET_MEMBER_NAME_CHECKED(UMLDeformerModel, AlignmentTransform))
-		{
-			if (PropertyChangedEvent.ChangeType == EPropertyChangeType::ValueSet)
-			{
-				SetResamplingInputOutputsNeeded(true);
-				SampleDeltas();
-			}
-		}
-	}
-
-	void FVertexDeltaEditorModel::OnInputAssetsChanged()
-	{
-		// Update the skeletal mesh components of the training, test base, and ml deformed actor.
-		FMLDeformerEditorModel::OnInputAssetsChanged();
-
-		UVertexDeltaModelVizSettings* VizSettings = GetVertexDeltaModelVizSettings();
-		check(VizSettings);
-		const float TestAnimSpeed = VizSettings->GetAnimPlaySpeed();
-		UAnimSequence* TestAnimSequence = VizSettings->GetTestAnimSequence();
-
-		// Update the training geometry cache.
-		UGeometryCacheComponent* GeometryCacheComponent = FindVertexDeltaModelEditorActor(ActorID_Train_GroundTruth)->GetGeometryCacheComponent();
-		check(GeometryCacheComponent);
-		GeometryCacheComponent->SetGeometryCache(GetVertexDeltaModel()->GetGeometryCache());
-		GeometryCacheComponent->ResetAnimationTime();
-		GeometryCacheComponent->SetLooping(false);
-		GeometryCacheComponent->SetManualTick(true);
-		GeometryCacheComponent->SetPlaybackSpeed(TestAnimSpeed);
-		GeometryCacheComponent->Play();
-
-		// Update the test geometry cache (ground truth) component.
-		GeometryCacheComponent = FindVertexDeltaModelEditorActor(ActorID_Test_GroundTruth)->GetGeometryCacheComponent();
-		GeometryCacheComponent->SetGeometryCache(VizSettings->GetTestGroundTruth());
-		GeometryCacheComponent->ResetAnimationTime();
-		GeometryCacheComponent->SetLooping(true);
-		GeometryCacheComponent->SetManualTick(true);
-		GeometryCacheComponent->SetPlaybackSpeed(TestAnimSpeed);
-		GeometryCacheComponent->Play();
-
-		UVertexDeltaModel* VertexDeltaModel = GetVertexDeltaModel();
-		VertexDeltaModel->MeshMappings.Reset();
-	}
-
-	void FVertexDeltaEditorModel::CreateTrainingGroundTruthActor(UWorld* World)
-	{
-		UGeometryCache* GeomCache = GetVertexDeltaModel()->GetGeometryCache();
-		const FLinearColor LabelColor = FMLDeformerEditorStyle::Get().GetColor("MLDeformer.TargetMesh.LabelColor");
-		const FLinearColor WireframeColor = FMLDeformerEditorStyle::Get().GetColor("MLDeformer.TargetMesh.WireframeColor");
-		CreateGeomCacheActor(
-			World, 
-			ActorID_Train_GroundTruth, 
-			"Train GroundTruth", 
-			GeomCache, 
-			LabelColor, 
-			WireframeColor, 
-			LOCTEXT("TrainGroundTruthActorLabelText", "Target Mesh"),
-			true);
-	}
-
-	void FVertexDeltaEditorModel::CreateTestGroundTruthActor(UWorld* World)
-	{
-		UGeometryCache* GeomCache = GetVertexDeltaModelVizSettings()->GetTestGroundTruth();
-		const FLinearColor LabelColor = FMLDeformerEditorStyle::Get().GetColor("MLDeformer.GroundTruth.LabelColor");
-		const FLinearColor WireframeColor = FMLDeformerEditorStyle::Get().GetColor("MLDeformer.GroundTruth.WireframeColor");
-		CreateGeomCacheActor(
-			World, 
-			ActorID_Test_GroundTruth, 
-			"Test GroundTruth", 
-			GeomCache, 
-			LabelColor, 
-			WireframeColor,
-			LOCTEXT("TestGroundTruthActorLabelText", "Ground Truth"),
-			false);
-	}
-
-	void FVertexDeltaEditorModel::CreateGeomCacheActor(UWorld* World, int32 ActorID, const FName& Name, UGeometryCache* GeomCache, FLinearColor LabelColor, FLinearColor WireframeColor, const FText& LabelText, bool bIsTrainingActor)
-	{
-		FActorSpawnParameters SpawnParams;
-		SpawnParams.Name = MakeUniqueObjectName(World, AActor::StaticClass(), Name);
-		AActor* Actor = World->SpawnActor<AActor>(SpawnParams);
-		Actor->SetFlags(RF_Transient);
-
-		// Create the Geometry Cache Component.
-		UGeometryCacheComponent* GeomCacheComponent = NewObject<UGeometryCacheComponent>(Actor);
-		GeomCacheComponent->SetGeometryCache(GeomCache);
-		GeomCacheComponent->RegisterComponent();
-		GeomCacheComponent->SetOverrideWireframeColor(true);
-		GeomCacheComponent->SetWireframeOverrideColor(WireframeColor);
-		GeomCacheComponent->MarkRenderStateDirty();
-		GeomCacheComponent->SetVisibility(false);
-		Actor->SetRootComponent(GeomCacheComponent);
-
-		// Create the editor actor.
-		FMLDeformerEditorActor::FConstructSettings Settings;
-		Settings.Actor = Actor;
-		Settings.TypeID = ActorID;
-		Settings.LabelColor = LabelColor;
-		Settings.LabelText = LabelText;
-		Settings.bIsTrainingActor = bIsTrainingActor;
-
-		FMLDeformerGeomCacheActor* EditorActor = static_cast<FMLDeformerGeomCacheActor*>(CreateEditorActor(Settings));
-		EditorActor->SetGeometryCacheComponent(GeomCacheComponent);
-		EditorActors.Add(EditorActor);
-	} 
-
-	FMLDeformerEditorActor* FVertexDeltaEditorModel::CreateEditorActor(const FMLDeformerEditorActor::FConstructSettings& Settings) const
-	{
-		return new FMLDeformerGeomCacheActor(Settings);
-	}
-
-	FMLDeformerGeomCacheActor* FVertexDeltaEditorModel::FindVertexDeltaModelEditorActor(int32 TypeID) const
-	{
-		return static_cast<FMLDeformerGeomCacheActor*>(FindEditorActor(TypeID));
-	}
-
-	UVertexDeltaModelVizSettings* FVertexDeltaEditorModel::GetVertexDeltaModelVizSettings() const
-	{
-		return Cast<UVertexDeltaModelVizSettings>(GetVertexDeltaModel()->GetVizSettings());
-	}
-
-	double FVertexDeltaEditorModel::GetTrainingTimeAtFrame(int32 FrameNumber) const
-	{
-		const FMLDeformerGeomCacheActor* EditorActor = static_cast<FMLDeformerGeomCacheActor*>(FindEditorActor(ActorID_Train_GroundTruth));
-		if (EditorActor && EditorActor->GetGeometryCacheComponent() && EditorActor->GetGeometryCacheComponent()->GeometryCache.Get())
-		{
-			return EditorActor->GetGeometryCacheComponent()->GetTimeAtFrame(FrameNumber);
-		}
-
-		return FMLDeformerEditorModel::GetTrainingTimeAtFrame(FrameNumber);
-	}
-
-	int32 FVertexDeltaEditorModel::GetTrainingFrameAtTime(double TimeInSeconds) const
-	{
-		const FMLDeformerGeomCacheActor* EditorActor = static_cast<FMLDeformerGeomCacheActor*>(FindEditorActor(ActorID_Train_GroundTruth));
-		if (EditorActor && EditorActor->GetGeometryCacheComponent() && EditorActor->GetGeometryCacheComponent()->GeometryCache.Get())
-		{
-			return EditorActor->GetGeometryCacheComponent()->GetFrameAtTime(TimeInSeconds);
-		}
-
-		return FMLDeformerEditorModel::GetTrainingFrameAtTime(TimeInSeconds);
-	}
-
-	double FVertexDeltaEditorModel::GetTestTimeAtFrame(int32 FrameNumber) const
-	{
-		const FMLDeformerGeomCacheActor* EditorActor = static_cast<FMLDeformerGeomCacheActor*>(FindEditorActor(ActorID_Test_GroundTruth));
-		if (EditorActor && EditorActor->GetGeometryCacheComponent() && EditorActor->GetGeometryCacheComponent()->GeometryCache.Get())
-		{
-			return EditorActor->GetGeometryCacheComponent()->GetTimeAtFrame(FrameNumber);
-		}
-
-		return FMLDeformerEditorModel::GetTestTimeAtFrame(FrameNumber);
-	}
-
-	int32 FVertexDeltaEditorModel::GetTestFrameAtTime(double TimeInSeconds) const
-	{
-		const FMLDeformerGeomCacheActor* EditorActor = static_cast<FMLDeformerGeomCacheActor*>(FindEditorActor(ActorID_Test_GroundTruth));
-		if (EditorActor && EditorActor->GetGeometryCacheComponent() && EditorActor->GetGeometryCacheComponent()->GeometryCache.Get())
-		{
-			return EditorActor->GetGeometryCacheComponent()->GetFrameAtTime(TimeInSeconds);
-		}
-
-		return FMLDeformerEditorModel::GetTestFrameAtTime(TimeInSeconds);
-	}
-
-	int32 FVertexDeltaEditorModel::GetNumTrainingFrames() const
-	{
-		const UGeometryCache* GeometryCache = GetVertexDeltaModel()->GetGeometryCache();
-		if (GeometryCache == nullptr)
-		{
-			return 0;
-		}
-		const int32 StartFrame = GeometryCache->GetStartFrame();
-		const int32 EndFrame = GeometryCache->GetEndFrame();
-		return (EndFrame - StartFrame) + 1;
-	}
-
-	void FVertexDeltaEditorModel::UpdateIsReadyForTrainingState()
-	{
-		bIsReadyForTraining = false;
-
-		// Do some basic checks first, like if there is a skeletal mesh, ground truth, anim sequence, and if there are frames.
-		if (!IsEditorReadyForTrainingBasicChecks())
-		{
-			return;
-		}
-
-		// Now make sure the assets are compatible.
-		UVertexDeltaModel* VertexDeltaModel = GetVertexDeltaModel();
-		UGeometryCache* GeomCache = VertexDeltaModel->GetGeometryCache();
-		USkeletalMesh* SkeletalMesh = VertexDeltaModel->GetSkeletalMesh();
-		if (!GetGeomCacheErrorText(SkeletalMesh, GeomCache).IsEmpty())
-		{
-			return;
-		}
-
-		// Make sure every skeletal imported mesh has some geometry track.
-		const int32 NumGeomCacheTracks = GeomCache ? GeomCache->Tracks.Num() : 0;
-		int32 NumSkelMeshes = 0;
-		if (SkeletalMesh)
-		{
-			FSkeletalMeshModel* ImportedModel = SkeletalMesh->GetImportedModel();
-			if (ImportedModel)
-			{
-				NumSkelMeshes = ImportedModel->LODModels[0].ImportedMeshInfos.Num();		
-			}
-		}
-
-		// Allow the special case where there is just one mesh and track.
-		if (NumGeomCacheTracks != 1 || NumSkelMeshes != 1)
-		{
-			if (!GetGeomCacheSampler()->GetFailedImportedMeshNames().IsEmpty())
-			{
-				return;
-			}
-		}
-
-		bIsReadyForTraining = true;
-	}
-
 	ETrainingResult FVertexDeltaEditorModel::Train()
 	{
 		return TrainModel<UVertexDeltaTrainingModel>(this);
 	}
-
-	FString FVertexDeltaEditorModel::GetTrainedNetworkOnnxFile() const
-	{
-		return FString(FPaths::ProjectIntermediateDir() + TEXT("VertexDeltaModel/VertexDeltaModel.onnx"));
-	}
 }	// namespace UE::VertexDeltaModel
 
 #undef LOCTEXT_NAMESPACE
diff --git a/Engine/Plugins/Experimental/MLDeformer/VertexDeltaModel/Source/VertexDeltaModelEditor/Private/VertexDeltaModelDetails.cpp b/Engine/Plugins/Experimental/MLDeformer/VertexDeltaModel/Source/VertexDeltaModelEditor/Private/VertexDeltaModelDetails.cpp
--- a/Engine/Plugins/Experimental/MLDeformer/VertexDeltaModel/Source/VertexDeltaModelEditor/Private/VertexDeltaModelDetails.cpp
+++ b/Engine/Plugins/Experimental/MLDeformer/VertexDeltaModel/Source/VertexDeltaModelEditor/Private/VertexDeltaModelDetails.cpp
@@ -1,32 +1,8 @@
 // Copyright Epic Games, Inc. All Rights Reserved.
 
 #include "VertexDeltaModelDetails.h"
-#include "VertexDeltaEditorModel.h"
 #include "VertexDeltaModel.h"
-#include "MLDeformerModule.h"
-#include "MLDeformerAsset.h"
-#include "MLDeformerModel.h"
-#include "MLDeformerEditorModel.h"
-#include "MLDeformerEditorModule.h"
-#include "MLDeformerGeomCacheHelpers.h"
-#include "NeuralNetwork.h"
-#include "GeometryCache.h"
-#include "GeometryCacheTrack.h"
-#include "Animation/AnimSequence.h"
-#include "Engine/SkeletalMesh.h"
-#include "PropertyCustomizationHelpers.h"
-#include "AssetRegistry/AssetData.h"
-#include "Modules/ModuleManager.h"
-#include "DetailLayoutBuilder.h"
 #include "DetailCategoryBuilder.h"
-#include "DetailWidgetRow.h"
-#include "IDetailsView.h"
-#include "IDetailGroup.h"
-#include "Widgets/DeclarativeSyntaxSupport.h"
-#include "Widgets/Layout/SBox.h"
-#include "Widgets/Input/SButton.h"
-#include "DetailLayoutBuilder.h"
-#include "SWarningOrErrorBox.h"
 
 #define LOCTEXT_NAMESPACE "VertexDeltaModelDetails"
 
@@ -39,83 +15,17 @@ namespace UE::VertexDeltaModel
 		return MakeShareable(new FVertexDeltaModelDetails());
 	}
 
-	bool FVertexDeltaModelDetails::UpdateMemberPointers(const TArray<TWeakObjectPtr<UObject>>& Objects)
-	{
-		if (!FMLDeformerModelDetails::UpdateMemberPointers(Objects))
-		{
-			return false;
-		}
-
-		VertexModel = Cast<UVertexDeltaModel>(Model);
-		check(VertexModel);
-		VertexEditorModel = static_cast<FVertexDeltaEditorModel*>(EditorModel);
-
-		return (VertexModel != nullptr && VertexEditorModel != nullptr);
-	}
-
 	void FVertexDeltaModelDetails::CustomizeDetails(IDetailLayoutBuilder& DetailBuilder)
 	{
 		// Create all the detail categories and add the properties of the base class.
-		FMLDeformerModelDetails::CustomizeDetails(DetailBuilder);
+		FMLDeformerGeomCacheModelDetails::CustomizeDetails(DetailBuilder);
 
 		// Training settings.
-		SettingsCategoryBuilder->AddProperty(GET_MEMBER_NAME_CHECKED(UVertexDeltaModel, NumHiddenLayers));
-		SettingsCategoryBuilder->AddProperty(GET_MEMBER_NAME_CHECKED(UVertexDeltaModel, NumNeuronsPerLayer));
-		SettingsCategoryBuilder->AddProperty(GET_MEMBER_NAME_CHECKED(UVertexDeltaModel, NumIterations));
-		SettingsCategoryBuilder->AddProperty(GET_MEMBER_NAME_CHECKED(UVertexDeltaModel, BatchSize));
-		SettingsCategoryBuilder->AddProperty(GET_MEMBER_NAME_CHECKED(UVertexDeltaModel, LearningRate));
-	}
-	
-	void FVertexDeltaModelDetails::AddAnimSequenceErrors()
-	{
-		const FText WarningText = GetGeomCacheAnimSequenceErrorText(VertexModel->GetGeometryCache(), Model->GetAnimSequence());
-		BaseMeshCategoryBuilder->AddCustomRow(FText::FromString("AnimSeqWarning"))
-			.Visibility(!WarningText.IsEmpty() ? EVisibility::Visible : EVisibility::Collapsed)
-			.WholeRowContent()
-			[
-				SNew(SBox)
-				.Padding(FMargin(0.0f, 4.0f))
-				[
-					SNew(SWarningOrErrorBox)
-					.MessageStyle(EMessageStyle::Warning)
-					.Message(WarningText)
-				]
-			];
-	}
-
-	void FVertexDeltaModelDetails::AddTargetMesh()
-	{
-		TargetMeshCategoryBuilder->AddProperty(GET_MEMBER_NAME_CHECKED(UVertexDeltaModel, GeometryCache));
-
-		const FText TargetMeshErrorText = GetGeomCacheErrorText(VertexModel->GetSkeletalMesh(), VertexModel->GetGeometryCache());
-		TargetMeshCategoryBuilder->AddCustomRow(FText::FromString("TargetMeshError"))
-			.Visibility(!TargetMeshErrorText.IsEmpty() ? EVisibility::Visible : EVisibility::Collapsed)
-			.WholeRowContent()
-			[
-				SNew(SBox)
-				.Padding(FMargin(0.0f, 4.0f))
-				[
-					SNew(SWarningOrErrorBox)
-					.MessageStyle(EMessageStyle::Error)
-					.Message(TargetMeshErrorText)
-				]
-			];
-
-		const FText ChangedErrorText = EditorModel->GetTargetAssetChangedErrorText();
-		TargetMeshCategoryBuilder->AddCustomRow(FText::FromString("TargetMeshChangedError"))
-			.Visibility(!ChangedErrorText.IsEmpty() ? EVisibility::Visible : EVisibility::Collapsed)
-			.WholeRowContent()
-			[
-				SNew(SBox)
-				.Padding(FMargin(0.0f, 4.0f))
-				[
-					SNew(SWarningOrErrorBox)
-					.MessageStyle(EMessageStyle::Error)
-					.Message(ChangedErrorText)
-				]
-			];
-
-		AddGeomCacheMeshMappingWarnings(TargetMeshCategoryBuilder, Model->GetSkeletalMesh(), VertexModel->GetGeometryCache());
+		TrainingSettingsCategoryBuilder->AddProperty(GET_MEMBER_NAME_CHECKED(UVertexDeltaModel, NumHiddenLayers));
+		TrainingSettingsCategoryBuilder->AddProperty(GET_MEMBER_NAME_CHECKED(UVertexDeltaModel, NumNeuronsPerLayer));
+		TrainingSettingsCategoryBuilder->AddProperty(GET_MEMBER_NAME_CHECKED(UVertexDeltaModel, NumIterations));
+		TrainingSettingsCategoryBuilder->AddProperty(GET_MEMBER_NAME_CHECKED(UVertexDeltaModel, BatchSize));
+		TrainingSettingsCategoryBuilder->AddProperty(GET_MEMBER_NAME_CHECKED(UVertexDeltaModel, LearningRate));
 	}
 }	// namespace UE::VertexDeltaModel
 
diff --git a/Engine/Plugins/Experimental/MLDeformer/VertexDeltaModel/Source/VertexDeltaModelEditor/Private/VertexDeltaModelEditorModule.cpp b/Engine/Plugins/Experimental/MLDeformer/VertexDeltaModel/Source/VertexDeltaModelEditor/Private/VertexDeltaModelEditorModule.cpp
--- a/Engine/Plugins/Experimental/MLDeformer/VertexDeltaModel/Source/VertexDeltaModelEditor/Private/VertexDeltaModelEditorModule.cpp
+++ b/Engine/Plugins/Experimental/MLDeformer/VertexDeltaModel/Source/VertexDeltaModelEditor/Private/VertexDeltaModelEditorModule.cpp
@@ -1,6 +1,5 @@
 // Copyright Epic Games, Inc. All Rights Reserved.
 
-#include "VertexDeltaModelEditorModule.h"
 #include "MLDeformerEditorModule.h"
 #include "VertexDeltaModelVizSettingsDetails.h"
 #include "VertexDeltaModelDetails.h"
@@ -11,6 +10,19 @@
 
 #define LOCTEXT_NAMESPACE "VertexDeltaModelEditorModule"
 
+namespace UE::VertexDeltaModel
+{
+	class VERTEXDELTAMODELEDITOR_API FVertexDeltaModelEditorModule
+		: public IModuleInterface
+	{
+	public:
+		// IModuleInterface overrides.
+		void StartupModule() override;
+		void ShutdownModule() override;
+		// ~END IModuleInterface overrides.
+	};
+}	// namespace UE::VertexDeltaModel
+
 IMPLEMENT_MODULE(UE::VertexDeltaModel::FVertexDeltaModelEditorModule, VertexDeltaModelEditor)
 
 namespace UE::VertexDeltaModel
diff --git a/Engine/Plugins/Experimental/MLDeformer/VertexDeltaModel/Source/VertexDeltaModelEditor/Private/VertexDeltaModelVizSettingsDetails.cpp b/Engine/Plugins/Experimental/MLDeformer/VertexDeltaModel/Source/VertexDeltaModelEditor/Private/VertexDeltaModelVizSettingsDetails.cpp
deleted file mode 100644
diff --git a/Engine/Plugins/Experimental/MLDeformer/VertexDeltaModel/Source/VertexDeltaModelEditor/Private/VertexDeltaTrainingModel.cpp b/Engine/Plugins/Experimental/MLDeformer/VertexDeltaModel/Source/VertexDeltaModelEditor/Private/VertexDeltaTrainingModel.cpp
deleted file mode 100644
diff --git a/Engine/Plugins/Experimental/MLDeformer/VertexDeltaModel/Source/VertexDeltaModelEditor/Public/VertexDeltaEditorModel.h b/Engine/Plugins/Experimental/MLDeformer/VertexDeltaModel/Source/VertexDeltaModelEditor/Public/VertexDeltaEditorModel.h
--- a/Engine/Plugins/Experimental/MLDeformer/VertexDeltaModel/Source/VertexDeltaModelEditor/Public/VertexDeltaEditorModel.h
+++ b/Engine/Plugins/Experimental/MLDeformer/VertexDeltaModel/Source/VertexDeltaModelEditor/Public/VertexDeltaEditorModel.h
@@ -3,65 +3,25 @@
 #pragma once
 
 #include "CoreMinimal.h"
-#include "MLDeformerEditorModel.h"
-#include "MLDeformerGeomCacheSampler.h"
-#include "VertexDeltaModel.h"
-#include "MLDeformerEditorActor.h"
-#include "UObject/GCObject.h"
-
-namespace UE::MLDeformer
-{
-	class FMLDeformerGeomCacheActor;
-}
-
-class UMLDeformerAsset;
-class USkeletalMesh;
-class UGeometryCache;
-class UVertexDeltaModel;
-class UVertexDeltaModelVizSettings;
+#include "MLDeformerGeomCacheEditorModel.h"
 
 namespace UE::VertexDeltaModel
 {
 	using namespace UE::MLDeformer;
 
-	class FVertexDeltaModelSampler;
-
 	class VERTEXDELTAMODELEDITOR_API FVertexDeltaEditorModel 
-		: public UE::MLDeformer::FMLDeformerEditorModel
+		: public UE::MLDeformer::FMLDeformerGeomCacheEditorModel
 	{
 	public:
 		// We need to implement this static MakeInstance method.
 		static FMLDeformerEditorModel* MakeInstance();
 
 		// FGCObject overrides.
-		virtual void AddReferencedObjects(FReferenceCollector& Collector) override;
 		virtual FString GetReferencerName() const override { return TEXT("FVertexDeltaEditorModel"); }
 		// ~END FGCObject overrides.
 	
 		// FMLDeformerEditorModel overrides.
-		virtual FMLDeformerEditorActor* CreateEditorActor(const FMLDeformerEditorActor::FConstructSettings& Settings) const override;
-		virtual FMLDeformerSampler* CreateSampler() const;
-		virtual void CreateTrainingGroundTruthActor(UWorld* World) override;
-		virtual void CreateTestGroundTruthActor(UWorld* World) override;
-		virtual void OnInputAssetsChanged() override;
-		virtual int32 GetNumTrainingFrames() const override;
-		virtual double GetTrainingTimeAtFrame(int32 FrameNumber) const override;
-		virtual int32 GetTrainingFrameAtTime(double TimeInSeconds) const override;
-		virtual double GetTestTimeAtFrame(int32 FrameNumber) const override;
-		virtual int32 GetTestFrameAtTime(double TimeInSeconds) const override;
-		virtual void UpdateIsReadyForTrainingState() override;
 		virtual ETrainingResult Train() override;
-		virtual FString GetTrainedNetworkOnnxFile() const override;
-		virtual void OnPostEditChangeProperty(FPropertyChangedEvent& PropertyChangedEvent) override;
 		// ~END FMLDeformerEditorModel overrides.
-
-		// Some helpers that cast to this model's variants of some classes.
-		UVertexDeltaModel* GetVertexDeltaModel() const { return Cast<UVertexDeltaModel>(Model); }
-		UVertexDeltaModelVizSettings* GetVertexDeltaModelVizSettings() const;
-		FMLDeformerGeomCacheActor* FindVertexDeltaModelEditorActor(int32 TypeID) const;
-		FMLDeformerGeomCacheSampler* GetGeomCacheSampler() const { return static_cast<FMLDeformerGeomCacheSampler*>(Sampler); }
-
-	protected:
-		void CreateGeomCacheActor(UWorld* World, int32 ActorID, const FName& Name, UGeometryCache* GeomCache, FLinearColor LabelColor, FLinearColor WireframeColor, const FText& LabelText, bool bIsTrainingActor);
 	};
 }	// namespace UE::VertexDeltaModel
diff --git a/Engine/Plugins/Experimental/MLDeformer/VertexDeltaModel/Source/VertexDeltaModelEditor/Public/VertexDeltaModelDetails.h b/Engine/Plugins/Experimental/MLDeformer/VertexDeltaModel/Source/VertexDeltaModelEditor/Public/VertexDeltaModelDetails.h
--- a/Engine/Plugins/Experimental/MLDeformer/VertexDeltaModel/Source/VertexDeltaModelEditor/Public/VertexDeltaModelDetails.h
+++ b/Engine/Plugins/Experimental/MLDeformer/VertexDeltaModel/Source/VertexDeltaModelEditor/Public/VertexDeltaModelDetails.h
@@ -3,25 +3,12 @@
 #pragma once
 
 #include "CoreMinimal.h"
-#include "MLDeformerModelDetails.h"
-#include "IDetailCustomization.h"
-#include "AssetRegistry/AssetData.h"
-#include "Input/Reply.h"
-
-class IDetailLayoutBuilder;
-class UVertexDeltaModel;
-
-namespace UE::MLDeformer
-{
-	class FMLDeformerEditorModel;
-}
+#include "MLDeformerGeomCacheModelDetails.h"
 
 namespace UE::VertexDeltaModel
 {
-	class FVertexDeltaEditorModel;
-
 	class VERTEXDELTAMODELEDITOR_API FVertexDeltaModelDetails
-		: public UE::MLDeformer::FMLDeformerModelDetails
+		: public UE::MLDeformer::FMLDeformerGeomCacheModelDetails
 	{
 	public:
 		/** Makes a new instance of this detail layout class for a specific detail view requesting it. */
@@ -30,15 +17,5 @@ namespace UE::VertexDeltaModel
 		// ILayoutDetails overrides.
 		void CustomizeDetails(IDetailLayoutBuilder& DetailBuilder) override;
 		// ~END ILayoutDetails overrides.
-
-		// FMLDeformerModelDetails overrides.
-		bool UpdateMemberPointers(const TArray<TWeakObjectPtr<UObject>>& Objects) override;
-		void AddTargetMesh() override;
-		void AddAnimSequenceErrors() override;
-		// ~END FMLDeformerModelDetails overrides.
-
-	protected:
-		TObjectPtr<UVertexDeltaModel> VertexModel = nullptr;
-		FVertexDeltaEditorModel* VertexEditorModel = nullptr;
 	};
 }	// namespace UE::VertexDeltaModel
diff --git a/Engine/Plugins/Experimental/MLDeformer/VertexDeltaModel/Source/VertexDeltaModelEditor/Public/VertexDeltaModelEditorModule.h b/Engine/Plugins/Experimental/MLDeformer/VertexDeltaModel/Source/VertexDeltaModelEditor/Public/VertexDeltaModelEditorModule.h
deleted file mode 100644
diff --git a/Engine/Plugins/Experimental/MLDeformer/VertexDeltaModel/Source/VertexDeltaModelEditor/Public/VertexDeltaModelVizSettingsDetails.h b/Engine/Plugins/Experimental/MLDeformer/VertexDeltaModel/Source/VertexDeltaModelEditor/Public/VertexDeltaModelVizSettingsDetails.h
--- a/Engine/Plugins/Experimental/MLDeformer/VertexDeltaModel/Source/VertexDeltaModelEditor/Public/VertexDeltaModelVizSettingsDetails.h
+++ b/Engine/Plugins/Experimental/MLDeformer/VertexDeltaModel/Source/VertexDeltaModelEditor/Public/VertexDeltaModelVizSettingsDetails.h
@@ -3,32 +3,18 @@
 #pragma once
 
 #include "CoreMinimal.h"
-#include "MLDeformerVizSettingsDetails.h"
-#include "IDetailCustomization.h"
-#include "AssetRegistry/AssetData.h"
-#include "PropertyHandle.h"
-
-class IDetailLayoutBuilder;
-class USkeleton;
-class UVertexDeltaModel;
-class UVertexDeltaModelVizSettings;
+#include "MLDeformerGeomCacheVizSettingsDetails.h"
 
 namespace UE::VertexDeltaModel
 {
 	class VERTEXDELTAMODELEDITOR_API FVertexDeltaModelVizSettingsDetails
-		: public UE::MLDeformer::FMLDeformerVizSettingsDetails
+		: public UE::MLDeformer::FMLDeformerGeomCacheVizSettingsDetails
 	{
 	public:
 		/** Makes a new instance of this detail layout class for a specific detail view requesting it. */
-		static TSharedRef<IDetailCustomization> MakeInstance();
-
-		// FMLDeformerVizSettingsDetails overrides.
-		bool UpdateMemberPointers(const TArray<TWeakObjectPtr<UObject>>& Objects) override;
-		void AddGroundTruth() override;
-		// ~END FMLDeformerVizSettingsDetails overrides.
-
-	protected:
-		TObjectPtr<UVertexDeltaModel> VertexDeltaModel = nullptr;
-		TObjectPtr<UVertexDeltaModelVizSettings> VertexDeltaVizSettings = nullptr;
+		static TSharedRef<IDetailCustomization> MakeInstance()
+		{
+			return MakeShareable(new FVertexDeltaModelVizSettingsDetails());
+		}
 	};
 }	// namespace UE::VertexDeltaModel
diff --git a/Engine/Plugins/Experimental/MLDeformer/VertexDeltaModel/Source/VertexDeltaModelEditor/Public/VertexDeltaTrainingModel.h b/Engine/Plugins/Experimental/MLDeformer/VertexDeltaModel/Source/VertexDeltaModelEditor/Public/VertexDeltaTrainingModel.h
--- a/Engine/Plugins/Experimental/MLDeformer/VertexDeltaModel/Source/VertexDeltaModelEditor/Public/VertexDeltaTrainingModel.h
+++ b/Engine/Plugins/Experimental/MLDeformer/VertexDeltaModel/Source/VertexDeltaModelEditor/Public/VertexDeltaTrainingModel.h
@@ -6,12 +6,6 @@
 #include "MLDeformerTrainingModel.h"
 #include "VertexDeltaTrainingModel.generated.h"
 
-namespace UE::VertexDeltaModel
-{
-	class FVertexDeltaEditorModel;
-}
-
-class UVertexDeltaModel;
 
 UCLASS(Blueprintable)
 class VERTEXDELTAMODELEDITOR_API UVertexDeltaTrainingModel
@@ -23,7 +17,4 @@ public:
 	/** Main training function, with implementation in python. */
 	UFUNCTION(BlueprintImplementableEvent, Category = "Python")
 	int32 Train() const;
-
-	UVertexDeltaModel* GetVertexDeltaModel() const;
-	UE::VertexDeltaModel::FVertexDeltaEditorModel* GetVertexDeltaEditorModel() const;
 };
