MLDeformer: When binding MLDeformer to a Skeletal Mesh, search for one that matches the expected vertex count.
This fixes issues where there are multiple Skeletal Mesh Components on an actor and we select the wrong one to apply the ML Deformer to.
#rnx

diff --git a/Engine/Plugins/Experimental/MLDeformer/MLDeformerFramework/Source/MLDeformerFramework/Private/MLDeformerComponent.cpp b/Engine/Plugins/Experimental/MLDeformer/MLDeformerFramework/Source/MLDeformerFramework/Private/MLDeformerComponent.cpp
--- a/Engine/Plugins/Experimental/MLDeformer/MLDeformerFramework/Source/MLDeformerFramework/Private/MLDeformerComponent.cpp
+++ b/Engine/Plugins/Experimental/MLDeformer/MLDeformerFramework/Source/MLDeformerFramework/Private/MLDeformerComponent.cpp
@@ -103,9 +103,38 @@ void UMLDeformerComponent::BeginDestroy()
 void UMLDeformerComponent::Activate(bool bReset)
 {
 	// If we haven't pointed to some skeletal mesh component to use, then try to find one on the actor.
-	USkeletalMeshComponent* SkelMeshComponentToUse = SkelMeshComponent;
-	if (SkelMeshComponentToUse == nullptr)
+	
+	if (SkelMeshComponent == nullptr)
+	{
+		// First search for a skeletal mesh component with the expected number of vertices.
+		UMLDeformerModel const* Model = DeformerAsset ? DeformerAsset->GetModel() : nullptr;
+		const int32 NumModelVertices = Model ? Model->GetVertexMap().Num() : -1;
+
+		if (NumModelVertices > 0)
+		{
+			TArray<USkeletalMeshComponent*> Components;
+		
+			AActor* Actor = Cast<AActor>(GetOuter());
+			Actor->GetComponents(Components);
+		
+			for (USkeletalMeshComponent* Component : Components)
+			{
+				USkeletalMesh const* SkeletalMesh = Component->GetSkeletalMeshAsset();
+				FSkeletalMeshRenderData const* RenderData = SkeletalMesh ? SkeletalMesh->GetResourceForRendering() : nullptr;
+				const int32 NumComponentVertices = RenderData && RenderData->LODRenderData.IsValidIndex(0) ? RenderData->LODRenderData[0].GetNumVertices() : -1;
+
+				if (NumComponentVertices == NumModelVertices)
+				{
+					SkelMeshComponent = Component;
+					break;
+				}
+			}
+		}
+	}
+
+	if (SkelMeshComponent == nullptr)
 	{
+		// Fall back to the first skeletal mesh component.
 		AActor* Actor = Cast<AActor>(GetOuter());
 		SkelMeshComponent = Actor->FindComponentByClass<USkeletalMeshComponent>();
 	}
