[MLDeformer]
- API improvements, mainly making properties protected.
- Fixed some bugs that prevented the sampling cache from being invalidated when changing certain properties.
- Significantly more API documentation.
- Include optimizations.

diff --git a/Engine/Plugins/Experimental/MLDeformer/MLDeformerFramework/Source/MLDeformerFramework/Private/MLDeformerComponent.cpp b/Engine/Plugins/Experimental/MLDeformer/MLDeformerFramework/Source/MLDeformerFramework/Private/MLDeformerComponent.cpp
--- a/Engine/Plugins/Experimental/MLDeformer/MLDeformerFramework/Source/MLDeformerFramework/Private/MLDeformerComponent.cpp
+++ b/Engine/Plugins/Experimental/MLDeformer/MLDeformerFramework/Source/MLDeformerFramework/Private/MLDeformerComponent.cpp
@@ -72,7 +72,7 @@ void UMLDeformerComponent::AddNeuralNetworkModifyDelegate()
 	UMLDeformerModel* Model = DeformerAsset->GetModel();
 	if (Model)
 	{
-		NeuralNetworkModifyDelegateHandle = Model->NeuralNetworkModifyDelegate.AddLambda
+		NeuralNetworkModifyDelegateHandle = Model->GetNeuralNetworkModifyDelegate().AddLambda
 		(
 			([this]()
 			{
@@ -88,7 +88,7 @@ void UMLDeformerComponent::RemoveNeuralNetworkModifyDelegate()
 		NeuralNetworkModifyDelegateHandle != FDelegateHandle() && 
 		DeformerAsset->GetModel())
 	{
-		DeformerAsset->GetModel()->NeuralNetworkModifyDelegate.Remove(NeuralNetworkModifyDelegateHandle);
+		DeformerAsset->GetModel()->GetNeuralNetworkModifyDelegate().Remove(NeuralNetworkModifyDelegateHandle);
 	}
 	
 	NeuralNetworkModifyDelegateHandle = FDelegateHandle();
@@ -103,24 +103,24 @@ void UMLDeformerComponent::BeginDestroy()
 void UMLDeformerComponent::Activate(bool bReset)
 {
 	// If we haven't pointed to some skeletal mesh component to use, then try to find one on the actor.
-	
 	if (SkelMeshComponent == nullptr)
 	{
 		// First search for a skeletal mesh component with the expected number of vertices.
-		UMLDeformerModel const* Model = DeformerAsset ? DeformerAsset->GetModel() : nullptr;
+		// This will try to find a skeletal mesh with the same number of vertices.
+		const UMLDeformerModel* Model = DeformerAsset ? DeformerAsset->GetModel() : nullptr;
 		const int32 NumModelVertices = Model ? Model->GetVertexMap().Num() : -1;
-
 		if (NumModelVertices > 0)
 		{
+			// Get a list of all skeletal mesh components on the actor.
 			TArray<USkeletalMeshComponent*> Components;
-		
 			AActor* Actor = Cast<AActor>(GetOuter());
-			Actor->GetComponents(Components);
+			Actor->GetComponents<USkeletalMeshComponent>(Components);
 		
+			// Find a component that uses a mesh with the same vertex count.
 			for (USkeletalMeshComponent* Component : Components)
 			{
-				USkeletalMesh const* SkeletalMesh = Component->GetSkeletalMeshAsset();
-				FSkeletalMeshRenderData const* RenderData = SkeletalMesh ? SkeletalMesh->GetResourceForRendering() : nullptr;
+				const USkeletalMesh* SkeletalMesh = Component->GetSkeletalMeshAsset();
+				const FSkeletalMeshRenderData* RenderData = SkeletalMesh ? SkeletalMesh->GetResourceForRendering() : nullptr;
 				const int32 NumComponentVertices = RenderData && RenderData->LODRenderData.IsValidIndex(0) ? RenderData->LODRenderData[0].GetNumVertices() : -1;
 
 				if (NumComponentVertices == NumModelVertices)
@@ -135,7 +135,7 @@ void UMLDeformerComponent::Activate(bool bReset)
 	if (SkelMeshComponent == nullptr)
 	{
 		// Fall back to the first skeletal mesh component.
-		AActor* Actor = Cast<AActor>(GetOuter());
+		const AActor* Actor = Cast<AActor>(GetOuter());
 		SkelMeshComponent = Actor->FindComponentByClass<USkeletalMeshComponent>();
 	}
 
diff --git a/Engine/Plugins/Experimental/MLDeformer/MLDeformerFramework/Source/MLDeformerFramework/Private/MLDeformerGeomCacheHelpers.cpp b/Engine/Plugins/Experimental/MLDeformer/MLDeformerFramework/Source/MLDeformerFramework/Private/MLDeformerGeomCacheHelpers.cpp
--- a/Engine/Plugins/Experimental/MLDeformer/MLDeformerFramework/Source/MLDeformerFramework/Private/MLDeformerGeomCacheHelpers.cpp
+++ b/Engine/Plugins/Experimental/MLDeformer/MLDeformerFramework/Source/MLDeformerFramework/Private/MLDeformerGeomCacheHelpers.cpp
@@ -2,6 +2,7 @@
 
 #include "MLDeformerGeomCacheHelpers.h"
 #include "MLDeformerModel.h"
+#include "MLDeformerModule.h"
 #include "GeometryCache.h"
 #include "GeometryCacheMeshData.h"
 #include "GeometryCacheTrack.h"
diff --git a/Engine/Plugins/Experimental/MLDeformer/MLDeformerFramework/Source/MLDeformerFramework/Private/MLDeformerInputInfo.cpp b/Engine/Plugins/Experimental/MLDeformer/MLDeformerFramework/Source/MLDeformerFramework/Private/MLDeformerInputInfo.cpp
--- a/Engine/Plugins/Experimental/MLDeformer/MLDeformerFramework/Source/MLDeformerFramework/Private/MLDeformerInputInfo.cpp
+++ b/Engine/Plugins/Experimental/MLDeformer/MLDeformerFramework/Source/MLDeformerFramework/Private/MLDeformerInputInfo.cpp
@@ -10,6 +10,7 @@
 #include "Engine/SkeletalMesh.h"
 #include "Components/SkeletalMeshComponent.h"
 
+
 void UMLDeformerInputInfo::Reset()
 {
 	BoneNameStrings.Empty();
@@ -41,6 +42,86 @@ void UMLDeformerInputInfo::UpdateFNames()
 	}
 }
 
+void UMLDeformerInputInfo::OnPostLoad()
+{ 
+	UpdateFNames();
+}
+
+bool UMLDeformerInputInfo::IsEmpty() const
+{ 
+	return (BoneNameStrings.IsEmpty() && CurveNameStrings.IsEmpty());
+}
+
+int32 UMLDeformerInputInfo::GetNumBones() const
+{ 
+	return BoneNames.Num();
+}
+
+const FString& UMLDeformerInputInfo::GetBoneNameString(int32 Index) const
+{ 
+	return BoneNameStrings[Index];
+}
+
+const FName UMLDeformerInputInfo::GetBoneName(int32 Index) const
+{ 
+	return BoneNames[Index];
+}
+
+int32 UMLDeformerInputInfo::GetNumCurves() const
+{ 
+	return CurveNames.Num();
+}
+
+TArray<FString>& UMLDeformerInputInfo::GetBoneNameStrings()
+{ 
+	return BoneNameStrings;
+}
+
+TArray<FString>& UMLDeformerInputInfo::GetCurveNameStrings()
+{ 
+	return CurveNameStrings;
+}
+
+TArray<FName>& UMLDeformerInputInfo::GetBoneNames()
+{ 
+	return BoneNames;
+}
+
+TArray<FName>& UMLDeformerInputInfo::GetCurveNames()
+{ 
+	return CurveNames;
+}
+
+const FString& UMLDeformerInputInfo::GetCurveNameString(int32 Index) const
+{ 
+	return CurveNameStrings[Index];
+}
+
+const FName UMLDeformerInputInfo::GetCurveName(int32 Index) const
+{ 
+	return CurveNames[Index];
+}
+
+int32 UMLDeformerInputInfo::GetNumBaseMeshVertices() const
+{ 
+	return NumBaseMeshVertices;
+}
+
+void UMLDeformerInputInfo::SetNumBaseVertices(int32 NumVerts)
+{ 
+	NumBaseMeshVertices = NumVerts;
+}
+
+int32 UMLDeformerInputInfo::GetNumTargetMeshVertices() const
+{
+	return NumTargetMeshVertices;
+}
+
+void UMLDeformerInputInfo::SetNumTargetVertices(int32 NumVerts)
+{
+	NumTargetMeshVertices = NumVerts;
+}
+
 bool UMLDeformerInputInfo::IsCompatible(USkeletalMesh* SkeletalMesh) const
 {
 	if (SkeletalMesh == nullptr)
diff --git a/Engine/Plugins/Experimental/MLDeformer/MLDeformerFramework/Source/MLDeformerFramework/Private/MLDeformerModel.cpp b/Engine/Plugins/Experimental/MLDeformer/MLDeformerFramework/Source/MLDeformerFramework/Private/MLDeformerModel.cpp
--- a/Engine/Plugins/Experimental/MLDeformer/MLDeformerFramework/Source/MLDeformerFramework/Private/MLDeformerModel.cpp
+++ b/Engine/Plugins/Experimental/MLDeformer/MLDeformerFramework/Source/MLDeformerFramework/Private/MLDeformerModel.cpp
@@ -9,9 +9,6 @@
 #include "Engine/SkeletalMesh.h"
 #include "Rendering/SkeletalMeshModel.h"
 #include "Rendering/SkeletalMeshLODModel.h"
-#include "GeometryCache.h"
-#include "GeometryCacheMeshData.h"
-#include "GeometryCacheTrack.h"
 #include "NeuralNetwork.h"
 #include "Animation/AnimSequence.h"
 #include "UObject/UObjectGlobals.h"
@@ -22,7 +19,7 @@ namespace UE::MLDeformer
 	{
 		if (VertexMap.Num() > 0)
 		{
-			FRHIResourceCreateInfo CreateInfo(TEXT("FVertexMapBuffer"));
+			FRHIResourceCreateInfo CreateInfo(TEXT("UMLDeformerModel::FVertexMapBuffer"));
 
 			VertexBufferRHI = RHICreateVertexBuffer(VertexMap.Num() * sizeof(uint32), BUF_Static | BUF_ShaderResource, CreateInfo);
 			uint32* Data = reinterpret_cast<uint32*>(RHILockBuffer(VertexBufferRHI, 0, VertexMap.Num() * sizeof(uint32), RLM_WriteOnly));
diff --git a/Engine/Plugins/Experimental/MLDeformer/MLDeformerFramework/Source/MLDeformerFramework/Private/MLDeformerModelInstance.cpp b/Engine/Plugins/Experimental/MLDeformer/MLDeformerFramework/Source/MLDeformerFramework/Private/MLDeformerModelInstance.cpp
--- a/Engine/Plugins/Experimental/MLDeformer/MLDeformerFramework/Source/MLDeformerFramework/Private/MLDeformerModelInstance.cpp
+++ b/Engine/Plugins/Experimental/MLDeformer/MLDeformerFramework/Source/MLDeformerFramework/Private/MLDeformerModelInstance.cpp
@@ -26,6 +26,51 @@ void UMLDeformerModelInstance::Release()
 	}
 }
 
+USkeletalMeshComponent* UMLDeformerModelInstance::GetSkeletalMeshComponent() const
+{ 
+	return SkeletalMeshComponent;
+}
+
+UMLDeformerModel* UMLDeformerModelInstance::GetModel() const
+{ 
+	return Model.Get();
+}
+
+void UMLDeformerModelInstance::SetModel(UMLDeformerModel* InModel)
+{ 
+	Model = InModel;
+}
+
+int32 UMLDeformerModelInstance::GetNeuralNetworkInferenceHandle() const
+{ 
+	return NeuralNetworkInferenceHandle;
+}
+
+void UMLDeformerModelInstance::SetHasPostInitialized(bool bHasInitialized)
+{ 
+	bHasPostInitialized = bHasInitialized;
+}
+
+bool UMLDeformerModelInstance::HasPostInitialized() const
+{ 
+	return bHasPostInitialized;
+}
+
+const TArray<FTransform>& UMLDeformerModelInstance::GetBoneTransforms() const
+{ 
+	return BoneTransforms;
+}
+
+bool UMLDeformerModelInstance::IsCompatible() const
+{ 
+	return bIsCompatible;
+}
+
+const FString& UMLDeformerModelInstance::GetCompatibilityErrorText() const
+{ 
+	return ErrorText;
+}
+
 void UMLDeformerModelInstance::Init(USkeletalMeshComponent* SkelMeshComponent)
 {
 	SkeletalMeshComponent = SkelMeshComponent;
diff --git a/Engine/Plugins/Experimental/MLDeformer/MLDeformerFramework/Source/MLDeformerFramework/Private/MLDeformerMorphModel.cpp b/Engine/Plugins/Experimental/MLDeformer/MLDeformerFramework/Source/MLDeformerFramework/Private/MLDeformerMorphModel.cpp
--- a/Engine/Plugins/Experimental/MLDeformer/MLDeformerFramework/Source/MLDeformerFramework/Private/MLDeformerMorphModel.cpp
+++ b/Engine/Plugins/Experimental/MLDeformer/MLDeformerFramework/Source/MLDeformerFramework/Private/MLDeformerMorphModel.cpp
@@ -4,9 +4,6 @@
 #include "MLDeformerMorphModelInstance.h"
 #include "MLDeformerModelInstance.h"
 #include "MLDeformerComponent.h"
-#include "UObject/Object.h"
-#include "UObject/UObjectGlobals.h"
-#include "Components/SkinnedMeshComponent.h"
 #include "Components/SkeletalMeshComponent.h"
 #include "Rendering/MorphTargetVertexInfoBuffers.h"
 
@@ -101,11 +98,16 @@ void UMLDeformerMorphModel::PostMLDeformerComponentInit(UMLDeformerModelInstance
 	}
 }
 
-void UMLDeformerMorphModel::SetMorphTargetDeltas(const TArray<float>& Deltas)
+void UMLDeformerMorphModel::SetMorphTargetDeltaFloats(const TArray<float>& Deltas)
 {
 	FloatArrayToVector3Array(Deltas, MorphTargetDeltas);
 }
 
+void UMLDeformerMorphModel::SetMorphTargetDeltas(const TArray<FVector3f>& Deltas)
+{
+	MorphTargetDeltas = Deltas;
+}
+
 int32 UMLDeformerMorphModel::GetMorphTargetDeltaStartIndex(int32 BlendShapeIndex) const
 {
 	if (MorphTargetDeltas.Num() == 0)
diff --git a/Engine/Plugins/Experimental/MLDeformer/MLDeformerFramework/Source/MLDeformerFramework/Private/MLDeformerMorphModelInstance.cpp b/Engine/Plugins/Experimental/MLDeformer/MLDeformerFramework/Source/MLDeformerFramework/Private/MLDeformerMorphModelInstance.cpp
--- a/Engine/Plugins/Experimental/MLDeformer/MLDeformerFramework/Source/MLDeformerFramework/Private/MLDeformerMorphModelInstance.cpp
+++ b/Engine/Plugins/Experimental/MLDeformer/MLDeformerFramework/Source/MLDeformerFramework/Private/MLDeformerMorphModelInstance.cpp
@@ -16,7 +16,7 @@ void UMLDeformerMorphModelInstance::RunNeuralNetwork(float ModelWeight)
 		return;
 	}
 
-	const int32 ExternalMorphSetID = MorphModel->ExternalMorphSetID;
+	const int32 ExternalMorphSetID = MorphModel->GetExternalMorphSetID();
 	// If this check fails please set this member to some value larger than 0 in your model's constructor.
 	checkf(ExternalMorphSetID != -1, TEXT("Please set the ExternalMorphSetID member value to a unique value for your model type."));
 
diff --git a/Engine/Plugins/Experimental/MLDeformer/MLDeformerFramework/Source/MLDeformerFramework/Private/Tests/ModelTests.cpp b/Engine/Plugins/Experimental/MLDeformer/MLDeformerFramework/Source/MLDeformerFramework/Private/Tests/ModelTests.cpp
--- a/Engine/Plugins/Experimental/MLDeformer/MLDeformerFramework/Source/MLDeformerFramework/Private/Tests/ModelTests.cpp
+++ b/Engine/Plugins/Experimental/MLDeformer/MLDeformerFramework/Source/MLDeformerFramework/Private/Tests/ModelTests.cpp
@@ -8,6 +8,7 @@
 #include "MLDeformerTestModel.h"
 #include "MLDeformerAsset.h"
 #include "MLDeformerComponent.h"
+#include "MLDeformerInputInfo.h"
 #include "UObject/GCObjectScopeGuard.h"
 #include "PreviewScene.h"
 #include "Components/SkeletalMeshComponent.h"
@@ -31,7 +32,7 @@ namespace UE::MLDeformerTests
 		USkeletalMesh* SkeletalMesh = LoadObject<USkeletalMesh>(nullptr, TEXT("SkeletalMesh'/MLDeformerFramework/Tests/SKM_Cylinder.SKM_Cylinder'"));
 		UTEST_NOT_NULL(TEXT("Loading skeletal mesh"), SkeletalMesh);
 
-		TestModel->SkeletalMesh = SkeletalMesh;
+		TestModel->SetSkeletalMesh(SkeletalMesh);
 		UTEST_EQUAL(TEXT("GetSkeletalMesh check"), TestModel->GetSkeletalMesh(), SkeletalMesh);
 
 		DeformerAsset->SetModel(TestModel);
@@ -45,7 +46,7 @@ namespace UE::MLDeformerTests
 		UTEST_EQUAL(TEXT("Imported verts check"), ImportedVerts, ExpectedNumImportedVerts);
 
 		TestModel->UpdateCachedNumVertices();
-		UTEST_EQUAL(TEXT("Imported verts check"), TestModel->NumBaseMeshVerts, ImportedVerts);
+		UTEST_EQUAL(TEXT("Imported verts check"), TestModel->GetNumBaseMeshVerts(), ImportedVerts);
 
 		const USkeleton* Skeleton = SkeletalMesh->GetSkeleton();
 		bool bErrorOnInvalid = false;
diff --git a/Engine/Plugins/Experimental/MLDeformer/MLDeformerFramework/Source/MLDeformerFramework/Public/MLDeformerAsset.h b/Engine/Plugins/Experimental/MLDeformer/MLDeformerFramework/Source/MLDeformerFramework/Public/MLDeformerAsset.h
--- a/Engine/Plugins/Experimental/MLDeformer/MLDeformerFramework/Source/MLDeformerFramework/Public/MLDeformerAsset.h
+++ b/Engine/Plugins/Experimental/MLDeformer/MLDeformerFramework/Source/MLDeformerFramework/Public/MLDeformerAsset.h
@@ -19,11 +19,20 @@ class MLDEFORMERFRAMEWORK_API UMLDeformerAsset
 	GENERATED_BODY()
 
 public:
-	UMLDeformerModel* GetModel() const { return Model.Get(); }
-	void SetModel(UMLDeformerModel* InModel) { Model = InModel; }
+	/**
+	 * Get the ML Deformer model that is being applied by this asset.
+	 * @return A pointer to the model.
+	 */
+	UMLDeformerModel* GetModel() const			{ return Model.Get(); }
+
+	/**
+	 * Set the ML Deformer model that is used by this deformer asset.
+	 * @param InModel A pointer to the model object.
+	 */
+	void SetModel(UMLDeformerModel* InModel)	{ Model = InModel; }
 
 public:
 	/** The ML Deformer model, used to deform the mesh. */
 	UPROPERTY()
-	TObjectPtr<UMLDeformerModel> Model;
+	TObjectPtr<UMLDeformerModel> Model = nullptr;
 };
diff --git a/Engine/Plugins/Experimental/MLDeformer/MLDeformerFramework/Source/MLDeformerFramework/Public/MLDeformerComponent.h b/Engine/Plugins/Experimental/MLDeformer/MLDeformerFramework/Source/MLDeformerFramework/Public/MLDeformerComponent.h
--- a/Engine/Plugins/Experimental/MLDeformer/MLDeformerFramework/Source/MLDeformerFramework/Public/MLDeformerComponent.h
+++ b/Engine/Plugins/Experimental/MLDeformer/MLDeformerFramework/Source/MLDeformerFramework/Public/MLDeformerComponent.h
@@ -49,7 +49,7 @@ public:
 	 * Values can be anything between 0 and 1.
 	 */
 	UFUNCTION(BlueprintCallable, Category = "MLDeformer")
-	float GetWeight() const { return Weight; }
+	float GetWeight() const										{ return Weight; }
 
 	/**
 	 * Set the ML Deformer weight. This determines how active the deformer is. You can see it as a blend weight.
@@ -60,16 +60,41 @@ public:
 	 * @param NormalizedWeightValue The weight value.
 	 */
 	UFUNCTION(BlueprintCallable, Category = "MLDeformer")
-	void SetWeight(float NormalizedWeightValue) { Weight = FMath::Clamp<float>(NormalizedWeightValue, 0.0f, 1.0f);  }
+	void SetWeight(float NormalizedWeightValue)					{ Weight = FMath::Clamp<float>(NormalizedWeightValue, 0.0f, 1.0f);  }
 
 #if WITH_EDITOR
 	void PostEditChangeProperty(FPropertyChangedEvent& PropertyChangedEvent) override;
 #endif
 
-	UMLDeformerAsset* GetDeformerAsset() { return DeformerAsset; }
-	void SetDeformerAsset(UMLDeformerAsset* InDeformerAsset) { DeformerAsset = InDeformerAsset; }
-	UMLDeformerModelInstance* GetModelInstance() const { return ModelInstance; }
-	USkeletalMeshComponent* GetSkeletalMeshComponent() const { return SkelMeshComponent.Get(); }
+	/**
+	 * Get the ML Deformer asset that is used by this component.
+	 * @return A pointer to the ML Deformer asset.
+	 */
+	UMLDeformerAsset* GetDeformerAsset() const					{ return DeformerAsset; }
+
+	/**
+	 * Set the deformer asset that is used by this component.
+	 * @param InDeformerAsset A pointer to the deformer asset.
+	 */
+	void SetDeformerAsset(UMLDeformerAsset* InDeformerAsset)	{ DeformerAsset = InDeformerAsset; }
+
+	/**
+	 * Get the ML Deformer model instance that this component currently uses.
+	 * The instance is responsible for running inference and feeding the neural network with inputs.
+	 * @return A pointer to the model instance object.
+	 */
+	UMLDeformerModelInstance* GetModelInstance() const			{ return ModelInstance; }
+
+	/**
+	 * Get the skeletal mesh component that the ML Deformer will work on.
+	 * The skeletal mesh that is setup inside the skeletal mesh component will be the mesh that will be deformed by this ML Deformer component.
+	 * @return A pointer to the skeletal mesh component who's mesh will be deformed by this ML Deformer component.
+	 */
+	USkeletalMeshComponent* GetSkeletalMeshComponent() const	{ return SkelMeshComponent.Get(); }
+
+	// Get property names.
+	static FName GetDeformerAssetPropertyName()					{ return GET_MEMBER_NAME_CHECKED(UMLDeformerComponent, DeformerAsset); }
+	static FName GetWeightPropertyName()						{ return GET_MEMBER_NAME_CHECKED(UMLDeformerComponent, Weight); }
 
 protected:
 	// AActorComponent overrides.
@@ -101,7 +126,6 @@ protected:
 	/** DelegateHandle for NeuralNetwork modification. */
 	FDelegateHandle NeuralNetworkModifyDelegateHandle;
 
-public:
 	/** The deformer asset to use. */
 	UPROPERTY(EditAnywhere, DisplayName = "ML Deformer Asset", Category = "ML Deformer")
 	TObjectPtr<UMLDeformerAsset> DeformerAsset = nullptr;
diff --git a/Engine/Plugins/Experimental/MLDeformer/MLDeformerFramework/Source/MLDeformerFramework/Public/MLDeformerGeomCacheModel.h b/Engine/Plugins/Experimental/MLDeformer/MLDeformerFramework/Source/MLDeformerFramework/Public/MLDeformerGeomCacheModel.h
--- a/Engine/Plugins/Experimental/MLDeformer/MLDeformerFramework/Source/MLDeformerFramework/Public/MLDeformerGeomCacheModel.h
+++ b/Engine/Plugins/Experimental/MLDeformer/MLDeformerFramework/Source/MLDeformerFramework/Public/MLDeformerGeomCacheModel.h
@@ -1,4 +1,5 @@
 // Copyright Epic Games, Inc. All Rights Reserved.
+
 #pragma once
 
 #include "CoreMinimal.h"
@@ -6,7 +7,7 @@
 #include "MLDeformerVizSettings.h"
 #include "MLDeformerInputInfo.h"
 #include "MLDeformerGeomCacheHelpers.h"
-#include "UObject/Object.h"
+#include "UObject/ObjectPtr.h"
 #include "MLDeformerGeomCacheModel.generated.h"
 
 class UMLDeformerAsset;
@@ -64,16 +65,12 @@ public:
 	 * @return The geometry cache mesh to skeletal mesh mappings.
 	 */
 	TArray<UE::MLDeformer::FMLDeformerGeomCacheMeshMapping>& GetGeomCacheMeshMappings() { return MeshMappings; }
-
-	/**
-	 * Get the mapping between geometry cache tracks and meshes inside the skeletal mesh.
-	 * This lets us know what parts of the skeletal mesh are related to what geometry cache tracks.
-	 * Once we have that, we can calculate deltas between the two.
-	 * @return The geometry cache mesh to skeletal mesh mappings, in read-only mode.
-	 */
 	const TArray<UE::MLDeformer::FMLDeformerGeomCacheMeshMapping>& GetGeomCacheMeshMappings() const { return MeshMappings; }
 
-public:
+	// Get property names.
+	static FName GetGeometryCachePropertyName() { return GET_MEMBER_NAME_CHECKED(UMLDeformerGeomCacheModel, GeometryCache); }
+
+protected:
 	/** The mappings between the geometry cache tracks and skeletal mesh imported meshes. */
 	TArray<UE::MLDeformer::FMLDeformerGeomCacheMeshMapping> MeshMappings;
 
diff --git a/Engine/Plugins/Experimental/MLDeformer/MLDeformerFramework/Source/MLDeformerFramework/Public/MLDeformerGeomCacheVizSettings.h b/Engine/Plugins/Experimental/MLDeformer/MLDeformerFramework/Source/MLDeformerFramework/Public/MLDeformerGeomCacheVizSettings.h
--- a/Engine/Plugins/Experimental/MLDeformer/MLDeformerFramework/Source/MLDeformerFramework/Public/MLDeformerGeomCacheVizSettings.h
+++ b/Engine/Plugins/Experimental/MLDeformer/MLDeformerFramework/Source/MLDeformerFramework/Public/MLDeformerGeomCacheVizSettings.h
@@ -3,7 +3,7 @@
 #pragma once
 
 #include "CoreMinimal.h"
-#include "UObject/Object.h"
+#include "UObject/ObjectPtr.h"
 #include "MLDeformerVizSettings.h"
 #include "MLDeformerGeomCacheVizSettings.generated.h"
 
@@ -22,16 +22,19 @@ class MLDEFORMERFRAMEWORK_API UMLDeformerGeomCacheVizSettings
 public:
 #if WITH_EDITORONLY_DATA
 	// UMLDeformerModel overrides.
-	virtual bool HasTestGroundTruth() const override { return !GroundTruth.IsNull(); }
+	virtual bool HasTestGroundTruth() const override	{ return !GroundTruth.IsNull(); }
 	// ~END UMLDeformerModel overrides.
 
 	/**
 	 * Get the test ground truth geometry cache, which represents the ground truth version of the test animation sequence.
 	 * @return A pointer to the geometry cache object.
 	 */
-	UGeometryCache* GetTestGroundTruth() const { return GroundTruth; }
+	UGeometryCache* GetTestGroundTruth() const			{ return GroundTruth; }
 
-public:
+	// Get property names.
+	static FName GetTestGroundTruthPropertyName()		{ return GET_MEMBER_NAME_CHECKED(UMLDeformerGeomCacheVizSettings, GroundTruth); }
+
+protected:
 	/** The geometry cache that represents the ground truth of the test anim sequence. */
 	UPROPERTY(EditAnywhere, Category = "Test Assets")
 	TObjectPtr<UGeometryCache> GroundTruth = nullptr;
diff --git a/Engine/Plugins/Experimental/MLDeformer/MLDeformerFramework/Source/MLDeformerFramework/Public/MLDeformerInputInfo.h b/Engine/Plugins/Experimental/MLDeformer/MLDeformerFramework/Source/MLDeformerFramework/Public/MLDeformerInputInfo.h
--- a/Engine/Plugins/Experimental/MLDeformer/MLDeformerFramework/Source/MLDeformerFramework/Public/MLDeformerInputInfo.h
+++ b/Engine/Plugins/Experimental/MLDeformer/MLDeformerFramework/Source/MLDeformerFramework/Public/MLDeformerInputInfo.h
@@ -28,7 +28,7 @@ public:
 	 * This method is executed post loading.
 	 * On default it will initialize the FNames based on the bone and curve string based names.
 	 */
-	virtual void OnPostLoad() { UpdateFNames(); }
+	virtual void OnPostLoad();
 
 	/**
 	 * Check whether the current inputs are compatible with a given skeletal mesh.
@@ -63,7 +63,7 @@ public:
 	 * This happens when there are no bones or curves to use as inputs.
 	 * @return Returns true when there are no bones or curves specified as inputs.
 	 */
-	virtual bool IsEmpty() const { return (BoneNameStrings.IsEmpty() && CurveNameStrings.IsEmpty()); }
+	virtual bool IsEmpty() const;
 
 	/**
 	 * Calculate how many inputs this input info generates for the neural network.
@@ -76,65 +76,65 @@ public:
 	 * Get the number of bones that we trained on.
 	 * @return The number of bones.
 	 */
-	int32 GetNumBones() const { return BoneNames.Num(); }
+	int32 GetNumBones() const;
 
 	/**
 	 * Get the bone name as a string, for a given bone we included during training.
 	 * @param Index The bone index, which is a number in range of [0..GetNumBones()-1].
 	 * @result The name of the bone.
 	 */
-	const FString& GetBoneNameString(int32 Index) const { return BoneNameStrings[Index]; }
+	const FString& GetBoneNameString(int32 Index) const;
 
 	/**
 	 * Get the bone name as an FName, for a given bone we included during training.
 	 * @param Index The bone index, which is a number in range of [0..GetNumBones()-1].
 	 * @result The name of the bone.
 	 */
-	const FName GetBoneName(int32 Index) const { return BoneNames[Index]; }
+	const FName GetBoneName(int32 Index) const;
 
 	/**
 	 * Get the number of curves that we trained on.
 	 * @return The number of curves.
 	 */
-	int32 GetNumCurves() const { return CurveNames.Num(); }
+	int32 GetNumCurves() const;
 
 	/**
 	 * Get the list of bone names that should be included as training inputs.
 	 * @return An array of strings, one for each bone name.
 	 */
-	TArray<FString>& GetBoneNameStrings() { return BoneNameStrings; }
+	TArray<FString>& GetBoneNameStrings();
 
 	/**
 	 * Get the list of curve names that should be included as training inputs.
 	 * @return An array of strings, one for each curve name.
 	 */
-	TArray<FString>& GetCurveNameStrings() { return CurveNameStrings; }
+	TArray<FString>& GetCurveNameStrings();
 
 	/**
 	 * Get the list of bone names that should be included as training inputs.
 	 * @return An array of FName objects, one for each bone.
 	 */
-	TArray<FName>& GetBoneNames() { return BoneNames; }
+	TArray<FName>& GetBoneNames();
 
 	/**
 	 * Get the list of curve names that should be included as training inputs.
 	 * @return An array of FName objects, one for each curve.
 	 */
-	TArray<FName>& GetCurveNames() { return CurveNames; }
+	TArray<FName>& GetCurveNames();
 
 	/**
 	 * Get the curve name as a string, for a given bone we included during training.
 	 * @param Index The curve index, which is a number in range of [0..GetNumCurves()-1].
 	 * @result The name of the curve.
 	 */
-	const FString& GetCurveNameString(int32 Index) const { return CurveNameStrings[Index]; }
+	const FString& GetCurveNameString(int32 Index) const;
 
 	/**
 	 * Get the curve name as an FName, for a given curve we included during training.
 	 * @param Index The curve index, which is a number in range of [0..GetNumCurves()-1].
 	 * @result The name of the curve.
 	 */
-	const FName GetCurveName(int32 Index) const { return CurveNames[Index]; }
+	const FName GetCurveName(int32 Index) const;
 
 	/** 
 	 * Extract the curve values for all curves we're interested in.
@@ -156,28 +156,28 @@ public:
 	 * Get the number of imported vertices in the base mesh, which is the linear skinned skeletal mesh.
 	 * @return The number of imported vertices in the base mesh.
 	 */
-	int32 GetNumBaseMeshVertices() const { return NumBaseMeshVertices; }
+	int32 GetNumBaseMeshVertices() const;
 
 	/**
 	 * Set the number of vertices of our base mesh, which is the linear skinned skeletal mesh.
 	 * The vertex count must be the same as in the DCC, so not the render mesh vertex count.
 	 * @param NumVerts The number of vertices of our base mesh.
 	 */
-	void SetNumBaseVertices(int32 NumVerts) { NumBaseMeshVertices = NumVerts; }
+	void SetNumBaseVertices(int32 NumVerts);
 
 	/**
 	 * Get the number of imported vertices in the target mesh, which is our training target/ground truth mesh.
 	 * The number of vertices is the vertex count as in the DCC.
 	 * @return The number of imported vertices in the target mesh.
 	 */
-	int32 GetNumTargetMeshVertices() const { return NumTargetMeshVertices; }
+	int32 GetNumTargetMeshVertices() const;
 
 	/**
 	 * Set the number of target mesh vertices. The target mesh is our training target/ground truth mesh.
 	 * The vertex count must be the same as in the DCC, so not the render mesh vertex count.
 	 * @param NumVerts The number of vertices in the target mesh.
 	 */
-	void SetNumTargetVertices(int32 NumVerts) { NumTargetMeshVertices = NumVerts; }
+	void SetNumTargetVertices(int32 NumVerts);
 
 protected:
 	/** 
diff --git a/Engine/Plugins/Experimental/MLDeformer/MLDeformerFramework/Source/MLDeformerFramework/Public/MLDeformerModel.h b/Engine/Plugins/Experimental/MLDeformer/MLDeformerFramework/Source/MLDeformerFramework/Public/MLDeformerModel.h
--- a/Engine/Plugins/Experimental/MLDeformer/MLDeformerFramework/Source/MLDeformerFramework/Public/MLDeformerModel.h
+++ b/Engine/Plugins/Experimental/MLDeformer/MLDeformerFramework/Source/MLDeformerFramework/Public/MLDeformerModel.h
@@ -1,11 +1,10 @@
 // Copyright Epic Games, Inc. All Rights Reserved.
+
 #pragma once
 
 #include "CoreMinimal.h"
-#include "MLDeformerInputInfo.h"
 #include "UObject/ObjectPtr.h"
 #include "UObject/Object.h"
-#include "Templates/UniquePtr.h"
 #include "Interfaces/Interface_BoneReferenceSkeletonProvider.h"
 #include "BoneContainer.h"
 #include "RenderCommandFence.h"
@@ -17,6 +16,7 @@ class UMLDeformerAsset;
 class UMLDeformerVizSettings;
 class UMLDeformerModelInstance;
 class UMLDeformerComponent;
+class UMLDeformerInputInfo;
 class UAnimSequence;
 class UGeometryCache;
 class UNeuralNetwork;
@@ -29,7 +29,8 @@ namespace UE::MLDeformer
 	 * So if a cube requires 32 render vertices, there will be 32 ints inside this buffer, and each item in this buffer
 	 * will in this specific example case contain a value between 0 and 7, as a cube has only 8 vertices.
 	 */
-	class FVertexMapBuffer : public FVertexBufferWithSRV
+	class FVertexMapBuffer
+		: public FVertexBufferWithSRV
 	{
 	public:
 		/**
@@ -64,6 +65,8 @@ class MLDEFORMERFRAMEWORK_API UMLDeformerModel
 	GENERATED_BODY()
 
 public:
+	DECLARE_MULTICAST_DELEGATE(FNeuralNetworkModifyDelegate);
+
 	virtual ~UMLDeformerModel() = default;
 
 	/**
@@ -112,14 +115,14 @@ public:
 	 * On default this is set to return true as most models have bone rotations as inputs to the neural network.
 	 * @result Returns true when this model supports bones, or false otherwise.
 	 */
-	virtual bool DoesSupportBones() const { return true; }
+	virtual bool DoesSupportBones() const					{ return true; }
 
 	/**
 	 * Defines whether this model supports curves as inputs or not. A curve is just a single float value.
 	 * On default this returns true.
 	 * @result Returns true when this model supports curves, or false otherwise.
 	 */
-	virtual bool DoesSupportCurves() const { return true; }
+	virtual bool DoesSupportCurves() const					{ return true; }
 
 	/**
 	 * Check whether the neural network of this model should run on the GPU or not.
@@ -127,7 +130,7 @@ public:
 	 * Some code internally that creates and initializes the neural network will use the return value of this method to mark it to be on CPU or GPU.
 	 * @return Returns true if the neural network of this model should run on the GPU. False is returned when it should run on the CPU.
 	 */
-	virtual bool IsNeuralNetworkOnGPU() const { return false; }	// CPU neural network.
+	virtual bool IsNeuralNetworkOnGPU() const				{ return false; }	// CPU neural network.
 
 #if WITH_EDITORONLY_DATA
 	/**
@@ -135,7 +138,7 @@ public:
 	 * For example if there is say a GeometryCache as target mesh, it could check whether that property is a nullptr or not.
 	 * @return Returns true when the training target mesh has been selected, otherwise false is returned.
 	 */
-	virtual bool HasTrainingGroundTruth() const { return false; }
+	virtual bool HasTrainingGroundTruth() const				{ return false; }
 
 	/**
 	 * Sample the positions from the target (ground truth) mesh, at a specific time (in seconds).
@@ -207,21 +210,21 @@ public:
 	 * Inputs are things like bone transforms and curve values.
 	 * @return A pointer to the input info object.
 	 */
-	UMLDeformerInputInfo* GetInputInfo() const { return InputInfo.Get(); }
+	UMLDeformerInputInfo* GetInputInfo() const				{ return InputInfo.Get(); }
 
 	/**
 	 * Get the number of vertices in the base mesh (linear skinned skeletal mesh).
 	 * This is the number of vertices in the DCC, so not the render vertex count.
 	 * @return The number of imported vertices inside linear skinned skeletal mesh.
 	 */
-	int32 GetNumBaseMeshVerts() const { return NumBaseMeshVerts; }
+	int32 GetNumBaseMeshVerts() const						{ return NumBaseMeshVerts; }
 
 	/**
 	 * Get the number of vertices of the target mesh.
 	 * This is the number of vertices in the DCC, so not the render vertex count.
 	 * @return The number of imported vertices inside the target mesh.
 	 */
-	int32 GetNumTargetMeshVerts() const { return NumTargetMeshVerts; }
+	int32 GetNumTargetMeshVerts() const						{ return NumTargetMeshVerts; }
 
 	/**
 	 * The mapping that maps from render vertices into dcc vertices.
@@ -229,7 +232,7 @@ public:
 	 * For a cube with 32 render vertices, the item values would be between 0..7 as in the dcc the cube has 8 vertices.
 	 * @return A reference to the array that contains the DCC vertex number for each render vertex.
 	 */
-	const TArray<int32>& GetVertexMap() const { return VertexMap; }
+	const TArray<int32>& GetVertexMap() const				{ return VertexMap; }
 
 	/**
 	 * Get the GPU buffer of the VertexMap.
@@ -243,7 +246,7 @@ public:
 	 * This network is used during inference.
 	 * @return A pointer to the neural network, or nullptr when the network has not yet been trained.
 	 */
-	UNeuralNetwork* GetNeuralNetwork() const { return NeuralNetwork.Get(); }
+	UNeuralNetwork* GetNeuralNetwork() const				{ return NeuralNetwork.Get(); }
 
 	/**
 	 * Set the neural network object that we use during inference.
@@ -251,6 +254,13 @@ public:
 	 */
 	void SetNeuralNetwork(UNeuralNetwork* InNeuralNetwork);
 
+	/**
+	 * Get the neural network modified delegate.
+	 * This triggers when the neural network pointers changes.
+	 * @return A reference to the delegate.
+	 */
+	FNeuralNetworkModifyDelegate& GetNeuralNetworkModifyDelegate() { return NeuralNetworkModifyDelegate; }
+
 #if WITH_EDITORONLY_DATA
 	// UObject overrides.
 	virtual void PostEditChangeProperty(FPropertyChangedEvent& PropertyChangedEvent) override;
@@ -268,14 +278,14 @@ public:
 	 * Check whether we should include bone transforms as input to the model during training or not.
 	 * @return Returns true when bone transfomations should be a part of the network inputs, during the training process.
 	 */
-	bool ShouldIncludeBonesInTraining() const { return bIncludeBones; }
+	bool ShouldIncludeBonesInTraining() const				{ return bIncludeBones; }
 
 	/**
 	 * Check whether we should include curve values as input to the model during training or not.
 	 * Curve values are single floats.
 	 * @return Returns true when curve values should be a part of the network inputs, during the training process.
 	 */
-	bool ShouldIncludeCurvesInTraining() const { return bIncludeCurves; }
+	bool ShouldIncludeCurvesInTraining() const				{ return bIncludeCurves; }
 
 	/**
 	 * The delegate that gets fired when a property value changes.
@@ -290,23 +300,29 @@ public:
 	 * @return A pointer to the visualization settings. You can cast this to the type specific for your model, in case you inherited from 
 	 *         the UMLDeformerVizSettings base class. This never return a nullptr.
 	 */
-	UMLDeformerVizSettings* GetVizSettings() const { return VizSettings; }
+	UMLDeformerVizSettings* GetVizSettings() const			{ return VizSettings; }
 
 	/**
 	 * Get the skeletal mesh that is used during training.
 	 * You typically want to apply the ML Deformer on this specific skeletal mesh in your game as well.
 	 * @return A pointer to the skeletal mesh.
 	 */
-	const USkeletalMesh* GetSkeletalMesh() const { return SkeletalMesh;  }
-	USkeletalMesh* GetSkeletalMesh() { return SkeletalMesh; }
+	const USkeletalMesh* GetSkeletalMesh() const			{ return SkeletalMesh;  }
+	USkeletalMesh* GetSkeletalMesh()						{ return SkeletalMesh; }
+
+	/**
+	 * Set the skeletal mesh that this deformer uses.
+	 * @param SkelMesh The skeletal mesh.
+	 */
+	void SetSkeletalMesh(USkeletalMesh* SkelMesh)			{ SkeletalMesh = SkelMesh; }
 
 	/**
 	 * Get the animation sequence that is used during training.
 	 * Each frame of this anim sequence will contain a training pose.
 	 * @return A pointer to the animation sequence used for training.
 	 */
-	const UAnimSequence* GetAnimSequence() const { return AnimSequence;  }
-	UAnimSequence* GetAnimSequence() { return AnimSequence; }
+	const UAnimSequence* GetAnimSequence() const			{ return AnimSequence;  }
+	UAnimSequence* GetAnimSequence()						{ return AnimSequence; }
 
 	/**
 	 * Get the maximum number of training frames to use during training.
@@ -314,7 +330,7 @@ public:
 	 * want to train on only 2000 frames instead. You can do this by setting the maximum training frames to 2000.
 	 * @return The max number of frames to use during training.
 	 */
-	int32 GetTrainingFrameLimit() const { return MaxTrainingFrames; }
+	int32 GetTrainingFrameLimit() const						{ return MaxTrainingFrames; }
 
 	/**
 	 * Get the target mesh alignment tranformation.
@@ -324,7 +340,7 @@ public:
 	 * The alignment transform is then used to correct this and align both base and target mesh.
 	 * @return The alignment transformation. When set to Identity it will not do anything, which is its default.
 	 */
-	const FTransform& GetAlignmentTransform() const { return AlignmentTransform; }
+	const FTransform& GetAlignmentTransform() const			{ return AlignmentTransform; }
 
 	/**
 	 * Get the list of bones that we configured to be included during training.
@@ -336,8 +352,8 @@ public:
 	 * will contain different items.
 	 * @return The list of bones that should be included when performing the next training session.
 	 */
-	TArray<FBoneReference>& GetBoneIncludeList() { return BoneIncludeList; }
-	const TArray<FBoneReference>& GetBoneIncludeList() const { return BoneIncludeList; }
+	TArray<FBoneReference>& GetBoneIncludeList()				{ return BoneIncludeList; }
+	const TArray<FBoneReference>& GetBoneIncludeList() const	{ return BoneIncludeList; }
 
 	/**
 	 * Get the list of curves that we configured to be included during training.
@@ -349,8 +365,8 @@ public:
 	 * will contain different items.
 	 * @return The list of curves that should be included when performing the next training session.
 	 */
-	TArray<FMLDeformerCurveReference>& GetCurveIncludeList() { return CurveIncludeList; }
-	const TArray<FMLDeformerCurveReference>& GetCurveIncludeList() const { return CurveIncludeList; }
+	TArray<FMLDeformerCurveReference>& GetCurveIncludeList()				{ return CurveIncludeList; }
+	const TArray<FMLDeformerCurveReference>& GetCurveIncludeList() const	{ return CurveIncludeList; }
 
 	/**
 	 * Get the delta cutoff length. Deltas that have a length larger than this length will be set to zero.
@@ -358,22 +374,27 @@ public:
 	 * Skipping those deltas will prevent issues.
 	 * @return The length after which deltas will be ignored. So anything delta length larger than this value will be ignored.
 	 */
-	float GetDeltaCutoffLength() const { return DeltaCutoffLength; }
+	float GetDeltaCutoffLength() const					{ return DeltaCutoffLength; }
+
+	// Get property names.
+	static FName GetSkeletalMeshPropertyName()			{ return GET_MEMBER_NAME_CHECKED(UMLDeformerModel, SkeletalMesh); }
+	static FName GetShouldIncludeBonesPropertyName()	{ return GET_MEMBER_NAME_CHECKED(UMLDeformerModel, bIncludeBones); }
+	static FName GetShouldIncludeCurvesPropertyName()	{ return GET_MEMBER_NAME_CHECKED(UMLDeformerModel, bIncludeCurves); }
+	static FName GetAnimSequencePropertyName()			{ return GET_MEMBER_NAME_CHECKED(UMLDeformerModel, AnimSequence); }
+	static FName GetAlignmentTransformPropertyName()	{ return GET_MEMBER_NAME_CHECKED(UMLDeformerModel, AlignmentTransform); }
+	static FName GetBoneIncludeListPropertyName()		{ return GET_MEMBER_NAME_CHECKED(UMLDeformerModel, BoneIncludeList); }
+	static FName GetCurveIncludeListPropertyName()		{ return GET_MEMBER_NAME_CHECKED(UMLDeformerModel, CurveIncludeList); }
+	static FName GetMaxTrainingFramesPropertyName()		{ return GET_MEMBER_NAME_CHECKED(UMLDeformerModel, MaxTrainingFrames); }
+	static FName GetDeltaCutoffLengthPropertyName()		{ return GET_MEMBER_NAME_CHECKED(UMLDeformerModel, DeltaCutoffLength); }
 #endif	// #if WITH_EDITORONLY_DATA
 
 protected:
-	/** The deformer asset that this model is part of. */
-	TObjectPtr<UMLDeformerAsset> DeformerAsset = nullptr;
-
-	/** The delegate that gets fired when a property is being modified. */
-	FMLDeformerModelOnPostEditProperty PostEditPropertyDelegate;
-
 	/**
 	 * Set the training input information.
 	 * @param Input A pointer to the input information object.
 	 * @see CreateInputInfo
 	 */
-	void SetInputInfo(UMLDeformerInputInfo* Input) { InputInfo = Input; }
+	void SetInputInfo(UMLDeformerInputInfo* Input)		{ InputInfo = Input; }
 
 	/**
 	 * Convert an array of floats to an array of Vector3's.
@@ -385,7 +406,22 @@ protected:
 	 */
 	void FloatArrayToVector3Array(const TArray<float>& FloatArray, TArray<FVector3f>& OutVectorArray);
 
-public:
+protected:
+	/** The deformer asset that this model is part of. */
+	TObjectPtr<UMLDeformerAsset> DeformerAsset = nullptr;
+
+	/** The delegate that gets fired when a property is being modified. */
+	FMLDeformerModelOnPostEditProperty PostEditPropertyDelegate;
+
+	/** GPU buffers for Vertex Map. */
+	UE::MLDeformer::FVertexMapBuffer VertexMapBuffer;
+
+	/** Fence used in render thread cleanup on destruction. */
+	FRenderCommandFence RenderResourceDestroyFence;
+
+	/** Delegate that will be called immediately before the NeuralNetwork is changed. */
+	FNeuralNetworkModifyDelegate NeuralNetworkModifyDelegate;
+
 	/** Cached number of skeletal mesh vertices. */
 	UPROPERTY()
 	int32 NumBaseMeshVerts = 0;
@@ -408,16 +444,6 @@ public:
 	UPROPERTY()
 	TObjectPtr<UNeuralNetwork> NeuralNetwork = nullptr;
 
-	/** GPU buffers for Vertex Map. */
-	UE::MLDeformer::FVertexMapBuffer VertexMapBuffer;
-
-	/** Fence used in render thread cleanup on destruction. */
-	FRenderCommandFence RenderResourceDestroyFence;
-
-	/** Delegate that will be called immediately before the NeuralNetwork is changed. */
-	DECLARE_MULTICAST_DELEGATE(FNeuralNetworkModifyDelegate);
-	FNeuralNetworkModifyDelegate NeuralNetworkModifyDelegate;
-
 #if WITH_EDITORONLY_DATA
 	UPROPERTY()
 	TObjectPtr<UMLDeformerVizSettings> VizSettings = nullptr;
diff --git a/Engine/Plugins/Experimental/MLDeformer/MLDeformerFramework/Source/MLDeformerFramework/Public/MLDeformerModelInstance.h b/Engine/Plugins/Experimental/MLDeformer/MLDeformerFramework/Source/MLDeformerFramework/Public/MLDeformerModelInstance.h
--- a/Engine/Plugins/Experimental/MLDeformer/MLDeformerFramework/Source/MLDeformerFramework/Public/MLDeformerModelInstance.h
+++ b/Engine/Plugins/Experimental/MLDeformer/MLDeformerFramework/Source/MLDeformerFramework/Public/MLDeformerModelInstance.h
@@ -1,4 +1,5 @@
 // Copyright Epic Games, Inc. All Rights Reserved.
+
 #pragma once
 
 #include "CoreMinimal.h"
@@ -10,6 +11,14 @@ class UMLDeformerModel;
 class UNeuralNetwork;
 class USkeletalMeshComponent;
 
+/**
+ * An instance of the ML Deformer model.
+ * The ML Deformer model contains shared data, while this instance contains data unique to the actor it is being applied to.
+ * So if you have 1 ML Deformer model, applied to 100 enemy actors, you will get 100 ML Deformer Model Instance objects.
+ * The model instance is created and managed by the ML Deformer component.
+ * The main task of the model instance is to feed the neural network with the right input values and run the neural network
+ * inference at runtime, and possibly do something with the output.
+ */
 UCLASS()
 class MLDEFORMERFRAMEWORK_API UMLDeformerModelInstance
 	: public UObject
@@ -96,45 +105,46 @@ public:
 	 * Get the current component space bone transforms that we grabbed from the skeletal mesh component.
 	 * @return The cached bone transforms, in component space.
 	 */
-	const TArray<FTransform>& GetBoneTransforms() const { return BoneTransforms; }
+	const TArray<FTransform>& GetBoneTransforms() const;
 
 	/**
 	 * Is the deformer asset used compatible with the skeletal mesh component used during the Init call?
 	 * @return True if compatible, false if not.
 	 */
-	bool IsCompatible() const { return bIsCompatible; }
+	bool IsCompatible() const;
 
 	/** 
 	 * Get the compatibility error text. This will be a non-empty string in case IsCompatible() return false. 
 	 * @return Returns the string containing more information about the reasons for compatibility issues.
 	 */
-	const FString& GetCompatibilityErrorText() const { return ErrorText; }
+	const FString& GetCompatibilityErrorText() const;
 
 	/**
 	 * Get the skeletal mesh component we're working with.
 	 * @return A pointer to the skeletal mesh component.
 	 */
-	USkeletalMeshComponent* GetSkeletalMeshComponent() const { return SkeletalMeshComponent; }
+	USkeletalMeshComponent* GetSkeletalMeshComponent() const;
 
 	/** Update the compatibility status, as returned by IsCompatible() and GetCompatibilityErrorText(). */
 	void UpdateCompatibilityStatus();
 
 	/** Get the model that this is an instance of. */
-	UMLDeformerModel* GetModel() const { return Model.Get(); }
+	UMLDeformerModel* GetModel() const;
 
 	/** Set the deformer model that this is an instance of. */
-	void SetModel(UMLDeformerModel* InModel) { Model = InModel; }
+	void SetModel(UMLDeformerModel* InModel);
 
 	/**
 	 * Get the neural network inference handle.
 	 * @return Returns the handle. This is -1 for an invalid handle.
 	 */
-	int32 GetNeuralNetworkInferenceHandle() const { return NeuralNetworkInferenceHandle; }
+	int32 GetNeuralNetworkInferenceHandle() const;
 
 	/** Set whether we had a successful call to UMLDeformerModel::PostMLDeformerComponentInit. */
-	void SetHasPostInitialized(bool bHasInitialized) { bHasPostInitialized = bHasInitialized; }
+	void SetHasPostInitialized(bool bHasInitialized);
 
-	bool HasPostInitialized() const { return bHasPostInitialized; }
+	/** Check whether we already called UMLDeformerModel::PostMLDeformerComponentInit. */
+	bool HasPostInitialized() const;
 
 protected:
 	/**
@@ -160,12 +170,11 @@ protected:
 	 */
 	void UpdateBoneTransforms();
 
-public:
+protected:
 	/** The ML Deformer model that this is an instance of. */
 	UPROPERTY(Transient)
 	TObjectPtr<UMLDeformerModel> Model = nullptr;
 
-protected:
 	/** The skeletal mesh component we work with. This is mainly used for compatibility checks. */
 	TObjectPtr<USkeletalMeshComponent> SkeletalMeshComponent = nullptr;
 
diff --git a/Engine/Plugins/Experimental/MLDeformer/MLDeformerFramework/Source/MLDeformerFramework/Public/MLDeformerModule.h b/Engine/Plugins/Experimental/MLDeformer/MLDeformerFramework/Source/MLDeformerFramework/Public/MLDeformerModule.h
--- a/Engine/Plugins/Experimental/MLDeformer/MLDeformerFramework/Source/MLDeformerFramework/Public/MLDeformerModule.h
+++ b/Engine/Plugins/Experimental/MLDeformer/MLDeformerFramework/Source/MLDeformerFramework/Public/MLDeformerModule.h
@@ -5,15 +5,20 @@
 #include "CoreMinimal.h"
 #include "Modules/ModuleManager.h"
 
+// The log category for the ML Deformer framework.
 MLDEFORMERFRAMEWORK_API DECLARE_LOG_CATEGORY_EXTERN(LogMLDeformer, Log, All);
 
 namespace UE::MLDeformer
 {
-	class MLDEFORMERFRAMEWORK_API FMLDeformerModule : public IModuleInterface
+	/**
+	 * The runtime module for the ML Deformer.
+	 */
+	class MLDEFORMERFRAMEWORK_API FMLDeformerModule
+		: public IModuleInterface
 	{
 	public:
 		// IModuleInterface overrides.
 		virtual void StartupModule() override;
 		// ~END IModuleInterface overrides.
 	};
-}
+}	// namespace UE::MLDeformer
diff --git a/Engine/Plugins/Experimental/MLDeformer/MLDeformerFramework/Source/MLDeformerFramework/Public/MLDeformerMorphModel.h b/Engine/Plugins/Experimental/MLDeformer/MLDeformerFramework/Source/MLDeformerFramework/Public/MLDeformerMorphModel.h
--- a/Engine/Plugins/Experimental/MLDeformer/MLDeformerFramework/Source/MLDeformerFramework/Public/MLDeformerMorphModel.h
+++ b/Engine/Plugins/Experimental/MLDeformer/MLDeformerFramework/Source/MLDeformerFramework/Public/MLDeformerMorphModel.h
@@ -1,14 +1,11 @@
 // Copyright Epic Games, Inc. All Rights Reserved.
+
 #pragma once
 
 #include "CoreMinimal.h"
 #include "MLDeformerGeomCacheModel.h"
-#include "UObject/Object.h"
 #include "MLDeformerMorphModel.generated.h"
 
-class USkeletalMesh;
-class UMLDeformerAsset;
-class UMLDeformerModelInstance;
 struct FExternalMorphSet;
 
 UCLASS()
@@ -21,10 +18,10 @@ public:
 	UMLDeformerMorphModel(const FObjectInitializer& ObjectInitializer);
 
 	// UMLDeformerModel overrides.
-	virtual FString GetDisplayName() const override { return "Morph Base Model"; }
+	virtual FString GetDisplayName() const override			{ return "Morph Base Model"; }
+	virtual bool IsNeuralNetworkOnGPU() const override		{ return false; }	// CPU based neural network.
 	virtual void Serialize(FArchive& Archive) override;
 	virtual void PostMLDeformerComponentInit(UMLDeformerModelInstance* ModelInstance) override;
-	virtual bool IsNeuralNetworkOnGPU() const override { return false; }	// CPU neural network.
 	virtual UMLDeformerModelInstance* CreateModelInstance(UMLDeformerComponent* Component) override;
 	// ~END UMLDeformerModel overrides.
 
@@ -33,16 +30,40 @@ public:
 	// ~END UObject overrides.
 
 #if WITH_EDITORONLY_DATA
-	float GetMorphTargetDeltaThreshold() const { return MorphTargetDeltaThreshold; }
+	float GetMorphTargetDeltaThreshold() const				{ return MorphTargetDeltaThreshold; }
+	float GetMorphTargetErrorTolerance() const				{ return MorphTargetErrorTolerance; }
+
+	// Get property names.
+	static FName GetMorphTargetDeltaThresholdPropertyName() { return GET_MEMBER_NAME_CHECKED(UMLDeformerMorphModel, MorphTargetDeltaThreshold); }
+	static FName GetMorphTargetErrorTolerancePropertyName() { return GET_MEMBER_NAME_CHECKED(UMLDeformerMorphModel, MorphTargetErrorTolerance); }
 #endif
 
-	UFUNCTION(BlueprintCallable, Category = "NeuralMorphModel")
-	void SetMorphTargetDeltas(const TArray<float>& Deltas);
+	/**
+	 * Set the per vertex deltas, as a set of floats. Each vertex delta must have 3 floats.
+	 * Concatenate all deltas into one buffer, so like this [morphdeltas_target0, morphdeltas_target1, ..., morphdeltas_targetN].
+	 * The vertex ordering should be: [(x, y, z), (x, y, z), (x, y, z)].
+	 * @param Deltas The array of floats that contains the deltas. The number of items in the array must be equal to (NumMorphs * NumBaseMeshVerts * 3).
+	 */
+	UFUNCTION(BlueprintCallable, Category = "MLDeformerMorphModel")
+	void SetMorphTargetDeltaFloats(const TArray<float>& Deltas);
+
+	UFUNCTION(BlueprintCallable, Category = "MLDeformerMorphModel")
+	void SetMorphTargetDeltas(const TArray<FVector3f>& Deltas);
 
-	const TArray<FVector3f>& GetMorphTargetDeltas() const { return MorphTargetDeltas; }
+	const TArray<FVector3f>& GetMorphTargetDeltas() const	{ return MorphTargetDeltas; }
+	int32 GetExternalMorphSetID() const						{ return ExternalMorphSetID; }
+	void SetExternalMorphSetID(int32 ID)					{ check(ID != -1);  ExternalMorphSetID = ID; }
+	TSharedPtr<FExternalMorphSet> GetMorphTargetSet() const { return MorphTargetSet; }
+
+	/**
+	 * Get the start index into the array of deltas, for a given morph target.
+	 * This does not perform a bounds check to see if MorphTargetIndex is in a valid range, so be aware.
+	 * @param MorphTargetIndex The morph target index.
+	 * @return The start index, or INDEX_NONE in case there are no deltas.
+	 */
 	int32 GetMorphTargetDeltaStartIndex(int32 MorphTargetIndex) const;
 
-public:
+protected:
 	/** The compressed morph target data, ready for the GPU. */
 	TSharedPtr<FExternalMorphSet> MorphTargetSet;
 
diff --git a/Engine/Plugins/Experimental/MLDeformer/MLDeformerFramework/Source/MLDeformerFramework/Public/MLDeformerMorphModelInstance.h b/Engine/Plugins/Experimental/MLDeformer/MLDeformerFramework/Source/MLDeformerFramework/Public/MLDeformerMorphModelInstance.h
--- a/Engine/Plugins/Experimental/MLDeformer/MLDeformerFramework/Source/MLDeformerFramework/Public/MLDeformerMorphModelInstance.h
+++ b/Engine/Plugins/Experimental/MLDeformer/MLDeformerFramework/Source/MLDeformerFramework/Public/MLDeformerMorphModelInstance.h
@@ -5,7 +5,14 @@
 #include "MLDeformerModelInstance.h"
 #include "MLDeformerMorphModelInstance.generated.h"
 
-
+/**
+ * The model instance for the UMLDeformerMorphModel.
+ * This instance will assume the neural network outputs a set of weights, one for each morph target.
+ * The weights of the morph targets in the external morph target set related to the ID of the model will
+ * be set to the weights that the neural network outputs.
+ * The first morph target contains the means, which need to always be added to the results. Therefore the 
+ * weight of the first morph target will always be forced to a value of 1.
+ */
 UCLASS()
 class MLDEFORMERFRAMEWORK_API UMLDeformerMorphModelInstance
 	: public UMLDeformerModelInstance
diff --git a/Engine/Plugins/Experimental/MLDeformer/MLDeformerFramework/Source/MLDeformerFramework/Public/MLDeformerMorphModelVizSettings.h b/Engine/Plugins/Experimental/MLDeformer/MLDeformerFramework/Source/MLDeformerFramework/Public/MLDeformerMorphModelVizSettings.h
--- a/Engine/Plugins/Experimental/MLDeformer/MLDeformerFramework/Source/MLDeformerFramework/Public/MLDeformerMorphModelVizSettings.h
+++ b/Engine/Plugins/Experimental/MLDeformer/MLDeformerFramework/Source/MLDeformerFramework/Public/MLDeformerMorphModelVizSettings.h
@@ -3,12 +3,11 @@
 #pragma once
 
 #include "CoreMinimal.h"
-#include "UObject/Object.h"
 #include "MLDeformerGeomCacheVizSettings.h"
 #include "MLDeformerMorphModelVizSettings.generated.h"
 
 /**
- * The vizualization settings specific to this model.
+ * The vizualization settings specific to the UMLDeformerMorphModel class, or inherited classes.
  */
 UCLASS()
 class MLDEFORMERFRAMEWORK_API UMLDeformerMorphModelVizSettings
@@ -18,6 +17,17 @@ class MLDEFORMERFRAMEWORK_API UMLDeformerMorphModelVizSettings
 
 public:
 #if WITH_EDITORONLY_DATA
+	void SetMorphTargetNumber(int32 MorphIndexToPreview)	{ MorphTargetNumber = MorphIndexToPreview; }
+	int32 GetMorphTargetNumber() const						{ return MorphTargetNumber; }
+	float GetMorphTargetDeltaThreshold() const				{ return MorphTargetDeltaThreshold; }
+	bool GetDrawMorphTargets() const						{ return bDrawMorphTargets; }
+
+	// Get property names.
+	static FName GetMorphTargetNumberPropertyName()			{ return GET_MEMBER_NAME_CHECKED(UMLDeformerMorphModelVizSettings, MorphTargetNumber); }
+	static FName GetMorphTargetDeltaThresholdPropertyName() { return GET_MEMBER_NAME_CHECKED(UMLDeformerMorphModelVizSettings, MorphTargetDeltaThreshold); }
+	static FName GetDrawMorphTargetsPropertyName()			{ return GET_MEMBER_NAME_CHECKED(UMLDeformerMorphModelVizSettings, bDrawMorphTargets); }
+
+protected:
 	/**
 	 * The morph target to visualize. The first one always being the means, so not a sparse target.
 	 * This only can be used after you trained, in the same editor session directly after training.
diff --git a/Engine/Plugins/Experimental/MLDeformer/MLDeformerFramework/Source/MLDeformerFramework/Public/MLDeformerVizSettings.h b/Engine/Plugins/Experimental/MLDeformer/MLDeformerFramework/Source/MLDeformerFramework/Public/MLDeformerVizSettings.h
--- a/Engine/Plugins/Experimental/MLDeformer/MLDeformerFramework/Source/MLDeformerFramework/Public/MLDeformerVizSettings.h
+++ b/Engine/Plugins/Experimental/MLDeformer/MLDeformerFramework/Source/MLDeformerFramework/Public/MLDeformerVizSettings.h
@@ -36,47 +36,75 @@ enum class EMLDeformerHeatMapMode : uint8
  * The vizualization settings.
  */
 UCLASS()
-class MLDEFORMERFRAMEWORK_API UMLDeformerVizSettings : public UObject
+class MLDEFORMERFRAMEWORK_API UMLDeformerVizSettings
+	: public UObject
 {
 	GENERATED_BODY()
 
 public:
 	// UObject overrides.
-	virtual bool IsEditorOnly() const override { return true; }
+	virtual bool IsEditorOnly() const override				{ return true; }
 #if WITH_EDITOR
 	virtual void PostEditChangeProperty(FPropertyChangedEvent& PropertyChangedEvent) { GetOuter()->PostEditChangeProperty(PropertyChangedEvent); } // Forward to the UMLDeformerModel.
 #endif
 	// ~END UObject overrides.
 
 #if WITH_EDITORONLY_DATA
-	virtual bool HasTestGroundTruth() const { return false; }
-
-	FVector GetMeshSpacingOffsetVector() const { return FVector(MeshSpacing, 0.0f, 0.0f); }
-	float GetMeshSpacing() const { return MeshSpacing; }
-	float GetLabelHeight() const { return LabelHeight; }
-	bool GetDrawLabels() const { return bDrawLabels; }
-	float GetLabelScale() const { return LabelScale; }
-	EMLDeformerVizMode GetVisualizationMode() const { return VisualizationMode; }
-	int32 GetTrainingFrameNumber() const { return TrainingFrameNumber; }
-	int32 GetTestingFrameNumber() const { return TestingFrameNumber; }
-	float GetAnimPlaySpeed() const { return AnimPlaySpeed; }
-	const UAnimSequence* GetTestAnimSequence() const { return TestAnimSequence; }
-	UAnimSequence* GetTestAnimSequence() { return TestAnimSequence; }
-	bool GetDrawLinearSkinnedActor() const { return bDrawLinearSkinnedActor; }
-	bool GetDrawMLDeformedActor() const { return bDrawMLDeformedActor; }
-	bool GetDrawGroundTruthActor() const { return bDrawGroundTruthActor; }
-	bool GetShowHeatMap() const { return bShowHeatMap; }
-	EMLDeformerHeatMapMode GetHeatMapMode() const { return HeatMapMode; }
-	float GetHeatMapMax() const { return HeatMapMax; }
-	float GetGroundTruthLerp() const { return HeatMapMode == EMLDeformerHeatMapMode::GroundTruth ? GroundTruthLerp : 0.0f; }
-	UMeshDeformer* GetDeformerGraph() const { return DeformerGraph; }
-	void SetDeformerGraph(UMeshDeformer* InDeformerGraph) { DeformerGraph = InDeformerGraph; }
-	float GetWeight() const { return Weight; }
-	bool GetXRayDeltas() const { return bXRayDeltas; }
-	bool GetDrawVertexDeltas() const { return bDrawDeltas; }
+	virtual bool HasTestGroundTruth() const					{ return false; }
+
+	void SetVisualizationMode(EMLDeformerVizMode Mode)		{ VisualizationMode = Mode; }
+	void SetDeformerGraph(UMeshDeformer* InDeformerGraph)	{ DeformerGraph = InDeformerGraph; }
+	void SetTrainingFrameNumber(int32 FrameNumber)			{ TrainingFrameNumber = FrameNumber; }
+	void SetTestingFrameNumber(int32 FrameNumber)			{ TestingFrameNumber = FrameNumber; }
+	void SetWeight(float InWeight)							{ Weight = InWeight; }
+
+	FVector GetMeshSpacingOffsetVector() const				{ return FVector(MeshSpacing, 0.0f, 0.0f); }
+	float GetMeshSpacing() const							{ return MeshSpacing; }
+	float GetLabelHeight() const							{ return LabelHeight; }
+	bool GetDrawLabels() const								{ return bDrawLabels; }
+	float GetLabelScale() const								{ return LabelScale; }
+	EMLDeformerVizMode GetVisualizationMode() const			{ return VisualizationMode; }
+	int32 GetTrainingFrameNumber() const					{ return TrainingFrameNumber; }
+	int32 GetTestingFrameNumber() const						{ return TestingFrameNumber; }
+	float GetAnimPlaySpeed() const							{ return AnimPlaySpeed; }
+	const UAnimSequence* GetTestAnimSequence() const		{ return TestAnimSequence; }
+	UAnimSequence* GetTestAnimSequence()					{ return TestAnimSequence; }
+	bool GetDrawLinearSkinnedActor() const					{ return bDrawLinearSkinnedActor; }
+	bool GetDrawMLDeformedActor() const						{ return bDrawMLDeformedActor; }
+	bool GetDrawGroundTruthActor() const					{ return bDrawGroundTruthActor; }
+	bool GetShowHeatMap() const								{ return bShowHeatMap; }
+	EMLDeformerHeatMapMode GetHeatMapMode() const			{ return HeatMapMode; }
+	float GetHeatMapMax() const								{ return HeatMapMax; }
+	float GetGroundTruthLerp() const						{ return HeatMapMode == EMLDeformerHeatMapMode::GroundTruth ? GroundTruthLerp : 0.0f; }
+	UMeshDeformer* GetDeformerGraph() const					{ return DeformerGraph; }
+	float GetWeight() const									{ return Weight; }
+	bool GetXRayDeltas() const								{ return bXRayDeltas; }
+	bool GetDrawVertexDeltas() const						{ return bDrawDeltas; }
+
+	// Get property names.
+	static FName GetVisualizationModePropertyName()			{ return GET_MEMBER_NAME_CHECKED(UMLDeformerVizSettings, VisualizationMode); }
+	static FName GetMeshSpacingPropertyName()				{ return GET_MEMBER_NAME_CHECKED(UMLDeformerVizSettings, MeshSpacing); }
+	static FName GetDrawLabelsPropertyName()				{ return GET_MEMBER_NAME_CHECKED(UMLDeformerVizSettings, bDrawLabels); }
+	static FName GetLabelHeightPropertyName()				{ return GET_MEMBER_NAME_CHECKED(UMLDeformerVizSettings, LabelHeight); }
+	static FName GetLabelScalePropertyName()				{ return GET_MEMBER_NAME_CHECKED(UMLDeformerVizSettings, LabelScale); }
+	static FName GetTrainingFrameNumberPropertyName()		{ return GET_MEMBER_NAME_CHECKED(UMLDeformerVizSettings, TrainingFrameNumber); }
+	static FName GetTestingFrameNumberPropertyName()		{ return GET_MEMBER_NAME_CHECKED(UMLDeformerVizSettings, TestingFrameNumber); }
+	static FName GetAnimPlaySpeedPropertyName()				{ return GET_MEMBER_NAME_CHECKED(UMLDeformerVizSettings, AnimPlaySpeed); }
+	static FName GetTestAnimSequencePropertyName()			{ return GET_MEMBER_NAME_CHECKED(UMLDeformerVizSettings, TestAnimSequence); }
+	static FName GetDrawLinearSkinnedActorPropertyName()	{ return GET_MEMBER_NAME_CHECKED(UMLDeformerVizSettings, bDrawLinearSkinnedActor); }
+	static FName GetDrawMLDeformedActorPropertyName()		{ return GET_MEMBER_NAME_CHECKED(UMLDeformerVizSettings, bDrawMLDeformedActor); }
+	static FName GetDrawGroundTruthActorPropertyName()		{ return GET_MEMBER_NAME_CHECKED(UMLDeformerVizSettings, bDrawGroundTruthActor); }
+	static FName GetShowHeatMapPropertyName()				{ return GET_MEMBER_NAME_CHECKED(UMLDeformerVizSettings, bShowHeatMap); }
+	static FName GetHeatMapModePropertyName()				{ return GET_MEMBER_NAME_CHECKED(UMLDeformerVizSettings, HeatMapMode); }
+	static FName GetHeatMapMaxPropertyName()				{ return GET_MEMBER_NAME_CHECKED(UMLDeformerVizSettings, HeatMapMax); }
+	static FName GetGroundTruthLerpPropertyName()			{ return GET_MEMBER_NAME_CHECKED(UMLDeformerVizSettings, GroundTruthLerp); }
+	static FName GetDeformerGraphPropertyName()				{ return GET_MEMBER_NAME_CHECKED(UMLDeformerVizSettings, DeformerGraph); }
+	static FName GetWeightPropertyName()					{ return GET_MEMBER_NAME_CHECKED(UMLDeformerVizSettings, Weight); }
+	static FName GetXRayDeltasPropertyName()				{ return GET_MEMBER_NAME_CHECKED(UMLDeformerVizSettings, bXRayDeltas); }
+	static FName GetDrawVertexDeltasPropertyName()			{ return GET_MEMBER_NAME_CHECKED(UMLDeformerVizSettings, bDrawDeltas); }
 #endif
 
-public:
+protected:
 #if WITH_EDITORONLY_DATA
 	/** The data to visualize. */
 	UPROPERTY()
@@ -96,7 +124,7 @@ public:
 
 	/** The frame number of the training data to visualize. */
 	UPROPERTY(EditAnywhere, Category = "Training Meshes", meta = (ClampMin = "0"))
-	uint32 TrainingFrameNumber = 0;
+	int32 TrainingFrameNumber = 0;
 
 	/** Specifies whether we should draw the labels, such as "Linear Skinned" and "ML Deformed" or not. */
 	UPROPERTY(EditAnywhere, Category = "Shared Settings")
@@ -116,7 +144,7 @@ public:
 
 	/** The frame number of the test data to visualize. */
 	UPROPERTY(EditAnywhere, Category = "Live Settings", meta = (ClampMin = "0"))
-	uint32 TestingFrameNumber = 0;
+	int32 TestingFrameNumber = 0;
 
 	/** Specify whether the heatmap is enabled or not. */
 	UPROPERTY(EditAnywhere, Category = "Live Settings")
diff --git a/Engine/Plugins/Experimental/MLDeformer/MLDeformerFramework/Source/MLDeformerFrameworkEditor/Private/MLDeformerAssetActions.h b/Engine/Plugins/Experimental/MLDeformer/MLDeformerFramework/Source/MLDeformerFrameworkEditor/Private/MLDeformerAssetActions.h
--- a/Engine/Plugins/Experimental/MLDeformer/MLDeformerFramework/Source/MLDeformerFrameworkEditor/Private/MLDeformerAssetActions.h
+++ b/Engine/Plugins/Experimental/MLDeformer/MLDeformerFramework/Source/MLDeformerFrameworkEditor/Private/MLDeformerAssetActions.h
@@ -7,17 +7,20 @@
 
 namespace UE::MLDeformer
 {
+	/**
+	 * The asset actions class, describing what actions to take when double clicking on the ML Deformer asset type, etc.
+	 */
 	class FMLDeformerAssetActions
 		: public FAssetTypeActions_Base
 	{
 	public:
 		// IAssetTypeActions overrides.
-		virtual FText GetName() const override { return NSLOCTEXT("MLDeformerAssetActions", "Name", "ML Deformer"); }
 		virtual void GetActions(const TArray<UObject*>& InObjects, struct FToolMenuSection& Section) override;
-		virtual FColor GetTypeColor() const override { return FColor(255, 255, 0); }
+		virtual FText GetName() const override			{ return NSLOCTEXT("MLDeformerAssetActions", "Name", "ML Deformer"); }
+		virtual FColor GetTypeColor() const override	{ return FColor(255, 255, 0); }
+		virtual uint32 GetCategories() override			{ return EAssetTypeCategories::Animation; }
 		virtual UClass* GetSupportedClass() const override;
 		virtual void OpenAssetEditor(const TArray<UObject*>& InObjects, TSharedPtr<class IToolkitHost> EditWithinLevelEditor = TSharedPtr<IToolkitHost>()) override;
-		virtual uint32 GetCategories() override { return EAssetTypeCategories::Animation; }
 		virtual const TArray<FText>& GetSubMenus() const override;
 		// ~END IAssetTypeActions overrides.
 	};
diff --git a/Engine/Plugins/Experimental/MLDeformer/MLDeformerFramework/Source/MLDeformerFrameworkEditor/Private/MLDeformerAssetFactory.h b/Engine/Plugins/Experimental/MLDeformer/MLDeformerFramework/Source/MLDeformerFrameworkEditor/Private/MLDeformerAssetFactory.h
--- a/Engine/Plugins/Experimental/MLDeformer/MLDeformerFramework/Source/MLDeformerFrameworkEditor/Private/MLDeformerAssetFactory.h
+++ b/Engine/Plugins/Experimental/MLDeformer/MLDeformerFramework/Source/MLDeformerFrameworkEditor/Private/MLDeformerAssetFactory.h
@@ -6,6 +6,10 @@
 #include "Factories/Factory.h"
 #include "MLDeformerAssetFactory.generated.h"
 
+/**
+ * The factory for the ML Deformer asset type.
+ * This basically integrates the new asset type into the editor, so you can right click and create a new ML Deformer asset.
+ */
 UCLASS(hidecategories=Object)
 class UMLDeformerFactory
 	: public UFactory
diff --git a/Engine/Plugins/Experimental/MLDeformer/MLDeformerFramework/Source/MLDeformerFrameworkEditor/Private/MLDeformerEditorActor.cpp b/Engine/Plugins/Experimental/MLDeformer/MLDeformerFramework/Source/MLDeformerFrameworkEditor/Private/MLDeformerEditorActor.cpp
--- a/Engine/Plugins/Experimental/MLDeformer/MLDeformerFramework/Source/MLDeformerFrameworkEditor/Private/MLDeformerEditorActor.cpp
+++ b/Engine/Plugins/Experimental/MLDeformer/MLDeformerFramework/Source/MLDeformerFrameworkEditor/Private/MLDeformerEditorActor.cpp
@@ -10,7 +10,7 @@ namespace UE::MLDeformer
 {
 	FMLDeformerEditorActor::FMLDeformerEditorActor(const FConstructSettings& Settings)
 	{
-		check(Settings.TypeID != -1);
+		checkf(Settings.TypeID != -1, TEXT("Expected a valid actor type ID, please set this to a valid ID."));
 		check(Settings.Actor != nullptr);
 
 		TypeID = Settings.TypeID;
diff --git a/Engine/Plugins/Experimental/MLDeformer/MLDeformerFramework/Source/MLDeformerFrameworkEditor/Private/MLDeformerEditorModel.cpp b/Engine/Plugins/Experimental/MLDeformer/MLDeformerFramework/Source/MLDeformerFrameworkEditor/Private/MLDeformerEditorModel.cpp
--- a/Engine/Plugins/Experimental/MLDeformer/MLDeformerFramework/Source/MLDeformerFrameworkEditor/Private/MLDeformerEditorModel.cpp
+++ b/Engine/Plugins/Experimental/MLDeformer/MLDeformerFramework/Source/MLDeformerFrameworkEditor/Private/MLDeformerEditorModel.cpp
@@ -522,7 +522,7 @@ namespace UE::MLDeformer
 					EditorActor->SetPlayPosition(PlayOffset);
 				}
 			}
-			VizSettings->TrainingFrameNumber = TargetFrame;
+			VizSettings->SetTrainingFrameNumber(TargetFrame);
 		}
 		else if (VizSettings->GetVisualizationMode() == EMLDeformerVizMode::TestData)
 		{
@@ -538,7 +538,7 @@ namespace UE::MLDeformer
 					EditorActor->SetPlayPosition(PlayOffset);
 				}
 			}
-			VizSettings->TestingFrameNumber = TargetFrame;
+			VizSettings->SetTestingFrameNumber(TargetFrame);
 		}
 	}
 
@@ -565,7 +565,7 @@ namespace UE::MLDeformer
 	void FMLDeformerEditorModel::SetTrainingFrame(int32 FrameNumber)
 	{
 		UMLDeformerVizSettings* VizSettings = Model->GetVizSettings();
-		VizSettings->TrainingFrameNumber = FrameNumber;
+		VizSettings->SetTrainingFrameNumber(FrameNumber);
 		ClampCurrentTrainingFrameIndex();
 		if (VizSettings->GetVisualizationMode() == EMLDeformerVizMode::TrainingData)
 		{
@@ -576,7 +576,7 @@ namespace UE::MLDeformer
 	void FMLDeformerEditorModel::SetTestFrame(int32 FrameNumber)
 	{
 		UMLDeformerVizSettings* VizSettings = Model->GetVizSettings();
-		VizSettings->TestingFrameNumber = FrameNumber;
+		VizSettings->SetTestingFrameNumber(FrameNumber);
 		ClampCurrentTestFrameIndex();
 		if (VizSettings->GetVisualizationMode() == EMLDeformerVizMode::TestData)
 		{
@@ -593,95 +593,94 @@ namespace UE::MLDeformer
 		}
 
 		// When we change one of these properties below, restart animations etc.
-		if (Property->GetFName() == GET_MEMBER_NAME_CHECKED(UMLDeformerModel, SkeletalMesh))
+		if (Property->GetFName() == UMLDeformerModel::GetSkeletalMeshPropertyName())
 		{
 			TriggerInputAssetChanged();
+			SetResamplingInputOutputsNeeded(true);
 			Model->InitVertexMap();
 			Model->InitGPUData();
 			UpdateDeformerGraph();
 		}
-		else
-		if (Property->GetFName() == GET_MEMBER_NAME_CHECKED(UMLDeformerModel, AnimSequence) ||
-			Property->GetFName() == GET_MEMBER_NAME_CHECKED(UMLDeformerVizSettings, TestAnimSequence))
+		else if (Property->GetFName() == UMLDeformerVizSettings::GetTestAnimSequencePropertyName())
 		{
 			TriggerInputAssetChanged(true);
 		}
-		else if (Property->GetFName() == GET_MEMBER_NAME_CHECKED(UMLDeformerModel, AlignmentTransform))
+		else if (Property->GetFName() == UMLDeformerModel::GetAnimSequencePropertyName())
+		{
+			TriggerInputAssetChanged();
+			SetResamplingInputOutputsNeeded(true);
+		}
+		if (Property->GetFName() == UMLDeformerModel::GetAlignmentTransformPropertyName() ||
+		    Property->GetFName() == UMLDeformerModel::GetDeltaCutoffLengthPropertyName())
 		{
 			if (PropertyChangedEvent.ChangeType == EPropertyChangeType::ValueSet)
 			{
+				SetResamplingInputOutputsNeeded(true);
 				SampleDeltas();
 			}
 		}
-		else if (Property->GetFName() == GET_MEMBER_NAME_CHECKED(UMLDeformerModel, MaxTrainingFrames))
+		else if (Property->GetFName() == UMLDeformerModel::GetMaxTrainingFramesPropertyName())
 		{
 			TriggerInputAssetChanged();
+			SetResamplingInputOutputsNeeded(true);
 		}
 		else
-		if (Property->GetFName() == GET_MEMBER_NAME_CHECKED(UMLDeformerModel, bIncludeBones) ||
-			Property->GetFName() == GET_MEMBER_NAME_CHECKED(UMLDeformerModel, bIncludeCurves))
+		if (Property->GetFName() == UMLDeformerModel::GetShouldIncludeBonesPropertyName() ||
+			Property->GetFName() == UMLDeformerModel::GetShouldIncludeCurvesPropertyName())
 		{
 			if (PropertyChangedEvent.ChangeType == EPropertyChangeType::ValueSet)
 			{
+				SetResamplingInputOutputsNeeded(true);
 				UpdateEditorInputInfo();
 				UpdateIsReadyForTrainingState();
 				GetEditor()->GetModelDetailsView()->ForceRefresh();
 			}
 		}
 		else
-		if (Property->GetFName() == GET_MEMBER_NAME_CHECKED(UMLDeformerModel, BoneIncludeList) ||
-			Property->GetFName() == GET_MEMBER_NAME_CHECKED(UMLDeformerModel, CurveIncludeList))
+		if (Property->GetFName() == UMLDeformerModel::GetBoneIncludeListPropertyName() ||
+			Property->GetFName() == UMLDeformerModel::GetCurveIncludeListPropertyName())
 		{
+			SetResamplingInputOutputsNeeded(true);
 			UpdateEditorInputInfo();
 			GetEditor()->GetModelDetailsView()->ForceRefresh();
 		}
-		else if (Property->GetFName() == GET_MEMBER_NAME_CHECKED(UMLDeformerVizSettings, AnimPlaySpeed))
+		else if (Property->GetFName() == UMLDeformerVizSettings::GetAnimPlaySpeedPropertyName())
 		{
 			UpdateTestAnimPlaySpeed();
 		}
-		else if (Property->GetFName() == GET_MEMBER_NAME_CHECKED(UMLDeformerVizSettings, TrainingFrameNumber))
+		else if (Property->GetFName() == UMLDeformerVizSettings::GetTrainingFrameNumberPropertyName())
 		{
 			ClampCurrentTrainingFrameIndex();
 			const int32 CurrentFrameNumber = Model->GetVizSettings()->GetTrainingFrameNumber();
 			OnTimeSliderScrubPositionChanged(GetTrainingTimeAtFrame(CurrentFrameNumber), false);
 		}
-		else if (Property->GetFName() == GET_MEMBER_NAME_CHECKED(UMLDeformerVizSettings, TestingFrameNumber))
+		else if (Property->GetFName() == UMLDeformerVizSettings::GetTestingFrameNumberPropertyName())
 		{
 			ClampCurrentTestFrameIndex();
 			const int32 CurrentFrameNumber = Model->GetVizSettings()->GetTestingFrameNumber();
 			OnTimeSliderScrubPositionChanged(GetTestTimeAtFrame(CurrentFrameNumber), false);
 		}
-		else if (Property->GetFName() == GET_MEMBER_NAME_CHECKED(UMLDeformerVizSettings, bShowHeatMap))
+		else if (Property->GetFName() == UMLDeformerVizSettings::GetShowHeatMapPropertyName())
 		{
 			SetHeatMapMaterialEnabled(Model->GetVizSettings()->GetShowHeatMap());
 			UpdateDeformerGraph();
 		}
 		else
-		if (Property->GetFName() == GET_MEMBER_NAME_CHECKED(UMLDeformerVizSettings, bDrawLinearSkinnedActor) ||
-			Property->GetFName() == GET_MEMBER_NAME_CHECKED(UMLDeformerVizSettings, bDrawMLDeformedActor) ||
-			Property->GetFName() == GET_MEMBER_NAME_CHECKED(UMLDeformerVizSettings, bDrawGroundTruthActor))
+		if (Property->GetFName() == UMLDeformerVizSettings::GetDrawLinearSkinnedActorPropertyName() ||
+			Property->GetFName() == UMLDeformerVizSettings::GetDrawMLDeformedActorPropertyName() ||
+			Property->GetFName() == UMLDeformerVizSettings::GetDrawGroundTruthActorPropertyName())
 		{
 			UpdateActorVisibility();
 		}
-		else if (Property->GetFName() == GET_MEMBER_NAME_CHECKED(UMLDeformerVizSettings, bDrawDeltas))
+		else if (Property->GetFName() == UMLDeformerVizSettings::GetDrawVertexDeltasPropertyName())
 		{
 			SampleDeltas();
 		} 
-		else if (Property->GetFName() == GET_MEMBER_NAME_CHECKED(UMLDeformerVizSettings, DeformerGraph))
+		else if (Property->GetFName() == UMLDeformerVizSettings::GetDeformerGraphPropertyName())
 		{
 			UpdateDeformerGraph();
 			GetEditor()->GetVizSettingsDetailsView()->ForceRefresh();
-		}
-		else
-		if (Property->GetFName() == GET_MEMBER_NAME_CHECKED(UMLDeformerModel, DeltaCutoffLength) ||
-			Property->GetFName() == GET_MEMBER_NAME_CHECKED(UMLDeformerModel, AlignmentTransform))
-		{
-			if (PropertyChangedEvent.ChangeType == EPropertyChangeType::ValueSet)
-			{
-				SetResamplingInputOutputsNeeded(true);
-				SampleDeltas();
-			}
-		}
+		};
 	}
 
 	FMLDeformerEditorActor* FMLDeformerEditorModel::GetVisualizationModeBaseActor() const
@@ -716,7 +715,7 @@ namespace UE::MLDeformer
 		{
 			if (EditorActor && ((EditorActor->IsTestActor() && bTestMode) || (EditorActor->IsTrainingActor() && !bTestMode)))
 			{
-				EditorActor->SetPlaySpeed(VizSettings->AnimPlaySpeed);
+				EditorActor->SetPlaySpeed(VizSettings->GetAnimPlaySpeed());
 				EditorActor->Pause(bMustPause);
 			}
 		}
@@ -799,11 +798,12 @@ namespace UE::MLDeformer
 		UMLDeformerVizSettings* VizSettings = Model->GetVizSettings();
 		if (GetNumTrainingFrames() > 0)
 		{
-			VizSettings->TrainingFrameNumber = FMath::Min(VizSettings->TrainingFrameNumber, static_cast<uint32>(GetNumTrainingFrames() - 1));
+			const int32 ClampedNumber = FMath::Min(VizSettings->GetTrainingFrameNumber(), GetNumTrainingFrames() - 1);
+			VizSettings->SetTrainingFrameNumber(ClampedNumber);
 		}
 		else
 		{
-			VizSettings->TrainingFrameNumber = 0;
+			VizSettings->SetTrainingFrameNumber(0);
 		}
 	}
 
@@ -812,11 +812,12 @@ namespace UE::MLDeformer
 		UMLDeformerVizSettings* VizSettings = Model->GetVizSettings();
 		if (GetNumTestFrames() > 0)
 		{
-			VizSettings->TestingFrameNumber = FMath::Min(VizSettings->TestingFrameNumber, static_cast<uint32>(GetNumTestFrames() - 1));
+			const int32 ClampedNumber = FMath::Min(VizSettings->GetTestingFrameNumber(), GetNumTestFrames() - 1);
+			VizSettings->SetTestingFrameNumber(ClampedNumber);
 		}
 		else
 		{
-			VizSettings->TestingFrameNumber = 0;
+			VizSettings->SetTestingFrameNumber(0);
 		}
 	}
 
@@ -838,14 +839,14 @@ namespace UE::MLDeformer
 	FText FMLDeformerEditorModel::GetBaseAssetChangedErrorText() const
 	{
 		FText Result;
-		if (Model->SkeletalMesh && Model->GetInputInfo())
+		if (Model->GetSkeletalMesh() && Model->GetInputInfo())
 		{
-			if (Model->NumBaseMeshVerts != Model->GetInputInfo()->GetNumBaseMeshVertices() &&
-				Model->NumBaseMeshVerts > 0 && Model->GetInputInfo()->GetNumBaseMeshVertices() > 0)
+			if (Model->GetNumBaseMeshVerts() != Model->GetInputInfo()->GetNumBaseMeshVertices() &&
+				Model->GetNumBaseMeshVerts() > 0 && Model->GetInputInfo()->GetNumBaseMeshVertices() > 0)
 			{
 				Result = FText::Format(LOCTEXT("BaseMeshMismatch", "Number of vertices in base mesh has changed from {0} to {1} vertices since this ML Deformer Asset was saved! {2}"),
 					Model->GetInputInfo()->GetNumBaseMeshVertices(),
-					Model->NumBaseMeshVerts,
+					Model->GetNumBaseMeshVerts(),
 					IsTrained() ? LOCTEXT("BaseMeshMismatchNN", "Neural network needs to be retrained!") : FText());
 			}
 		}
@@ -855,10 +856,10 @@ namespace UE::MLDeformer
 	FText FMLDeformerEditorModel::GetVertexMapChangedErrorText() const
 	{
 		FText Result;
-		if (Model->SkeletalMesh)
+		if (Model->GetSkeletalMesh())
 		{
 			bool bVertexMapMatch = true;
-			const FSkeletalMeshModel* ImportedModel = Model->SkeletalMesh->GetImportedModel();
+			const FSkeletalMeshModel* ImportedModel = Model->GetSkeletalMesh()->GetImportedModel();
 			if (ImportedModel)
 			{
 				const TArray<int32>& MeshVertexMap = ImportedModel->LODModels[0].MeshToImportVertexMap;
@@ -888,9 +889,9 @@ namespace UE::MLDeformer
 	{
 		FText Result;
 
-		if (Model->SkeletalMesh)
+		if (Model->GetSkeletalMesh())
 		{
-			FSkeletalMeshModel* ImportedModel = Model->SkeletalMesh->GetImportedModel();
+			FSkeletalMeshModel* ImportedModel = Model->GetSkeletalMesh()->GetImportedModel();
 			check(ImportedModel);
 
 			const TArray<FSkelMeshImportedMeshInfo>& SkelMeshInfos = ImportedModel->LODModels[0].ImportedMeshInfos;
@@ -905,7 +906,7 @@ namespace UE::MLDeformer
 
 	FText FMLDeformerEditorModel::GetInputsErrorText() const
 	{
-		if (Model->SkeletalMesh && GetEditorInputInfo()->IsEmpty())
+		if (Model->GetSkeletalMesh() && GetEditorInputInfo()->IsEmpty())
 		{
 			FString ErrorString;
 			if (Model->DoesSupportBones() && Model->ShouldIncludeBonesInTraining())
@@ -949,12 +950,12 @@ namespace UE::MLDeformer
 		UMLDeformerInputInfo* InputInfo = Model->GetInputInfo();
 		if (Model->HasTrainingGroundTruth() && InputInfo)
 		{
-			if (Model->NumTargetMeshVerts != InputInfo->GetNumTargetMeshVertices() &&
-				Model->NumTargetMeshVerts > 0 && InputInfo->GetNumTargetMeshVertices() > 0)
+			if (Model->GetNumTargetMeshVerts() != InputInfo->GetNumTargetMeshVertices() &&
+				Model->GetNumTargetMeshVerts() > 0 && InputInfo->GetNumTargetMeshVertices() > 0)
 			{
 				Result = FText::Format(LOCTEXT("TargetMeshMismatch", "Number of vertices in target mesh has changed from {0} to {1} vertices since this ML Deformer Asset was saved! {2}"),
 					InputInfo->GetNumTargetMeshVertices(),
-					Model->NumTargetMeshVerts,
+					Model->GetNumTargetMeshVerts(),
 					IsTrained() ? LOCTEXT("BaseMeshMismatchModel", "Model needs to be retrained!") : FText());
 			}
 		}
@@ -1061,26 +1062,26 @@ namespace UE::MLDeformer
 
 	void FMLDeformerEditorModel::InitBoneIncludeListToAnimatedBonesOnly()
 	{
-		if (!Model->AnimSequence)
+		if (!Model->GetAnimSequence())
 		{
 			UE_LOG(LogMLDeformer, Warning, TEXT("Cannot initialize bone list as no Anim Sequence has been picked."));
 			return;
 		}
 
-		const UAnimDataModel* DataModel = Model->AnimSequence->GetDataModel();
+		const UAnimDataModel* DataModel = Model->GetAnimSequence()->GetDataModel();
 		if (!DataModel)
 		{
 			UE_LOG(LogMLDeformer, Warning, TEXT("Anim sequence has no data model."));
 			return;
 		}
 
-		if (!Model->SkeletalMesh)
+		if (!Model->GetSkeletalMesh())
 		{
 			UE_LOG(LogMLDeformer, Warning, TEXT("Skeletal Mesh has not been set."));
 			return;
 		}
 
-		USkeleton* Skeleton = Model->SkeletalMesh->GetSkeleton();
+		USkeleton* Skeleton = Model->GetSkeletalMesh()->GetSkeleton();
 		if (!Skeleton)
 		{
 			UE_LOG(LogMLDeformer, Warning, TEXT("Skeletal Mesh has no skeleton."));
@@ -1131,45 +1132,46 @@ namespace UE::MLDeformer
 		// Init the bone include list using the animated bones.
 		if (!AnimatedBoneList.IsEmpty())
 		{
-			Model->BoneIncludeList.Empty();
-			Model->BoneIncludeList.Reserve(AnimatedBoneList.Num());
+			TArray<FBoneReference>& BoneList = Model->GetBoneIncludeList();
+			BoneList.Empty();
+			BoneList.Reserve(AnimatedBoneList.Num());
 			for (FName BoneName : AnimatedBoneList)
 			{
-				Model->BoneIncludeList.AddDefaulted();
-				FBoneReference& BoneRef = Model->BoneIncludeList.Last();
+				BoneList.AddDefaulted();
+				FBoneReference& BoneRef = BoneList.Last();
 				BoneRef.BoneName = BoneName;
 			}
 		}
 		else
 		{
-			Model->BoneIncludeList.Empty();
-			UE_LOG(LogMLDeformer, Warning, TEXT("There are no animated bone rotations in Anim Sequence '%s'."), *(Model->AnimSequence->GetName()));
+			Model->GetBoneIncludeList().Empty();
+			UE_LOG(LogMLDeformer, Warning, TEXT("There are no animated bone rotations in Anim Sequence '%s'."), *(Model->GetAnimSequence()->GetName()));
 		}
 		UpdateEditorInputInfo();
 	}
 
 	void FMLDeformerEditorModel::InitCurveIncludeListToAnimatedCurvesOnly()
 	{
-		if (!Model->AnimSequence)
+		if (!Model->GetAnimSequence())
 		{
 			UE_LOG(LogMLDeformer, Warning, TEXT("Cannot initialize curve list as no Anim Sequence has been picked."));
 			return;
 		}
 
-		const UAnimDataModel* DataModel = Model->AnimSequence->GetDataModel();
+		const UAnimDataModel* DataModel = Model->GetAnimSequence()->GetDataModel();
 		if (!DataModel)
 		{
 			UE_LOG(LogMLDeformer, Warning, TEXT("Anim sequence has no data model."));
 			return;
 		}
 
-		if (!Model->SkeletalMesh)
+		if (!Model->GetSkeletalMesh())
 		{
 			UE_LOG(LogMLDeformer, Warning, TEXT("Skeletal Mesh has not been set."));
 			return;
 		}
 
-		USkeleton* Skeleton = Model->SkeletalMesh->GetSkeleton();
+		USkeleton* Skeleton = Model->GetSkeletalMesh()->GetSkeleton();
 		if (!Skeleton)
 		{
 			UE_LOG(LogMLDeformer, Warning, TEXT("Skeletal Mesh has no skeleton."));
@@ -1214,19 +1216,20 @@ namespace UE::MLDeformer
 		// Init the bone include list using the animated bones.
 		if (!AnimatedCurveList.IsEmpty())
 		{
-			Model->CurveIncludeList.Empty();
-			Model->CurveIncludeList.Reserve(AnimatedCurveList.Num());
+			TArray<FMLDeformerCurveReference>& CurveList = Model->GetCurveIncludeList();
+			CurveList.Empty();
+			CurveList.Reserve(AnimatedCurveList.Num());
 			for (FName CurveName : AnimatedCurveList)
 			{
-				Model->CurveIncludeList.AddDefaulted();
-				FMLDeformerCurveReference& CurveRef = Model->CurveIncludeList.Last();
+				CurveList.AddDefaulted();
+				FMLDeformerCurveReference& CurveRef = CurveList.Last();
 				CurveRef.CurveName = CurveName;
 			}
 		}
 		else
 		{
-			Model->CurveIncludeList.Empty();
-			UE_LOG(LogMLDeformer, Warning, TEXT("There are no animated curves in Anim Sequence '%s'."), *(Model->AnimSequence->GetName()));
+			Model->GetCurveIncludeList().Empty();
+			UE_LOG(LogMLDeformer, Warning, TEXT("There are no animated curves in Anim Sequence '%s'."), *(Model->GetAnimSequence()->GetName()));
 		}
 		UpdateEditorInputInfo();
 	}
@@ -1293,7 +1296,7 @@ namespace UE::MLDeformer
 		if (Sampler->IsInitialized())
 		{
 			Sampler->SetVertexDeltaSpace(EVertexDeltaSpace::PostSkinning);
-			Sampler->Sample(VizSettings->TrainingFrameNumber);
+			Sampler->Sample(VizSettings->GetTrainingFrameNumber());
 		}
 	}
 
@@ -1301,7 +1304,7 @@ namespace UE::MLDeformer
 	{
 		UMLDeformerVizSettings* VizSettings = Model->GetVizSettings();
 		ClampCurrentTrainingFrameIndex();
-		if (CurrentTrainingFrame != VizSettings->TrainingFrameNumber)
+		if (CurrentTrainingFrame != VizSettings->GetTrainingFrameNumber())
 		{
 			OnTrainingDataFrameChanged();
 		}
@@ -1312,9 +1315,9 @@ namespace UE::MLDeformer
 		UMLDeformerVizSettings* VizSettings = Model->GetVizSettings();
 
 		// If the current frame number changed, re-sample the deltas if needed.
-		if (CurrentTrainingFrame != VizSettings->TrainingFrameNumber)
+		if (CurrentTrainingFrame != VizSettings->GetTrainingFrameNumber())
 		{
-			CurrentTrainingFrame = VizSettings->TrainingFrameNumber;
+			CurrentTrainingFrame = VizSettings->GetTrainingFrameNumber();
 			if (VizSettings->GetDrawVertexDeltas() && VizSettings->GetVisualizationMode() == EMLDeformerVizMode::TrainingData)
 			{
 				SampleDeltas();
@@ -1634,10 +1637,11 @@ namespace UE::MLDeformer
 			}
 
 			// Init deltas for this morph target.
+			const TArray<int32>& VertexMap = Model->GetVertexMap();
 			MorphLODModel.Vertices.Reserve(NumRenderVertices);
 			for (int32 VertexIndex = 0; VertexIndex < NumRenderVertices; ++VertexIndex)
 			{
-				const int32 ImportedVertexNumber = Model->VertexMap[VertexIndex];
+				const int32 ImportedVertexNumber = VertexMap[VertexIndex];
 				if (ImportedVertexNumber != INDEX_NONE)
 				{
 					const FVector3f Delta = Deltas[ImportedVertexNumber + BlendShapeIndex * NumBaseMeshVerts];
diff --git a/Engine/Plugins/Experimental/MLDeformer/MLDeformerFramework/Source/MLDeformerFrameworkEditor/Private/MLDeformerEditorToolkit.cpp b/Engine/Plugins/Experimental/MLDeformer/MLDeformerFramework/Source/MLDeformerFrameworkEditor/Private/MLDeformerEditorToolkit.cpp
--- a/Engine/Plugins/Experimental/MLDeformer/MLDeformerFramework/Source/MLDeformerFrameworkEditor/Private/MLDeformerEditorToolkit.cpp
+++ b/Engine/Plugins/Experimental/MLDeformer/MLDeformerFramework/Source/MLDeformerFrameworkEditor/Private/MLDeformerEditorToolkit.cpp
@@ -237,7 +237,7 @@ namespace UE::MLDeformer
 			return;
 		}
 
-		ActiveModel->GetModel()->GetVizSettings()->VisualizationMode = Mode;
+		ActiveModel->GetModel()->GetVizSettings()->SetVisualizationMode(Mode);
 		OnSwitchedVisualizationMode();
 	}
 
@@ -653,7 +653,7 @@ namespace UE::MLDeformer
 			return;
 		}
 
-		if (Property->GetFName() == GET_MEMBER_NAME_CHECKED(UMLDeformerVizSettings, VisualizationMode))
+		if (Property->GetFName() == UMLDeformerVizSettings::GetVisualizationModePropertyName())
 		{
 			OnSwitchedVisualizationMode();
 		}
diff --git a/Engine/Plugins/Experimental/MLDeformer/MLDeformerFramework/Source/MLDeformerFrameworkEditor/Private/MLDeformerGeomCacheEditorModel.cpp b/Engine/Plugins/Experimental/MLDeformer/MLDeformerFramework/Source/MLDeformerFrameworkEditor/Private/MLDeformerGeomCacheEditorModel.cpp
--- a/Engine/Plugins/Experimental/MLDeformer/MLDeformerFramework/Source/MLDeformerFrameworkEditor/Private/MLDeformerGeomCacheEditorModel.cpp
+++ b/Engine/Plugins/Experimental/MLDeformer/MLDeformerFramework/Source/MLDeformerFrameworkEditor/Private/MLDeformerGeomCacheEditorModel.cpp
@@ -234,8 +234,8 @@ namespace UE::MLDeformer
 		FMLDeformerEditorModel::OnPropertyChanged(PropertyChangedEvent);
 
 		// Properties specific to the geometry cache editor model.
-		if (Property->GetFName() == GET_MEMBER_NAME_CHECKED(UMLDeformerGeomCacheModel, GeometryCache) ||
-			Property->GetFName() == GET_MEMBER_NAME_CHECKED(UMLDeformerGeomCacheVizSettings, GroundTruth))
+		if (Property->GetFName() == UMLDeformerGeomCacheModel::GetGeometryCachePropertyName() ||
+			Property->GetFName() == UMLDeformerGeomCacheVizSettings::GetTestGroundTruthPropertyName())
 		{
 			TriggerInputAssetChanged(true);
 		}
@@ -269,7 +269,7 @@ namespace UE::MLDeformer
 		GeometryCacheComponent->SetPlaybackSpeed(TestAnimSpeed);
 		GeometryCacheComponent->Play();
 
-		GetGeomCacheModel()->MeshMappings.Reset();
+		GetGeomCacheModel()->GetGeomCacheMeshMappings().Reset();
 	}
 }	// namespace UE::MLDeformer
 
diff --git a/Engine/Plugins/Experimental/MLDeformer/MLDeformerFramework/Source/MLDeformerFrameworkEditor/Private/MLDeformerGeomCacheModelDetails.cpp b/Engine/Plugins/Experimental/MLDeformer/MLDeformerFramework/Source/MLDeformerFrameworkEditor/Private/MLDeformerGeomCacheModelDetails.cpp
--- a/Engine/Plugins/Experimental/MLDeformer/MLDeformerFramework/Source/MLDeformerFrameworkEditor/Private/MLDeformerGeomCacheModelDetails.cpp
+++ b/Engine/Plugins/Experimental/MLDeformer/MLDeformerFramework/Source/MLDeformerFrameworkEditor/Private/MLDeformerGeomCacheModelDetails.cpp
@@ -45,7 +45,7 @@ namespace UE::MLDeformer
 
 	void FMLDeformerGeomCacheModelDetails::AddTargetMesh()
 	{
-		TargetMeshCategoryBuilder->AddProperty(GET_MEMBER_NAME_CHECKED(UMLDeformerGeomCacheModel, GeometryCache), UMLDeformerGeomCacheModel::StaticClass());
+		TargetMeshCategoryBuilder->AddProperty(UMLDeformerGeomCacheModel::GetGeometryCachePropertyName(), UMLDeformerGeomCacheModel::StaticClass());
 
 		const FText TargetMeshErrorText = GetGeomCacheErrorText(GeomCacheModel->GetSkeletalMesh(), GeomCacheModel->GetGeometryCache());
 		TargetMeshCategoryBuilder->AddCustomRow(FText::FromString("TargetMeshError"))
@@ -77,6 +77,23 @@ namespace UE::MLDeformer
 
 		AddGeomCacheMeshMappingWarnings(TargetMeshCategoryBuilder, Model->GetSkeletalMesh(), GeomCacheModel->GetGeometryCache());
 	}
+
+	void FMLDeformerGeomCacheModelDetails::AddGeomCacheMeshMappingWarnings(IDetailCategoryBuilder* InTargetMeshCategoryBuilder, USkeletalMesh* SkeletalMesh, UGeometryCache* GeometryCache)
+	{
+		const FText MeshMappingError = GetGeomCacheMeshMappingErrorText(SkeletalMesh, GeometryCache);
+		InTargetMeshCategoryBuilder->AddCustomRow(FText::FromString("MeshMappingWarning"))
+			.Visibility(!MeshMappingError.IsEmpty() ? EVisibility::Visible : EVisibility::Collapsed)
+			.WholeRowContent()
+			[
+				SNew(SBox)
+				.Padding(FMargin(0.0f, 4.0f))
+				[
+					SNew(SWarningOrErrorBox)
+					.MessageStyle(EMessageStyle::Warning)
+					.Message(MeshMappingError)
+				]
+			];
+	}
 }	// namespace UE::MLDeformer
 
 #undef LOCTEXT_NAMESPACE
diff --git a/Engine/Plugins/Experimental/MLDeformer/MLDeformerFramework/Source/MLDeformerFrameworkEditor/Private/MLDeformerGeomCacheVizSettingsDetails.cpp b/Engine/Plugins/Experimental/MLDeformer/MLDeformerFramework/Source/MLDeformerFrameworkEditor/Private/MLDeformerGeomCacheVizSettingsDetails.cpp
--- a/Engine/Plugins/Experimental/MLDeformer/MLDeformerFramework/Source/MLDeformerFrameworkEditor/Private/MLDeformerGeomCacheVizSettingsDetails.cpp
+++ b/Engine/Plugins/Experimental/MLDeformer/MLDeformerFramework/Source/MLDeformerFrameworkEditor/Private/MLDeformerGeomCacheVizSettingsDetails.cpp
@@ -33,7 +33,7 @@ namespace UE::MLDeformer
 
 	void FMLDeformerGeomCacheVizSettingsDetails::AddGroundTruth()
 	{
-		TestAssetsCategory->AddProperty(GET_MEMBER_NAME_CHECKED(UMLDeformerGeomCacheVizSettings, GroundTruth), UMLDeformerGeomCacheVizSettings::StaticClass());
+		TestAssetsCategory->AddProperty(UMLDeformerGeomCacheVizSettings::GetTestGroundTruthPropertyName(), UMLDeformerGeomCacheVizSettings::StaticClass());
 
 		// Show an error when the test anim sequence duration doesn't match the one of the ground truth.
 		const FText AnimErrorText = GetGeomCacheAnimSequenceErrorText(GeomCacheVizSettings->GetTestGroundTruth(), VizSettings->GetTestAnimSequence());
diff --git a/Engine/Plugins/Experimental/MLDeformer/MLDeformerFramework/Source/MLDeformerFrameworkEditor/Private/MLDeformerModelDetails.cpp b/Engine/Plugins/Experimental/MLDeformer/MLDeformerFramework/Source/MLDeformerFrameworkEditor/Private/MLDeformerModelDetails.cpp
--- a/Engine/Plugins/Experimental/MLDeformer/MLDeformerFramework/Source/MLDeformerFrameworkEditor/Private/MLDeformerModelDetails.cpp
+++ b/Engine/Plugins/Experimental/MLDeformer/MLDeformerFramework/Source/MLDeformerFrameworkEditor/Private/MLDeformerModelDetails.cpp
@@ -4,6 +4,7 @@
 #include "MLDeformerModule.h"
 #include "MLDeformerAsset.h"
 #include "MLDeformerModel.h"
+#include "MLDeformerInputInfo.h"
 #include "MLDeformerEditorModel.h"
 #include "MLDeformerEditorModule.h"
 #include "MLDeformerGeomCacheHelpers.h"
@@ -72,7 +73,7 @@ namespace UE::MLDeformer
 		CreateCategories();
 
 		// Base mesh details.
-		BaseMeshCategoryBuilder->AddProperty(GET_MEMBER_NAME_CHECKED(UMLDeformerModel, SkeletalMesh), UMLDeformerModel::StaticClass());
+		BaseMeshCategoryBuilder->AddProperty(UMLDeformerModel::GetSkeletalMeshPropertyName(), UMLDeformerModel::StaticClass());
 
 		AddBaseMeshErrors();
 
@@ -121,7 +122,7 @@ namespace UE::MLDeformer
 			];
 
 		// Animation sequence.
-		IDetailPropertyRow& AnimRow = BaseMeshCategoryBuilder->AddProperty(GET_MEMBER_NAME_CHECKED(UMLDeformerModel, AnimSequence), UMLDeformerModel::StaticClass());
+		IDetailPropertyRow& AnimRow = BaseMeshCategoryBuilder->AddProperty(UMLDeformerModel::GetAnimSequencePropertyName(), UMLDeformerModel::StaticClass());
 		AnimRow.CustomWidget()
 		.NameContent()
 		[
@@ -143,7 +144,7 @@ namespace UE::MLDeformer
 
 		AddAnimSequenceErrors();
 
-		const FText AnimErrorText = EditorModel->GetIncompatibleSkeletonErrorText(Model->GetSkeletalMesh(), Model->AnimSequence);
+		const FText AnimErrorText = EditorModel->GetIncompatibleSkeletonErrorText(Model->GetSkeletalMesh(), Model->GetAnimSequence());
 		FDetailWidgetRow& AnimErrorRow = BaseMeshCategoryBuilder->AddCustomRow(FText::FromString("AnimSkeletonMisMatchError"))
 			.Visibility(!AnimErrorText.IsEmpty() ? EVisibility::Visible : EVisibility::Collapsed)
 			.WholeRowContent()
@@ -160,12 +161,12 @@ namespace UE::MLDeformer
 
 		AddTargetMesh();
 
-		TargetMeshCategoryBuilder->AddProperty(GET_MEMBER_NAME_CHECKED(UMLDeformerModel, AlignmentTransform), UMLDeformerModel::StaticClass());
+		TargetMeshCategoryBuilder->AddProperty(UMLDeformerModel::GetAlignmentTransformPropertyName(), UMLDeformerModel::StaticClass());
 
-		InputOutputCategoryBuilder->AddProperty(GET_MEMBER_NAME_CHECKED(UMLDeformerModel, bIncludeBones), UMLDeformerModel::StaticClass())
+		InputOutputCategoryBuilder->AddProperty(UMLDeformerModel::GetShouldIncludeBonesPropertyName(), UMLDeformerModel::StaticClass())
 			.Visibility(IsBonesFlagVisible() ? EVisibility::Visible : EVisibility::Collapsed);
 
-		InputOutputCategoryBuilder->AddProperty(GET_MEMBER_NAME_CHECKED(UMLDeformerModel, bIncludeCurves), UMLDeformerModel::StaticClass())
+		InputOutputCategoryBuilder->AddProperty(UMLDeformerModel::GetShouldIncludeCurvesPropertyName(), UMLDeformerModel::StaticClass())
 			.Visibility(IsCurvesFlagVisible() ? EVisibility::Visible : EVisibility::Collapsed);
 
 		AddTrainingInputFlags();
@@ -185,8 +186,8 @@ namespace UE::MLDeformer
 				]
 			];
 
-		InputOutputCategoryBuilder->AddProperty(GET_MEMBER_NAME_CHECKED(UMLDeformerModel, MaxTrainingFrames), UMLDeformerModel::StaticClass());
-		InputOutputCategoryBuilder->AddProperty(GET_MEMBER_NAME_CHECKED(UMLDeformerModel, DeltaCutoffLength), UMLDeformerModel::StaticClass());
+		InputOutputCategoryBuilder->AddProperty(UMLDeformerModel::GetMaxTrainingFramesPropertyName(), UMLDeformerModel::StaticClass());
+		InputOutputCategoryBuilder->AddProperty(UMLDeformerModel::GetDeltaCutoffLengthPropertyName(), UMLDeformerModel::StaticClass());
 
 		// Bone include list group.
 		if (Model->DoesSupportBones())
@@ -233,13 +234,13 @@ namespace UE::MLDeformer
 					.OnClicked(FOnClicked::CreateSP(this, &FMLDeformerModelDetails::OnFilterAnimatedBonesOnly))
 					.IsEnabled_Lambda([this](){ return Model->ShouldIncludeBonesInTraining(); })
 				];
-			BoneIncludeGroup.AddPropertyRow(DetailBuilder.GetProperty(GET_MEMBER_NAME_CHECKED(UMLDeformerModel, BoneIncludeList), UMLDeformerModel::StaticClass()));
+			BoneIncludeGroup.AddPropertyRow(DetailBuilder.GetProperty(UMLDeformerModel::GetBoneIncludeListPropertyName(), UMLDeformerModel::StaticClass()));
 
 			AddBoneInputErrors();
 		}
 		else
 		{
-			InputOutputCategoryBuilder->AddProperty(GET_MEMBER_NAME_CHECKED(UMLDeformerModel, BoneIncludeList), UMLDeformerModel::StaticClass())
+			InputOutputCategoryBuilder->AddProperty(UMLDeformerModel::GetBoneIncludeListPropertyName(), UMLDeformerModel::StaticClass())
 				.Visibility(EVisibility::Collapsed);
 		}
 
@@ -294,13 +295,13 @@ namespace UE::MLDeformer
 					.OnClicked(FOnClicked::CreateSP(this, &FMLDeformerModelDetails::OnFilterAnimatedCurvesOnly))
 					.IsEnabled_Lambda([this](){ return Model->ShouldIncludeCurvesInTraining() && EditorModel->GetNumCurvesOnSkeletalMesh(Model->GetSkeletalMesh()) > 0; })
 				];
-			CurvesIncludeGroup.AddPropertyRow(DetailBuilder.GetProperty(GET_MEMBER_NAME_CHECKED(UMLDeformerModel, CurveIncludeList), UMLDeformerModel::StaticClass()));
+			CurvesIncludeGroup.AddPropertyRow(DetailBuilder.GetProperty(UMLDeformerModel::GetCurveIncludeListPropertyName(), UMLDeformerModel::StaticClass()));
 
 			AddCurveInputErrors();
 		}
 		else
 		{
-			InputOutputCategoryBuilder->AddProperty(GET_MEMBER_NAME_CHECKED(UMLDeformerModel, CurveIncludeList), UMLDeformerModel::StaticClass())
+			InputOutputCategoryBuilder->AddProperty(UMLDeformerModel::GetCurveIncludeListPropertyName(), UMLDeformerModel::StaticClass())
 				.Visibility(EVisibility::Collapsed);
 		}
 
@@ -365,22 +366,6 @@ namespace UE::MLDeformer
 		return FReply::Handled();
 	}
 
-	void FMLDeformerModelDetails::AddGeomCacheMeshMappingWarnings(IDetailCategoryBuilder* InTargetMeshCategoryBuilder, USkeletalMesh* SkeletalMesh, UGeometryCache* GeometryCache)
-	{
-		const FText MeshMappingError = GetGeomCacheMeshMappingErrorText(SkeletalMesh, GeometryCache);
-		InTargetMeshCategoryBuilder->AddCustomRow(FText::FromString("MeshMappingWarning"))
-			.Visibility(!MeshMappingError.IsEmpty() ? EVisibility::Visible : EVisibility::Collapsed)
-			.WholeRowContent()
-			[
-				SNew(SBox)
-				.Padding(FMargin(0.0f, 4.0f))
-				[
-					SNew(SWarningOrErrorBox)
-					.MessageStyle(EMessageStyle::Warning)
-					.Message(MeshMappingError)
-				]
-			];
-	}
 	bool FMLDeformerModelDetails::IsBonesFlagVisible() const
 	{
 		return (Model->DoesSupportBones() && Model->DoesSupportCurves());
diff --git a/Engine/Plugins/Experimental/MLDeformer/MLDeformerFramework/Source/MLDeformerFrameworkEditor/Private/MLDeformerMorphModelDetails.cpp b/Engine/Plugins/Experimental/MLDeformer/MLDeformerFramework/Source/MLDeformerFrameworkEditor/Private/MLDeformerMorphModelDetails.cpp
--- a/Engine/Plugins/Experimental/MLDeformer/MLDeformerFramework/Source/MLDeformerFrameworkEditor/Private/MLDeformerMorphModelDetails.cpp
+++ b/Engine/Plugins/Experimental/MLDeformer/MLDeformerFramework/Source/MLDeformerFrameworkEditor/Private/MLDeformerMorphModelDetails.cpp
@@ -35,8 +35,8 @@ namespace UE::MLDeformer
 		// Create all the detail categories and add the properties of the base class.
 		FMLDeformerGeomCacheModelDetails::CustomizeDetails(DetailBuilder);
 
-		MorphTargetCategoryBuilder->AddProperty(GET_MEMBER_NAME_CHECKED(UMLDeformerMorphModel, MorphTargetDeltaThreshold), UMLDeformerMorphModel::StaticClass());
-		MorphTargetCategoryBuilder->AddProperty(GET_MEMBER_NAME_CHECKED(UMLDeformerMorphModel, MorphTargetErrorTolerance), UMLDeformerMorphModel::StaticClass());
+		MorphTargetCategoryBuilder->AddProperty(UMLDeformerMorphModel::GetMorphTargetDeltaThresholdPropertyName(), UMLDeformerMorphModel::StaticClass());
+		MorphTargetCategoryBuilder->AddProperty(UMLDeformerMorphModel::GetMorphTargetErrorTolerancePropertyName(), UMLDeformerMorphModel::StaticClass());
 	}
 }	// namespace UE::MLDeformer
 
diff --git a/Engine/Plugins/Experimental/MLDeformer/MLDeformerFramework/Source/MLDeformerFrameworkEditor/Private/MLDeformerMorphModelEditorModel.cpp b/Engine/Plugins/Experimental/MLDeformer/MLDeformerFramework/Source/MLDeformerFrameworkEditor/Private/MLDeformerMorphModelEditorModel.cpp
--- a/Engine/Plugins/Experimental/MLDeformer/MLDeformerFramework/Source/MLDeformerFrameworkEditor/Private/MLDeformerMorphModelEditorModel.cpp
+++ b/Engine/Plugins/Experimental/MLDeformer/MLDeformerFramework/Source/MLDeformerFrameworkEditor/Private/MLDeformerMorphModelEditorModel.cpp
@@ -34,15 +34,22 @@ namespace UE::MLDeformer
 
 		FMLDeformerGeomCacheEditorModel::OnPropertyChanged(PropertyChangedEvent);
 
-		if (Property->GetFName() == GET_MEMBER_NAME_CHECKED(UMLDeformerMorphModelVizSettings, MorphTargetNumber))
+		// Clamp the morph target number, if it changes.
+		if (Property->GetFName() == UMLDeformerMorphModelVizSettings::GetMorphTargetNumberPropertyName())
 		{
-			const UMLDeformerMorphModel* MorphModel = GetMorphModel();
-			UMLDeformerMorphModelVizSettings* MorphViz = GetMorphModelVizSettings();
-			const int32 NumMorphTargets = MorphModel->MorphTargetSet.IsValid() ? MorphModel->MorphTargetSet->MorphBuffers.GetNumMorphs() : 0;
-			MorphViz->MorphTargetNumber = (NumMorphTargets > 0) ? FMath::Min<int32>(MorphViz->MorphTargetNumber, NumMorphTargets - 1) : 0;
+			ClampMorphTargetNumber();
 		}
 	}
 
+	void FMLDeformerMorphModelEditorModel::ClampMorphTargetNumber()
+	{
+		const UMLDeformerMorphModel* MorphModel = GetMorphModel();			
+		UMLDeformerMorphModelVizSettings* MorphViz = GetMorphModelVizSettings();
+		const int32 NumMorphTargets = MorphModel->GetMorphTargetSet().IsValid() ? MorphModel->GetMorphTargetSet()->MorphBuffers.GetNumMorphs() : 0;
+		const int32 ClampedMorphTargetNumber = (NumMorphTargets > 0) ? FMath::Min<int32>(MorphViz->GetMorphTargetNumber(), NumMorphTargets - 1) : 0;
+		MorphViz->SetMorphTargetNumber(ClampedMorphTargetNumber);
+	}
+
 	UMLDeformerMorphModel* FMLDeformerMorphModelEditorModel::GetMorphModel() const
 	{ 
 		return Cast<UMLDeformerMorphModel>(Model);
@@ -66,7 +73,7 @@ namespace UE::MLDeformer
 	void FMLDeformerMorphModelEditorModel::OnPreTraining()
 	{
 		// Backup the morph target deltas in case we abort training.
-		MorphTargetDeltasBackup = GetMorphModel()->MorphTargetDeltas;
+		MorphTargetDeltasBackup = GetMorphModel()->GetMorphTargetDeltas();
 	}
 
 	void FMLDeformerMorphModelEditorModel::OnPostTraining(ETrainingResult TrainingResult, bool bUsePartiallyTrainedWhenAborted)
@@ -75,16 +82,16 @@ namespace UE::MLDeformer
 		if (TrainingResult == ETrainingResult::Aborted && !bUsePartiallyTrainedWhenAborted)
 		{
 			// Restore the blendshape backup.
-			GetMorphModel()->MorphTargetDeltas = MorphTargetDeltasBackup;
+			GetMorphModel()->SetMorphTargetDeltas(MorphTargetDeltasBackup);
 		}
 		else if (TrainingResult == ETrainingResult::Success || (TrainingResult == ETrainingResult::Aborted && bUsePartiallyTrainedWhenAborted))
 		{
 			const UMLDeformerMorphModel* MorphModel = GetMorphModel();
-			if (!MorphModel->MorphTargetDeltas.IsEmpty())
+			if (!MorphModel->GetMorphTargetDeltas().IsEmpty())
 			{
 				// Set deltas with a length equal or below a given threshold to zero, for better compression.
-				TArray<FVector3f> MorphTargetDeltas = MorphModel->MorphTargetDeltas;
-				ZeroDeltasByThreshold(MorphTargetDeltas, MorphModel->MorphTargetDeltaThreshold);
+				TArray<FVector3f> MorphTargetDeltas = MorphModel->GetMorphTargetDeltas();
+				ZeroDeltasByThreshold(MorphTargetDeltas, MorphModel->GetMorphTargetDeltaThreshold());
 
 				// Build morph targets inside the engine, using the engine's compression scheme.
 				// Add one as we included the means now as extra morph target.
@@ -103,12 +110,12 @@ namespace UE::MLDeformer
 		// Turn the delta buffer in a set of engine morph targets.
 		UMLDeformerMorphModel* MorphModel = GetMorphModel();
 		TArray<UMorphTarget*> MorphTargets;	// These will be garbage collected.
-		CreateEngineMorphTargets(MorphTargets, Deltas, TEXT("MLDeformerMorph_"), LOD, MorphModel->MorphTargetDeltaThreshold);
+		CreateEngineMorphTargets(MorphTargets, Deltas, TEXT("MLDeformerMorph_"), LOD, MorphModel->GetMorphTargetDeltaThreshold());
 
 		// Now compress the morph targets to GPU friendly buffers.
-		check(MorphModel->MorphTargetSet.IsValid());
-		FMorphTargetVertexInfoBuffers& MorphBuffers = MorphModel->MorphTargetSet->MorphBuffers;
-		CompressEngineMorphTargets(MorphBuffers, MorphTargets, LOD, MorphModel->MorphTargetErrorTolerance);
+		check(MorphModel->GetMorphTargetSet().IsValid());
+		FMorphTargetVertexInfoBuffers& MorphBuffers = MorphModel->GetMorphTargetSet()->MorphBuffers;
+		CompressEngineMorphTargets(MorphBuffers, MorphTargets, LOD, MorphModel->GetMorphTargetErrorTolerance());
 	}
 
 	void FMLDeformerMorphModelEditorModel::Render(const FSceneView* View, FViewport* Viewport, FPrimitiveDrawInterface* PDI)
@@ -117,10 +124,10 @@ namespace UE::MLDeformer
 
 		// Debug draw the selected morph target.
 		const UMLDeformerMorphModelVizSettings* VizSettings = GetMorphModelVizSettings();
-		if (VizSettings->bDrawMorphTargets && VizSettings->GetVisualizationMode() == EMLDeformerVizMode::TestData)
+		if (VizSettings->GetDrawMorphTargets() && VizSettings->GetVisualizationMode() == EMLDeformerVizMode::TestData)
 		{
 			const FVector DrawOffset = -VizSettings->GetMeshSpacingOffsetVector();
-			DrawMorphTarget(PDI, GetMorphModel()->MorphTargetDeltas, VizSettings->MorphTargetDeltaThreshold, VizSettings->MorphTargetNumber, DrawOffset);
+			DrawMorphTarget(PDI, GetMorphModel()->GetMorphTargetDeltas(), VizSettings->GetMorphTargetDeltaThreshold(), VizSettings->GetMorphTargetNumber(), DrawOffset);
 		}
 	}
 }	// namespace UE::MLDeformer
diff --git a/Engine/Plugins/Experimental/MLDeformer/MLDeformerFramework/Source/MLDeformerFrameworkEditor/Private/MLDeformerMorphModelVizSettingsDetails.cpp b/Engine/Plugins/Experimental/MLDeformer/MLDeformerFramework/Source/MLDeformerFrameworkEditor/Private/MLDeformerMorphModelVizSettingsDetails.cpp
--- a/Engine/Plugins/Experimental/MLDeformer/MLDeformerFramework/Source/MLDeformerFrameworkEditor/Private/MLDeformerMorphModelVizSettingsDetails.cpp
+++ b/Engine/Plugins/Experimental/MLDeformer/MLDeformerFramework/Source/MLDeformerFrameworkEditor/Private/MLDeformerMorphModelVizSettingsDetails.cpp
@@ -36,19 +36,19 @@ namespace UE::MLDeformer
 
 	bool FMLDeformerMorphModelVizSettingsDetails::IsMorphTargetsEnabled() const
 	{
-		return MorphModelVizSettings->bDrawMorphTargets && !MorphModel->MorphTargetDeltas.IsEmpty();
+		return MorphModelVizSettings->GetDrawMorphTargets() && !MorphModel->GetMorphTargetDeltas().IsEmpty();
 	}
 
 	void FMLDeformerMorphModelVizSettingsDetails::AddAdditionalSettings()
 	{
 		IDetailGroup& MorphsGroup = LiveSettingsCategory->AddGroup("Morph Targets", LOCTEXT("MorphTargetsLabel", "Morph Targets"), false, true);
-		MorphsGroup.AddPropertyRow(DetailLayoutBuilder->GetProperty(GET_MEMBER_NAME_CHECKED(UMLDeformerMorphModelVizSettings, bDrawMorphTargets), UMLDeformerMorphModelVizSettings::StaticClass()))
-			.EditCondition(!MorphModel->MorphTargetDeltas.IsEmpty(), nullptr);
+		MorphsGroup.AddPropertyRow(DetailLayoutBuilder->GetProperty(UMLDeformerMorphModelVizSettings::GetDrawMorphTargetsPropertyName(), UMLDeformerMorphModelVizSettings::StaticClass()))
+			.EditCondition(!MorphModel->GetMorphTargetDeltas().IsEmpty(), nullptr);
 
-		MorphsGroup.AddPropertyRow(DetailLayoutBuilder->GetProperty(GET_MEMBER_NAME_CHECKED(UMLDeformerMorphModelVizSettings, MorphTargetNumber), UMLDeformerMorphModelVizSettings::StaticClass()))
+		MorphsGroup.AddPropertyRow(DetailLayoutBuilder->GetProperty(UMLDeformerMorphModelVizSettings::GetMorphTargetNumberPropertyName(), UMLDeformerMorphModelVizSettings::StaticClass()))
 			.EditCondition(TAttribute<bool>(this, &FMLDeformerMorphModelVizSettingsDetails::IsMorphTargetsEnabled), nullptr);
 
-		MorphsGroup.AddPropertyRow(DetailLayoutBuilder->GetProperty(GET_MEMBER_NAME_CHECKED(UMLDeformerMorphModelVizSettings, MorphTargetDeltaThreshold), UMLDeformerMorphModelVizSettings::StaticClass()))
+		MorphsGroup.AddPropertyRow(DetailLayoutBuilder->GetProperty(UMLDeformerMorphModelVizSettings::GetMorphTargetDeltaThresholdPropertyName(), UMLDeformerMorphModelVizSettings::StaticClass()))
 			.EditCondition(TAttribute<bool>(this, &FMLDeformerMorphModelVizSettingsDetails::IsMorphTargetsEnabled), nullptr);
 	}
 }	//namespace UE::MLDeformer
diff --git a/Engine/Plugins/Experimental/MLDeformer/MLDeformerFramework/Source/MLDeformerFrameworkEditor/Private/MLDeformerTrainingModel.cpp b/Engine/Plugins/Experimental/MLDeformer/MLDeformerFramework/Source/MLDeformerFrameworkEditor/Private/MLDeformerTrainingModel.cpp
--- a/Engine/Plugins/Experimental/MLDeformer/MLDeformerFramework/Source/MLDeformerFrameworkEditor/Private/MLDeformerTrainingModel.cpp
+++ b/Engine/Plugins/Experimental/MLDeformer/MLDeformerFramework/Source/MLDeformerFrameworkEditor/Private/MLDeformerTrainingModel.cpp
@@ -16,6 +16,16 @@ void UMLDeformerTrainingModel::Init(UE::MLDeformer::FMLDeformerEditorModel* InEd
 	EditorModel = InEditorModel;
 }
 
+void UMLDeformerTrainingModel::SetEditorModel(UE::MLDeformer::FMLDeformerEditorModel* InModel)
+{ 
+	EditorModel = InModel;
+}
+
+UE::MLDeformer::FMLDeformerEditorModel* UMLDeformerTrainingModel::GetEditorModel() const
+{
+	return EditorModel;
+}
+
 int32 UMLDeformerTrainingModel::GetNumberSampleTransforms() const
 {
 	return GetModel()->GetInputInfo()->GetNumBones();
diff --git a/Engine/Plugins/Experimental/MLDeformer/MLDeformerFramework/Source/MLDeformerFrameworkEditor/Private/MLDeformerVizSettingsDetails.cpp b/Engine/Plugins/Experimental/MLDeformer/MLDeformerFramework/Source/MLDeformerFrameworkEditor/Private/MLDeformerVizSettingsDetails.cpp
--- a/Engine/Plugins/Experimental/MLDeformer/MLDeformerFramework/Source/MLDeformerFrameworkEditor/Private/MLDeformerVizSettingsDetails.cpp
+++ b/Engine/Plugins/Experimental/MLDeformer/MLDeformerFramework/Source/MLDeformerFrameworkEditor/Private/MLDeformerVizSettingsDetails.cpp
@@ -75,15 +75,15 @@ namespace UE::MLDeformer
 		const bool bShowTestData = VizSettings ? (VizSettings->GetVisualizationMode() == EMLDeformerVizMode::TestData) : true;
 
 		// Shared settings.
-		SharedCategoryBuilder->AddProperty(GET_MEMBER_NAME_CHECKED(UMLDeformerVizSettings, bDrawLabels), UMLDeformerVizSettings::StaticClass());
-		SharedCategoryBuilder->AddProperty(GET_MEMBER_NAME_CHECKED(UMLDeformerVizSettings, LabelHeight), UMLDeformerVizSettings::StaticClass());
-		SharedCategoryBuilder->AddProperty(GET_MEMBER_NAME_CHECKED(UMLDeformerVizSettings, LabelScale), UMLDeformerVizSettings::StaticClass());
-		SharedCategoryBuilder->AddProperty(GET_MEMBER_NAME_CHECKED(UMLDeformerVizSettings, MeshSpacing), UMLDeformerVizSettings::StaticClass());
+		SharedCategoryBuilder->AddProperty(UMLDeformerVizSettings::GetDrawLabelsPropertyName(), UMLDeformerVizSettings::StaticClass());
+		SharedCategoryBuilder->AddProperty(UMLDeformerVizSettings::GetLabelHeightPropertyName(), UMLDeformerVizSettings::StaticClass());
+		SharedCategoryBuilder->AddProperty(UMLDeformerVizSettings::GetLabelScalePropertyName(), UMLDeformerVizSettings::StaticClass());
+		SharedCategoryBuilder->AddProperty(UMLDeformerVizSettings::GetMeshSpacingPropertyName(), UMLDeformerVizSettings::StaticClass());
 
 		// Test Assets.
 		TestAssetsCategory->SetCategoryVisibility(bShowTestData);
 	
-		IDetailPropertyRow& TestAnimRow = TestAssetsCategory->AddProperty(GET_MEMBER_NAME_CHECKED(UMLDeformerVizSettings, TestAnimSequence), UMLDeformerVizSettings::StaticClass());
+		IDetailPropertyRow& TestAnimRow = TestAssetsCategory->AddProperty(UMLDeformerVizSettings::GetTestAnimSequencePropertyName(), UMLDeformerVizSettings::StaticClass());
 		TestAnimRow.CustomWidget()
 		.NameContent()
 		[
@@ -122,13 +122,13 @@ namespace UE::MLDeformer
 		FIsResetToDefaultVisible IsResetVisible = FIsResetToDefaultVisible::CreateSP(this, &FMLDeformerVizSettingsDetails::IsResetToDefaultDeformerGraphVisible);
 		FResetToDefaultHandler ResetHandler = FResetToDefaultHandler::CreateSP(this, &FMLDeformerVizSettingsDetails::OnResetToDefaultDeformerGraph);
 		FResetToDefaultOverride ResetOverride = FResetToDefaultOverride::Create(IsResetVisible, ResetHandler);
-		TestAssetsCategory->AddProperty(GET_MEMBER_NAME_CHECKED(UMLDeformerVizSettings, DeformerGraph), UMLDeformerVizSettings::StaticClass()).OverrideResetToDefault(ResetOverride);
+		TestAssetsCategory->AddProperty(UMLDeformerVizSettings::GetDeformerGraphPropertyName(), UMLDeformerVizSettings::StaticClass()).OverrideResetToDefault(ResetOverride);
 
 		AddDeformerGraphErrors();
 
 		// Show a warning when no deformer graph has been selected.
 		UObject* Graph = nullptr;
-		TSharedRef<IPropertyHandle> DeformerGraphProperty = DetailBuilder.GetProperty(GET_MEMBER_NAME_CHECKED(UMLDeformerVizSettings, DeformerGraph), UMLDeformerVizSettings::StaticClass());
+		TSharedRef<IPropertyHandle> DeformerGraphProperty = DetailBuilder.GetProperty(UMLDeformerVizSettings::GetDeformerGraphPropertyName(), UMLDeformerVizSettings::StaticClass());
 		if (DeformerGraphProperty->GetValue(Graph) == FPropertyAccess::Result::Success)
 		{
 			FDetailWidgetRow& GraphErrorRow = TestAssetsCategory->AddCustomRow(FText::FromString("GraphError"))
@@ -148,29 +148,28 @@ namespace UE::MLDeformer
 		AddGroundTruth();
 
 		LiveSettingsCategory->SetCategoryVisibility(bShowTestData);
-		LiveSettingsCategory->AddProperty(GET_MEMBER_NAME_CHECKED(UMLDeformerVizSettings, Weight), UMLDeformerVizSettings::StaticClass());
-		LiveSettingsCategory->AddProperty(GET_MEMBER_NAME_CHECKED(UMLDeformerVizSettings, AnimPlaySpeed), UMLDeformerVizSettings::StaticClass());
-		LiveSettingsCategory->AddProperty(GET_MEMBER_NAME_CHECKED(UMLDeformerVizSettings, TestingFrameNumber), UMLDeformerVizSettings::StaticClass());
+		LiveSettingsCategory->AddProperty(UMLDeformerVizSettings::GetWeightPropertyName(), UMLDeformerVizSettings::StaticClass());
+		LiveSettingsCategory->AddProperty(UMLDeformerVizSettings::GetAnimPlaySpeedPropertyName(), UMLDeformerVizSettings::StaticClass());
+		LiveSettingsCategory->AddProperty(UMLDeformerVizSettings::GetTestingFrameNumberPropertyName(), UMLDeformerVizSettings::StaticClass());
 
 		IDetailGroup& HeatMapGroup = LiveSettingsCategory->AddGroup("HeatMap", LOCTEXT("HeatMap", "Heat Map"), false, true);
-		HeatMapGroup.AddPropertyRow(DetailBuilder.GetProperty(GET_MEMBER_NAME_CHECKED(UMLDeformerVizSettings, bShowHeatMap), UMLDeformerVizSettings::StaticClass()));
-		HeatMapGroup.AddPropertyRow(DetailBuilder.GetProperty(GET_MEMBER_NAME_CHECKED(UMLDeformerVizSettings, HeatMapMode), UMLDeformerVizSettings::StaticClass()));
-		HeatMapGroup.AddPropertyRow(DetailBuilder.GetProperty(GET_MEMBER_NAME_CHECKED(UMLDeformerVizSettings, HeatMapMax), UMLDeformerVizSettings::StaticClass()));
-		HeatMapGroup.AddPropertyRow(DetailBuilder.GetProperty(GET_MEMBER_NAME_CHECKED(UMLDeformerVizSettings, GroundTruthLerp), UMLDeformerVizSettings::StaticClass()));
+		HeatMapGroup.AddPropertyRow(DetailBuilder.GetProperty(UMLDeformerVizSettings::GetShowHeatMapPropertyName(), UMLDeformerVizSettings::StaticClass()));
+		HeatMapGroup.AddPropertyRow(DetailBuilder.GetProperty(UMLDeformerVizSettings::GetHeatMapModePropertyName(), UMLDeformerVizSettings::StaticClass()));
+		HeatMapGroup.AddPropertyRow(DetailBuilder.GetProperty(UMLDeformerVizSettings::GetHeatMapMaxPropertyName(), UMLDeformerVizSettings::StaticClass()));
+		HeatMapGroup.AddPropertyRow(DetailBuilder.GetProperty(UMLDeformerVizSettings::GetGroundTruthLerpPropertyName(), UMLDeformerVizSettings::StaticClass()));
 
 		AddAdditionalSettings();
 
 		IDetailGroup& VisGroup = LiveSettingsCategory->AddGroup("Visibility", LOCTEXT("VisibilityLabel", "Visibility"), false, true);
-		VisGroup.AddPropertyRow(DetailBuilder.GetProperty(GET_MEMBER_NAME_CHECKED(UMLDeformerVizSettings, bDrawLinearSkinnedActor), UMLDeformerVizSettings::StaticClass()));
-		VisGroup.AddPropertyRow(DetailBuilder.GetProperty(GET_MEMBER_NAME_CHECKED(UMLDeformerVizSettings, bDrawMLDeformedActor), UMLDeformerVizSettings::StaticClass()));
-		VisGroup.AddPropertyRow(DetailBuilder.GetProperty(GET_MEMBER_NAME_CHECKED(UMLDeformerVizSettings, bDrawGroundTruthActor), UMLDeformerVizSettings::StaticClass()))
+		VisGroup.AddPropertyRow(DetailBuilder.GetProperty(UMLDeformerVizSettings::GetDrawLinearSkinnedActorPropertyName(), UMLDeformerVizSettings::StaticClass()));
+		VisGroup.AddPropertyRow(DetailBuilder.GetProperty(UMLDeformerVizSettings::GetDrawMLDeformedActorPropertyName(), UMLDeformerVizSettings::StaticClass()));
+		VisGroup.AddPropertyRow(DetailBuilder.GetProperty(UMLDeformerVizSettings::GetDrawGroundTruthActorPropertyName(), UMLDeformerVizSettings::StaticClass()))
 			.EditCondition(VizSettings->HasTestGroundTruth(), nullptr);
 
-		// Training data.
 		TrainingMeshesCategoryBuilder->SetCategoryVisibility(bShowTrainingData);
-		TrainingMeshesCategoryBuilder->AddProperty(GET_MEMBER_NAME_CHECKED(UMLDeformerVizSettings, TrainingFrameNumber), UMLDeformerVizSettings::StaticClass());
-		TrainingMeshesCategoryBuilder->AddProperty(GET_MEMBER_NAME_CHECKED(UMLDeformerVizSettings, bDrawDeltas), UMLDeformerVizSettings::StaticClass());
-		TrainingMeshesCategoryBuilder->AddProperty(GET_MEMBER_NAME_CHECKED(UMLDeformerVizSettings, bXRayDeltas), UMLDeformerVizSettings::StaticClass());
+		TrainingMeshesCategoryBuilder->AddProperty(UMLDeformerVizSettings::GetTrainingFrameNumberPropertyName(), UMLDeformerVizSettings::StaticClass());
+		TrainingMeshesCategoryBuilder->AddProperty(UMLDeformerVizSettings::GetDrawVertexDeltasPropertyName(), UMLDeformerVizSettings::StaticClass());
+		TrainingMeshesCategoryBuilder->AddProperty(UMLDeformerVizSettings::GetXRayDeltasPropertyName(), UMLDeformerVizSettings::StaticClass());
 	}
 
 	bool FMLDeformerVizSettingsDetails::FilterAnimSequences(const FAssetData& AssetData, USkeleton* Skeleton)
diff --git a/Engine/Plugins/Experimental/MLDeformer/MLDeformerFramework/Source/MLDeformerFrameworkEditor/Private/SMLDeformerCurvePickerWidget.cpp b/Engine/Plugins/Experimental/MLDeformer/MLDeformerFramework/Source/MLDeformerFrameworkEditor/Private/SMLDeformerCurvePickerWidget.cpp
--- a/Engine/Plugins/Experimental/MLDeformer/MLDeformerFramework/Source/MLDeformerFrameworkEditor/Private/SMLDeformerCurvePickerWidget.cpp
+++ b/Engine/Plugins/Experimental/MLDeformer/MLDeformerFramework/Source/MLDeformerFrameworkEditor/Private/SMLDeformerCurvePickerWidget.cpp
@@ -14,7 +14,7 @@
 
 namespace UE::MLDeformer
 {
-	void SMLDeformerCurvePickerWidget::Construct(const FArguments& InArgs)
+	void SCurvePickerWidget::Construct(const FArguments& InArgs)
 	{
 		OnCurveNamePicked = InArgs._OnCurveNamePicked;
 		OnGetSkeleton = InArgs._OnGetSkeleton;
@@ -32,7 +32,7 @@ namespace UE::MLDeformer
 			[
 				SAssignNew(SearchBox, SSearchBox)
 					.HintText(LOCTEXT("SearchBoxHint", "Search Curves"))
-					.OnTextChanged(this, &SMLDeformerCurvePickerWidget::HandleFilterTextChanged)
+					.OnTextChanged(this, &SCurvePickerWidget::HandleFilterTextChanged)
 			]
 			+SVerticalBox::Slot()
 			.FillHeight(1.0f)
@@ -41,20 +41,20 @@ namespace UE::MLDeformer
 				.SelectionMode(ESelectionMode::Single)
 				.ItemHeight(20.0f)
 				.ListItemsSource(&CurveNames)
-				.OnSelectionChanged(this, &SMLDeformerCurvePickerWidget::HandleSelectionChanged)
-				.OnGenerateRow(this, &SMLDeformerCurvePickerWidget::HandleGenerateRow)
+				.OnSelectionChanged(this, &SCurvePickerWidget::HandleSelectionChanged)
+				.OnGenerateRow(this, &SCurvePickerWidget::HandleGenerateRow)
 			]
 		];
 
 		RefreshListItems();
 	}
 
-	void SMLDeformerCurvePickerWidget::HandleSelectionChanged(TSharedPtr<FString> InItem, ESelectInfo::Type InSelectionType)
+	void SCurvePickerWidget::HandleSelectionChanged(TSharedPtr<FString> InItem, ESelectInfo::Type InSelectionType)
 	{
 		OnCurveNamePicked.ExecuteIfBound(**InItem.Get());
 	}
 
-	TSharedRef<ITableRow> SMLDeformerCurvePickerWidget::HandleGenerateRow(TSharedPtr<FString> InItem, const TSharedRef<STableViewBase>& InOwnerTable)
+	TSharedRef<ITableRow> SCurvePickerWidget::HandleGenerateRow(TSharedPtr<FString> InItem, const TSharedRef<STableViewBase>& InOwnerTable)
 	{
 		return 
 			SNew(STableRow<TSharedPtr<FString>>, InOwnerTable)
@@ -70,7 +70,7 @@ namespace UE::MLDeformer
 			];
 	}
 
-	void SMLDeformerCurvePickerWidget::RefreshListItems()
+	void SCurvePickerWidget::RefreshListItems()
 	{
 		CurveNames.Empty();
 		UniqueCurveNames.Empty();
@@ -95,7 +95,7 @@ namespace UE::MLDeformer
 		FilterAvailableCurves();
 	}
 
-	void SMLDeformerCurvePickerWidget::FilterAvailableCurves()
+	void SCurvePickerWidget::FilterAvailableCurves()
 	{
 		CurveNames.Empty();
 
@@ -110,7 +110,7 @@ namespace UE::MLDeformer
 		NameListView->RequestListRefresh();
 	}
 
-	void SMLDeformerCurvePickerWidget::HandleFilterTextChanged(const FText& InFilterText)
+	void SCurvePickerWidget::HandleFilterTextChanged(const FText& InFilterText)
 	{
 		FilterText = InFilterText.ToString();
 		FilterAvailableCurves();
@@ -145,7 +145,7 @@ namespace UE::MLDeformer
 	{
 		bool bMultipleValues = false;
 
-		TSharedRef<SMLDeformerCurvePickerWidget> ListWidget = SNew(SMLDeformerCurvePickerWidget)
+		TSharedRef<SCurvePickerWidget> ListWidget = SNew(SCurvePickerWidget)
 			.OnCurveNamePicked(this, &SCurveSelectionWidget::OnSelectionChanged)
 			.OnGetSkeleton(OnGetSkeleton);
 
diff --git a/Engine/Plugins/Experimental/MLDeformer/MLDeformerFramework/Source/MLDeformerFrameworkEditor/Private/SMLDeformerTimeline.h b/Engine/Plugins/Experimental/MLDeformer/MLDeformerFramework/Source/MLDeformerFrameworkEditor/Private/SMLDeformerTimeline.h
--- a/Engine/Plugins/Experimental/MLDeformer/MLDeformerFramework/Source/MLDeformerFrameworkEditor/Private/SMLDeformerTimeline.h
+++ b/Engine/Plugins/Experimental/MLDeformer/MLDeformerFramework/Source/MLDeformerFrameworkEditor/Private/SMLDeformerTimeline.h
@@ -27,7 +27,7 @@ namespace UE::MLDeformer
 	class SAnimTimelineTransportControls; 
 
 	/**
-	 * Implements a timeline widget for the ML Deformer Editor.
+	 * Implements a timeline widget for the ML Deformer editor.
 	 */
 	class SMLDeformerTimeline : public SCompoundWidget
 	{
diff --git a/Engine/Plugins/Experimental/MLDeformer/MLDeformerFramework/Source/MLDeformerFrameworkEditor/Public/MLDeformerApplicationMode.h b/Engine/Plugins/Experimental/MLDeformer/MLDeformerFramework/Source/MLDeformerFrameworkEditor/Public/MLDeformerApplicationMode.h
--- a/Engine/Plugins/Experimental/MLDeformer/MLDeformerFramework/Source/MLDeformerFrameworkEditor/Public/MLDeformerApplicationMode.h
+++ b/Engine/Plugins/Experimental/MLDeformer/MLDeformerFramework/Source/MLDeformerFrameworkEditor/Public/MLDeformerApplicationMode.h
@@ -12,10 +12,15 @@ namespace UE::MLDeformer
 {
 	class FMLDeformerEditorToolkit;
 
+	/**
+	 * The application mode for the ML Deformer editor.
+	 * This defines the layout of the UI. It basically spawns all the tabs, such as the viewport, details panels, etc.
+	 */
 	class MLDEFORMERFRAMEWORKEDITOR_API FMLDeformerApplicationMode
 		: public FApplicationMode
 	{
 	public:
+		/** The name of this mode. */
 		static FName ModeName;
 
 		FMLDeformerApplicationMode(TSharedRef<FWorkflowCentricApplication> InHostingApp, TSharedRef<IPersonaPreviewScene> InPreviewScene);
diff --git a/Engine/Plugins/Experimental/MLDeformer/MLDeformerFramework/Source/MLDeformerFrameworkEditor/Public/MLDeformerCurveReferenceCustomization.h b/Engine/Plugins/Experimental/MLDeformer/MLDeformerFramework/Source/MLDeformerFrameworkEditor/Public/MLDeformerCurveReferenceCustomization.h
--- a/Engine/Plugins/Experimental/MLDeformer/MLDeformerFramework/Source/MLDeformerFrameworkEditor/Public/MLDeformerCurveReferenceCustomization.h
+++ b/Engine/Plugins/Experimental/MLDeformer/MLDeformerFramework/Source/MLDeformerFrameworkEditor/Public/MLDeformerCurveReferenceCustomization.h
@@ -3,7 +3,6 @@
 #pragma once
 
 #include "CoreMinimal.h"
-#include "Widgets/DeclarativeSyntaxSupport.h"
 #include "IPropertyTypeCustomization.h"
 #include "PropertyHandle.h"
 #include "DetailWidgetRow.h"
@@ -12,6 +11,9 @@ class USkeleton;
 
 namespace UE::MLDeformer
 {
+	/**
+	 * The curve reference property detail customization.
+	 */
 	class MLDEFORMERFRAMEWORKEDITOR_API FMLDeformerCurveReferenceCustomization
 		: public IPropertyTypeCustomization
 	{
diff --git a/Engine/Plugins/Experimental/MLDeformer/MLDeformerFramework/Source/MLDeformerFrameworkEditor/Public/MLDeformerEditorActor.h b/Engine/Plugins/Experimental/MLDeformer/MLDeformerFramework/Source/MLDeformerFrameworkEditor/Public/MLDeformerEditorActor.h
--- a/Engine/Plugins/Experimental/MLDeformer/MLDeformerFramework/Source/MLDeformerFrameworkEditor/Public/MLDeformerEditorActor.h
+++ b/Engine/Plugins/Experimental/MLDeformer/MLDeformerFramework/Source/MLDeformerFrameworkEditor/Public/MLDeformerEditorActor.h
@@ -11,24 +11,57 @@ class UMLDeformerComponent;
 
 namespace UE::MLDeformer
 {
+	/**
+	 * The default ID's for editor actors.
+	 * You can assign any ID to the editor actors, although some code in the editor model classes will assume
+	 * that some of these exist. You can always overload certain methods to change that behavior though.
+	 * @see FMLDeformerEditorActor
+	 */
 	enum : int32
 	{
+		/** The linear skinned (base) actor in training mode. */
 		ActorID_Train_Base,
+
+		/** The ground truth (target) actor in training mode. This is our training target that we try to approximate. */
 		ActorID_Train_GroundTruth,
+
+		/** The linear skinned (base) actor, in test mode. */
 		ActorID_Test_Base,
+
+		/** The ML Deformed actor, in test mode. This one should have an UMLDeformerComponent on it. */
 		ActorID_Test_MLDeformed,
+
+		/** The ground truth test actor, in test mode. This will represent your test anim sequence ground truth. */
 		ActorID_Test_GroundTruth
 	};
 
+	/**
+	 * An actor in the ML Deformer asset editor viewport.
+	 * The UE Actor object itself is created by the FMLDeformerEditorModel class, after which this editor actor class
+	 * can create components on it.
+	 * It also contains some methods to control visibility and playback, among some other helper methods.
+	 */
 	class MLDEFORMERFRAMEWORKEDITOR_API FMLDeformerEditorActor
 	{
 	public:
+		/**
+		 * The construction settings, which can be used as info when creating new components.
+		 */
 		struct MLDEFORMERFRAMEWORKEDITOR_API FConstructSettings
 		{
+			/** The UE actor object that we should add components to. */
 			AActor* Actor = nullptr;
+
+			/** The ID of the actor. Look at the UE::MLDeformer::ActorID_Train_Base etc for example. */
 			int32 TypeID = -1;
+
+			/** The color of the label text. */
 			FLinearColor LabelColor = FLinearColor(1.0f, 0.0f, 0.0f);
+
+			/** The label text, which should be something like "Linear Skinned" or "Ground Truth", basically describing the actor. This will appear inside the render viewport. */
 			FText LabelText;
+
+			/** Is this an actor used during training? Set to false if it is a testing mode actor. */
 			bool bIsTrainingActor = false;
 		};
 
diff --git a/Engine/Plugins/Experimental/MLDeformer/MLDeformerFramework/Source/MLDeformerFrameworkEditor/Public/MLDeformerEditorMode.h b/Engine/Plugins/Experimental/MLDeformer/MLDeformerFramework/Source/MLDeformerFrameworkEditor/Public/MLDeformerEditorMode.h
--- a/Engine/Plugins/Experimental/MLDeformer/MLDeformerFramework/Source/MLDeformerFrameworkEditor/Public/MLDeformerEditorMode.h
+++ b/Engine/Plugins/Experimental/MLDeformer/MLDeformerFramework/Source/MLDeformerFrameworkEditor/Public/MLDeformerEditorMode.h
@@ -12,10 +12,15 @@ namespace UE::MLDeformer
 	class FMLDeformerEditorToolkit;
 	class FMLDeformerPreviewScene;
 
+	/**
+	 * The ML Deformer Persona editor mode.
+	 * This handles ticking and rendering of the render viewport and camera focus.
+	 */
 	class MLDEFORMERFRAMEWORKEDITOR_API FMLDeformerEditorMode
 		: public IPersonaEditMode
 	{
 	public:
+		/** The name of the mode. */
 		static FName ModeName;
 
 		void SetEditorToolkit(FMLDeformerEditorToolkit* InToolkit) { DeformerEditorToolkit = InToolkit; }
diff --git a/Engine/Plugins/Experimental/MLDeformer/MLDeformerFramework/Source/MLDeformerFrameworkEditor/Public/MLDeformerEditorModule.h b/Engine/Plugins/Experimental/MLDeformer/MLDeformerFramework/Source/MLDeformerFrameworkEditor/Public/MLDeformerEditorModule.h
--- a/Engine/Plugins/Experimental/MLDeformer/MLDeformerFramework/Source/MLDeformerFrameworkEditor/Public/MLDeformerEditorModule.h
+++ b/Engine/Plugins/Experimental/MLDeformer/MLDeformerFramework/Source/MLDeformerFrameworkEditor/Public/MLDeformerEditorModule.h
@@ -10,6 +10,11 @@ namespace UE::MLDeformer
 {
 	class FMLDeformerAssetActions;
 
+	/**
+	 * The ML Deformer editor module.
+	 * This registers the editor mode and custom property customizations.
+	 * It also contains the model registry, which you will use to register your custom models to the editor.
+	 */
 	class MLDEFORMERFRAMEWORKEDITOR_API FMLDeformerEditorModule
 		: public IModuleInterface
 	{
@@ -19,12 +24,19 @@ namespace UE::MLDeformer
 		void ShutdownModule() override;
 		// ~END IModuleInterface overrides.
 
-		FMLDeformerEditorModelRegistry& GetModelRegistry() { return ModelRegistry; }
-		const FMLDeformerEditorModelRegistry& GetModelRegistry() const { return ModelRegistry; }
+		/**
+		 * Get a reference to the ML Deformer model registry.
+		 * You use this to see what kind of model types ther are, register new models, or create editor models for specific types of runtime models.
+		 * @return A reference to the ML Deformer model registry.
+		 */
+		FMLDeformerEditorModelRegistry& GetModelRegistry()				{ return ModelRegistry; }
+		const FMLDeformerEditorModelRegistry& GetModelRegistry() const	{ return ModelRegistry; }
 
 	private:
+		/** The model registry, which keeps track of all model types and instances created of these models, and how to create the editor models for specific runtime model types. */
 		FMLDeformerEditorModelRegistry ModelRegistry;
+
+		/** The asset actions for the ML Deformer asset type. */
 		TSharedPtr<FMLDeformerAssetActions> MLDeformerAssetActions;
 	};
-
 }	// namespace UE::MLDeformer
diff --git a/Engine/Plugins/Experimental/MLDeformer/MLDeformerFramework/Source/MLDeformerFrameworkEditor/Public/MLDeformerEditorStyle.h b/Engine/Plugins/Experimental/MLDeformer/MLDeformerFramework/Source/MLDeformerFrameworkEditor/Public/MLDeformerEditorStyle.h
--- a/Engine/Plugins/Experimental/MLDeformer/MLDeformerFramework/Source/MLDeformerFrameworkEditor/Public/MLDeformerEditorStyle.h
+++ b/Engine/Plugins/Experimental/MLDeformer/MLDeformerFramework/Source/MLDeformerFrameworkEditor/Public/MLDeformerEditorStyle.h
@@ -6,6 +6,9 @@
 
 namespace UE::MLDeformer
 {
+	/**
+	 * The editor style class that describes specific UI style related settings for the ML Deformer editor.
+	 */
 	class MLDEFORMERFRAMEWORKEDITOR_API FMLDeformerEditorStyle
 		: public FSlateStyleSet
 	{
@@ -15,5 +18,4 @@ namespace UE::MLDeformer
 
 		static FMLDeformerEditorStyle& Get();
 	};
-
 }	// namespace UE::MLDeformer
diff --git a/Engine/Plugins/Experimental/MLDeformer/MLDeformerFramework/Source/MLDeformerFrameworkEditor/Public/MLDeformerEditorToolkit.h b/Engine/Plugins/Experimental/MLDeformer/MLDeformerFramework/Source/MLDeformerFrameworkEditor/Public/MLDeformerEditorToolkit.h
--- a/Engine/Plugins/Experimental/MLDeformer/MLDeformerFramework/Source/MLDeformerFrameworkEditor/Public/MLDeformerEditorToolkit.h
+++ b/Engine/Plugins/Experimental/MLDeformer/MLDeformerFramework/Source/MLDeformerFrameworkEditor/Public/MLDeformerEditorToolkit.h
@@ -27,6 +27,10 @@ namespace UE::MLDeformer
 		extern const FName Editor;
 	}
 
+	/**
+	 * The ML Deformer asset editor toolkit.
+	 * This is the editor that opens when you double click an ML Deformer asset.
+	 */
 	class MLDEFORMERFRAMEWORKEDITOR_API FMLDeformerEditorToolkit :
 		public FPersonaAssetEditorToolkit,
 		public IHasPersonaToolkit,
@@ -129,7 +133,6 @@ namespace UE::MLDeformer
 		/** The timeline slider widget. */
 		TSharedPtr<SMLDeformerTimeline> TimeSlider = nullptr;
 
-
 		/** The currently active editor model. */
 		TSharedPtr<FMLDeformerEditorModel> ActiveModel = nullptr;
 
diff --git a/Engine/Plugins/Experimental/MLDeformer/MLDeformerFramework/Source/MLDeformerFrameworkEditor/Public/MLDeformerGeomCacheActor.h b/Engine/Plugins/Experimental/MLDeformer/MLDeformerFramework/Source/MLDeformerFrameworkEditor/Public/MLDeformerGeomCacheActor.h
--- a/Engine/Plugins/Experimental/MLDeformer/MLDeformerFramework/Source/MLDeformerFrameworkEditor/Public/MLDeformerGeomCacheActor.h
+++ b/Engine/Plugins/Experimental/MLDeformer/MLDeformerFramework/Source/MLDeformerFrameworkEditor/Public/MLDeformerGeomCacheActor.h
@@ -9,6 +9,10 @@ class UGeometryCacheComponent;
 
 namespace UE::MLDeformer
 {
+	/**
+	 * An editor actor with a geometry cache component on it.
+	 * This can for example be used as ground truth viewport actors.
+	 */
 	class MLDEFORMERFRAMEWORKEDITOR_API FMLDeformerGeomCacheActor
 		: public FMLDeformerEditorActor
 	{
@@ -16,8 +20,8 @@ namespace UE::MLDeformer
 		FMLDeformerGeomCacheActor(const FConstructSettings& Settings);
 		virtual ~FMLDeformerGeomCacheActor() override;
 
-		void SetGeometryCacheComponent(UGeometryCacheComponent* Component) { GeomCacheComponent = Component; }
-		UGeometryCacheComponent* GetGeometryCacheComponent() const { return GeomCacheComponent; }
+		void SetGeometryCacheComponent(UGeometryCacheComponent* Component)	{ GeomCacheComponent = Component; }
+		UGeometryCacheComponent* GetGeometryCacheComponent() const			{ return GeomCacheComponent; }
 
 		// FMLDeformerEditorActor overrides.
 		virtual void SetVisibility(bool bIsVisible) override;
@@ -32,7 +36,7 @@ namespace UE::MLDeformer
 		// ~END FMLDeformerEditorActor overrides.
 
 	protected:
-		/** The geometry cache component (can be nullptr). */
+		/** The geometry cache component. */
 		TObjectPtr<UGeometryCacheComponent> GeomCacheComponent = nullptr;
 	};
 }	// namespace UE::MLDeformer
diff --git a/Engine/Plugins/Experimental/MLDeformer/MLDeformerFramework/Source/MLDeformerFrameworkEditor/Public/MLDeformerGeomCacheEditorModel.h b/Engine/Plugins/Experimental/MLDeformer/MLDeformerFramework/Source/MLDeformerFrameworkEditor/Public/MLDeformerGeomCacheEditorModel.h
--- a/Engine/Plugins/Experimental/MLDeformer/MLDeformerFramework/Source/MLDeformerFrameworkEditor/Public/MLDeformerGeomCacheEditorModel.h
+++ b/Engine/Plugins/Experimental/MLDeformer/MLDeformerFramework/Source/MLDeformerFrameworkEditor/Public/MLDeformerGeomCacheEditorModel.h
@@ -14,7 +14,9 @@ namespace UE::MLDeformer
 	class FMLDeformerGeomCacheSampler;
 
 	/**
-	 * An editor model that is based on a geometry cache as training input.
+	 * An editor model for the a model derived from the UMLDeformerGeomCacheModel runtime model.
+	 * This editor model automatically handles sampling of the geometry caches to calculate deltas.
+	 * It also handles the creation of ground truth models with a geometry cache on them.
 	 */
 	class MLDEFORMERFRAMEWORKEDITOR_API FMLDeformerGeomCacheEditorModel
 		: public FMLDeformerEditorModel
@@ -47,6 +49,18 @@ namespace UE::MLDeformer
 		FMLDeformerGeomCacheActor* FindGeomCacheEditorActor(int32 ID) const;
 
 	protected:
+		/**
+		 * Creates a geometry cache editor actor (FMLDeformerGeomCacheActor) of a specific ID with a specific geometry cache and label.
+		 * This will add a new editor actor to the list of editor actors.
+		 * @param World The world to create the actor in.
+		 * @param ActorID The ID that the actor will have after creation, for example UE::MLDeformer::ActorID_Train_GroundTruth.
+		 * @param Name The name of the actor in the scene.
+		 * @param GeomCache The geometry cache to use on this actor.
+		 * @param LabelColor The color of the label that is rendered with the editor actor inside the editor viewport.
+		 * @param WireframeColor The color of this actor when wireframe rendering is enabled.
+		 * @param LabelText The text of the label.
+		 * @param bIsTrainingActor Set this to true when the actor is an actor to be used in training mode, or set to false when it is to be used in testing mode.
+		 */
 		void CreateGeomCacheActor(UWorld* World, int32 ActorID, const FName& Name, UGeometryCache* GeomCache, FLinearColor LabelColor, FLinearColor WireframeColor, const FText& LabelText, bool bIsTrainingActor);
 	};
 }	// namespace UE::MLDeformer
diff --git a/Engine/Plugins/Experimental/MLDeformer/MLDeformerFramework/Source/MLDeformerFrameworkEditor/Public/MLDeformerGeomCacheModelDetails.h b/Engine/Plugins/Experimental/MLDeformer/MLDeformerFramework/Source/MLDeformerFrameworkEditor/Public/MLDeformerGeomCacheModelDetails.h
--- a/Engine/Plugins/Experimental/MLDeformer/MLDeformerFramework/Source/MLDeformerFrameworkEditor/Public/MLDeformerGeomCacheModelDetails.h
+++ b/Engine/Plugins/Experimental/MLDeformer/MLDeformerFramework/Source/MLDeformerFrameworkEditor/Public/MLDeformerGeomCacheModelDetails.h
@@ -8,12 +8,20 @@
 
 class IDetailLayoutBuilder;
 class UMLDeformerGeomCacheModel;
+class UGeometryCache;
+class USkeletalMesh;
 
 namespace UE::MLDeformer
 {
 	class FMLDeformerGeomCacheEditorModel;
 	class FMLDeformerEditorModel;
 
+	/**
+	 * The detail customization of a geometry cache based model.
+	 * The model it customizes needs to be derived from UMLDeformerGeomCacheModel.
+	 * It will automatically insert some error when the geom cache doesn't match the animation sequence etc.
+	 * If your model inherits from a UMLDeformerGeomCacheModel then you can also inherit your model detail customization from this FMLDeformerGeomCacheModelDetails class.
+	 */
 	class MLDEFORMERFRAMEWORKEDITOR_API FMLDeformerGeomCacheModelDetails
 		: public FMLDeformerModelDetails
 	{
@@ -25,7 +33,14 @@ namespace UE::MLDeformer
 		// ~END FMLDeformerModelDetails overrides.
 
 	protected:
+		/** Add warnings to the UI about mesh mapping errors when mapping between geometry caches and a skeletal mesh. */
+		void AddGeomCacheMeshMappingWarnings(IDetailCategoryBuilder* InTargetMeshCategoryBuilder, USkeletalMesh* SkeletalMesh, UGeometryCache* GeometryCache);
+
+	protected:
+		/** The geometry cache based runtime model. This value gets updated once UpdateMemberPointers gets called. */
 		TObjectPtr<UMLDeformerGeomCacheModel> GeomCacheModel = nullptr;
+
+		/** The geometry cache based editor model. This value gets updated once UpdateMemberPointers gets called. */
 		FMLDeformerGeomCacheEditorModel* GeomCacheEditorModel = nullptr;
 	};
 }	// namespace UE::MLDeformer
diff --git a/Engine/Plugins/Experimental/MLDeformer/MLDeformerFramework/Source/MLDeformerFrameworkEditor/Public/MLDeformerGeomCacheSampler.h b/Engine/Plugins/Experimental/MLDeformer/MLDeformerFramework/Source/MLDeformerFrameworkEditor/Public/MLDeformerGeomCacheSampler.h
--- a/Engine/Plugins/Experimental/MLDeformer/MLDeformerFramework/Source/MLDeformerFrameworkEditor/Public/MLDeformerGeomCacheSampler.h
+++ b/Engine/Plugins/Experimental/MLDeformer/MLDeformerFramework/Source/MLDeformerFrameworkEditor/Public/MLDeformerGeomCacheSampler.h
@@ -17,9 +17,8 @@ namespace UE::MLDeformer
 	DECLARE_DELEGATE_RetVal(UGeometryCache*, FMLDeformerGetGeomCacheEvent)
 
 	/**
-	 * The input data sampler.
-	 * This class can sample bone rotations, curve values and vertex deltas.
-	 * It does this by creating two temp actors, one with skeletal mesh component and one with geom cache component.
+	 * The input data sampler, which is used to sample vertex positions from geometry caches.
+	 * It can then also calculate deltas between the sampled skeletal mesh data and geometry cache data.
 	 */
 	class MLDEFORMERFRAMEWORKEDITOR_API FMLDeformerGeomCacheSampler
 		: public FMLDeformerSampler
@@ -31,11 +30,33 @@ namespace UE::MLDeformer
 		virtual float GetTimeAtFrame(int32 InAnimFrameIndex) const override;
 		// ~END FVertexDeltaSampler overrides.
 
-		const TArray<FString>& GetFailedImportedMeshNames() const { return FailedImportedMeshNames; }
+		/**
+		 * Get the array of mesh names that we cannot do any sampling for.
+		 * The reason for this can be that the sampler cannot find a matching mesh in the skeletal mesh for specific geometry cache tracks.
+		 * For example if there is some geometry cache track that is named "Head", but there is no such mesh inside the skeletal mesh, then
+		 * the returned array of names will include a string with the value "Head".
+		 * @return An array of geometry cache track names for which no mesh could be found inside the SkeletalMesh / linear skinned actor.
+		 */
+		const TArray<FString>& GetFailedImportedMeshNames() const	{ return FailedImportedMeshNames; }
 
-		FMLDeformerGetGeomCacheEvent& OnGetGeometryCache() { return GetGeometryCacheEvent; }
+		/**
+		 * Executed when we need a pointer to the geometry cache object that we want to sample from.
+		 * @return A reference to the delegate.
+		 */
+		FMLDeformerGetGeomCacheEvent& OnGetGeometryCache()			{ return GetGeometryCacheEvent; }
 
 	protected:
+		/**
+		 * Calculate the vertex deltas between the linear skinned mesh (skeletal mesh) and geometry cache.
+		 * A delta cutoff length can be specified. Deltas with a length longer than this cutoff length will be ignored and set to zero.
+		 * This can be useful in cases where for some reason due to some mesh errors some deltas are very long.
+		 * Setting the DeltaCutoffLength to a very large value would essentially disable this filtering and always include all deltas.
+		 * Typically a good value for the delta cutoff length seems to be 30.
+		 * @param SkinnedPositions The vertex positions of the linear skinned mesh, so of the skeletal mesh.
+		 * @param DeltaCutoffLength The delta cutoff length as described above. Usually a value of 30 is a good value, or set to a very large value to disable it.
+		 * @param OutVertexDeltas The array that will receive the deltas, in a float buffer. The buffer will contain 3*NumVerts number of elements, and the layout is
+		 *                        (x, y, z, x, y, z, x, y, z...) so 3 values per vertex: x, y and z.
+		 */
 		void CalculateVertexDeltas(const TArray<FVector3f>& SkinnedPositions, float DeltaCutoffLength, TArray<float>& OutVertexDeltas);
 
 	protected:
diff --git a/Engine/Plugins/Experimental/MLDeformer/MLDeformerFramework/Source/MLDeformerFrameworkEditor/Public/MLDeformerGeomCacheVizSettingsDetails.h b/Engine/Plugins/Experimental/MLDeformer/MLDeformerFramework/Source/MLDeformerFrameworkEditor/Public/MLDeformerGeomCacheVizSettingsDetails.h
--- a/Engine/Plugins/Experimental/MLDeformer/MLDeformerFramework/Source/MLDeformerFrameworkEditor/Public/MLDeformerGeomCacheVizSettingsDetails.h
+++ b/Engine/Plugins/Experimental/MLDeformer/MLDeformerFramework/Source/MLDeformerFrameworkEditor/Public/MLDeformerGeomCacheVizSettingsDetails.h
@@ -4,14 +4,17 @@
 
 #include "CoreMinimal.h"
 #include "MLDeformerVizSettingsDetails.h"
-#include "IDetailCustomization.h"
 
-class IDetailLayoutBuilder;
 class UMLDeformerGeomCacheModel;
 class UMLDeformerGeomCacheVizSettings;
 
 namespace UE::MLDeformer
 {
+	/**
+	 * The visualization settings detail customization for a geometry cache based model.
+	 * A geometry cache based model is one inherited from the UMLDeformerGeomCacheModel base class.
+	 * You can inherit your own model detail customization from this class if you also inherited your model from the UMLDeformerGeomCacheModel class.
+	 */
 	class MLDEFORMERFRAMEWORKEDITOR_API FMLDeformerGeomCacheVizSettingsDetails
 		: public FMLDeformerVizSettingsDetails
 	{
@@ -22,7 +25,10 @@ namespace UE::MLDeformer
 		// ~END FMLDeformerVizSettingsDetails overrides.
 
 	protected:
+		/** The geometry cache based runtime model pointer. This is updated once UpdateMemberPointers is called. */
 		TObjectPtr<UMLDeformerGeomCacheModel> GeomCacheModel = nullptr;
+
+		/** A pointer to the geom cache based visualization settings. This is updated once UpdateMemberPointers is called. */
 		TObjectPtr<UMLDeformerGeomCacheVizSettings> GeomCacheVizSettings = nullptr;
 	};
 }	// namespace UE::MLDeformer
diff --git a/Engine/Plugins/Experimental/MLDeformer/MLDeformerFramework/Source/MLDeformerFrameworkEditor/Public/MLDeformerModelDetails.h b/Engine/Plugins/Experimental/MLDeformer/MLDeformerFramework/Source/MLDeformerFrameworkEditor/Public/MLDeformerModelDetails.h
--- a/Engine/Plugins/Experimental/MLDeformer/MLDeformerFramework/Source/MLDeformerFrameworkEditor/Public/MLDeformerModelDetails.h
+++ b/Engine/Plugins/Experimental/MLDeformer/MLDeformerFramework/Source/MLDeformerFrameworkEditor/Public/MLDeformerModelDetails.h
@@ -18,6 +18,13 @@ namespace UE::MLDeformer
 {
 	class FMLDeformerEditorModel;
 
+	/**
+	 * The ML Deformer model detail customization base class.
+	 * This adds the shared properties, inserts error messages, creates some groups, etc.
+	 * Model detail customizations should all inherit from this class. If you use a geometry cache based model however, you will 
+	 * likely inherit from other classes such as FMLDeformerGeomCacheModelDetails or FMLDeformerMorphModelDetails. 
+	 * Those classes are also inherited from this one though.
+	 */
 	class MLDEFORMERFRAMEWORKEDITOR_API FMLDeformerModelDetails
 		: public IDetailCustomization
 	{
@@ -26,35 +33,77 @@ namespace UE::MLDeformer
 		virtual void CustomizeDetails(IDetailLayoutBuilder& DetailBuilder) override;
 		// ~END ILayoutDetails overrides.
 
+		/**
+		 * Update the class member pointers, which includes the pointer to the model, and its editor model.
+		 * @param Objects The array of objects that the detail customization is showing.
+		 */
 		virtual bool UpdateMemberPointers(const TArray<TWeakObjectPtr<UObject>>& Objects);
+
+		/**
+		 * Create the categories that we can add properties to. This will update the category member pointers. */
 		virtual void CreateCategories();
+
+		/**
+		 * Add the target mesh property. For example the FMLDeformerGeomCacheModelDetails class will add the geometry cache property here. */
 		virtual void AddTargetMesh() {}
+
+		/** Add additional errors related to the base mesh. */
 		virtual void AddBaseMeshErrors() {}
+
+		/** Add additional errors related to the training animation sequence. */
 		virtual void AddAnimSequenceErrors() {}
+
+		/** Add additional errors related to bone inputs. */
 		virtual void AddBoneInputErrors() {}
+
+		/** Add additional errors related to curve inputs. */
 		virtual void AddCurveInputErrors() {}
+
+		/** Add training input flags, which basically are things like the check boxes that specify whether bones or curves (or both) should be included. */
 		virtual void AddTrainingInputFlags() {}
+
+		/** Add additional training input filters. Filters are things like a list of bones, or curves. */
 		virtual void AddTrainingInputFilters() {}
+
+		/** Add additional training input errors. */
 		virtual void AddTrainingInputErrors() {}
+
+		/** Is the "include bones" checkbox visible? */
 		virtual bool IsBonesFlagVisible() const;
+
+		/** Is the "include curves" checkbox visible? */
 		virtual bool IsCurvesFlagVisible() const;
 
 	protected:
+		/** The filter that only shows anim sequences that are compatible with the given skeleton. */
 		bool FilterAnimSequences(const FAssetData& AssetData, USkeleton* Skeleton);
-		void AddGeomCacheMeshMappingWarnings(IDetailCategoryBuilder* InTargetMeshCategoryBuilder, USkeletalMesh* SkeletalMesh, UGeometryCache* GeometryCache);
+
+		/** Executed when we press the Animated Bones Only button. */
 		FReply OnFilterAnimatedBonesOnly() const;
+
+		/** Executed when we press the Animated Curves Only button. */
 		FReply OnFilterAnimatedCurvesOnly() const;
 
 	protected:
 		/** Associated detail layout builder. */
 		IDetailLayoutBuilder* DetailLayoutBuilder = nullptr;
 
+		/** A pointer to the model we are customizing details for. */
 		TObjectPtr<UMLDeformerModel> Model = nullptr;
+
+		/** A pointer to the editor model of the runtime model. */
 		FMLDeformerEditorModel* EditorModel = nullptr;
 
+		/** The category related to the base mesh. */
 		IDetailCategoryBuilder* BaseMeshCategoryBuilder = nullptr;
+
+		/** The category related to the target mesh. */
 		IDetailCategoryBuilder* TargetMeshCategoryBuilder = nullptr;
+
+		/** The category related to the inputs and outputs. */
 		IDetailCategoryBuilder* InputOutputCategoryBuilder = nullptr;
+
+		/** The training settings category. You most likely add your model properties to this. */
 		IDetailCategoryBuilder* TrainingSettingsCategoryBuilder = nullptr;
 	};
 }	// namespace UE::MLDeformer
diff --git a/Engine/Plugins/Experimental/MLDeformer/MLDeformerFramework/Source/MLDeformerFrameworkEditor/Public/MLDeformerMorphModelDetails.h b/Engine/Plugins/Experimental/MLDeformer/MLDeformerFramework/Source/MLDeformerFrameworkEditor/Public/MLDeformerMorphModelDetails.h
--- a/Engine/Plugins/Experimental/MLDeformer/MLDeformerFramework/Source/MLDeformerFrameworkEditor/Public/MLDeformerMorphModelDetails.h
+++ b/Engine/Plugins/Experimental/MLDeformer/MLDeformerFramework/Source/MLDeformerFrameworkEditor/Public/MLDeformerMorphModelDetails.h
@@ -13,6 +13,10 @@ namespace UE::MLDeformer
 {
 	class FMLDeformerMorphModelEditorModel;
 
+	/**
+	 * The detail customization for models inherited from the UMLDeformerMorphModel class.
+	 * You can inherit the detail customization for your own model from this class if your model inherited from the UMLDeformerMorphModel class.
+	 */
 	class MLDEFORMERFRAMEWORKEDITOR_API FMLDeformerMorphModelDetails
 		: public FMLDeformerGeomCacheModelDetails
 	{
@@ -27,8 +31,13 @@ namespace UE::MLDeformer
 		// ~END FMLDeformerModelDetails overrides.
 
 	protected:
+		/** A pointer to the morph model. This is updated when UpdateMemberPointers is called. */
 		TObjectPtr<UMLDeformerMorphModel> MorphModel = nullptr;
+
+		/** A pointer to the editor model for the morph model. This is updated when UpdateMemberPointers is called. */
 		FMLDeformerMorphModelEditorModel* MorphModelEditorModel = nullptr;
+
+		/** The morph settings category. */
 		IDetailCategoryBuilder* MorphTargetCategoryBuilder = nullptr;
 	};
 }	// namespace UE::MLDeformer
diff --git a/Engine/Plugins/Experimental/MLDeformer/MLDeformerFramework/Source/MLDeformerFrameworkEditor/Public/MLDeformerMorphModelEditorModel.h b/Engine/Plugins/Experimental/MLDeformer/MLDeformerFramework/Source/MLDeformerFrameworkEditor/Public/MLDeformerMorphModelEditorModel.h
--- a/Engine/Plugins/Experimental/MLDeformer/MLDeformerFramework/Source/MLDeformerFrameworkEditor/Public/MLDeformerMorphModelEditorModel.h
+++ b/Engine/Plugins/Experimental/MLDeformer/MLDeformerFramework/Source/MLDeformerFrameworkEditor/Public/MLDeformerMorphModelEditorModel.h
@@ -6,19 +6,15 @@
 #include "MLDeformerGeomCacheEditorModel.h"
 #include "MLDeformerGeomCacheSampler.h"
 #include "MLDeformerEditorActor.h"
-#include "UObject/GCObject.h"
 
-class UMLDeformerAsset;
-class USkeletalMesh;
-class UGeometryCache;
 class UMLDeformerMorphModel;
 class UMLDeformerMorphModelVizSettings;
 
 namespace UE::MLDeformer
 {
-	class FMLDeformerGeomCacheActor;
-	class FNeuralMorphEditorModelActor;
-
+	/**
+	 * The editor model related to the runtime UMLDeformerMorphModel class, or models inherited from that.
+	 */
 	class MLDEFORMERFRAMEWORKEDITOR_API FMLDeformerMorphModelEditorModel
 		: public FMLDeformerGeomCacheEditorModel
 	{
@@ -53,6 +49,12 @@ namespace UE::MLDeformer
 		 */
 		void InitEngineMorphTargets(const TArray<FVector3f>& Deltas);
 
+		/**
+		 * Clamp the current morph target number inside the visualization settings to a valid range.
+		 * This won't let the number of the morph target that we are visualizing go above the total number of morph targets.
+		 */
+		void ClampMorphTargetNumber();
+
 	protected:
 		/**
 		 * The entire set of morph target deltas, 3 per vertex, for each morph target, as one flattened buffer.
diff --git a/Engine/Plugins/Experimental/MLDeformer/MLDeformerFramework/Source/MLDeformerFrameworkEditor/Public/MLDeformerMorphModelVizSettingsDetails.h b/Engine/Plugins/Experimental/MLDeformer/MLDeformerFramework/Source/MLDeformerFrameworkEditor/Public/MLDeformerMorphModelVizSettingsDetails.h
--- a/Engine/Plugins/Experimental/MLDeformer/MLDeformerFramework/Source/MLDeformerFrameworkEditor/Public/MLDeformerMorphModelVizSettingsDetails.h
+++ b/Engine/Plugins/Experimental/MLDeformer/MLDeformerFramework/Source/MLDeformerFrameworkEditor/Public/MLDeformerMorphModelVizSettingsDetails.h
@@ -4,14 +4,15 @@
 
 #include "CoreMinimal.h"
 #include "MLDeformerGeomCacheVizSettingsDetails.h"
-#include "IDetailCustomization.h"
 
-class IDetailLayoutBuilder;
 class UMLDeformerMorphModel;
 class UMLDeformerMorphModelVizSettings;
 
 namespace UE::MLDeformer
 {
+	/**
+	 * The visualization settings for models inherited from the UMLDeformerMorphModel class.
+	 */
 	class MLDEFORMERFRAMEWORKEDITOR_API FMLDeformerMorphModelVizSettingsDetails
 		: public FMLDeformerGeomCacheVizSettingsDetails
 	{
@@ -22,9 +23,13 @@ namespace UE::MLDeformer
 		// ~END FMLDeformerVizSettingsDetails overrides.
 
 	protected:
+		/** Is the morph target visualization option/checkbox enabled? */
 		bool IsMorphTargetsEnabled() const;
 
+		/** A pointer to the runtime morph model. This is updated when UpdateMemberPointers is called. */
 		TObjectPtr<UMLDeformerMorphModel> MorphModel = nullptr;
+
+		/** A pointer to the morph model visualization settings. This is updated when UpdateMemberPointers is called. */
 		TObjectPtr<UMLDeformerMorphModelVizSettings> MorphModelVizSettings = nullptr;
 	};
 }	// namespace UE::MLDeformer
diff --git a/Engine/Plugins/Experimental/MLDeformer/MLDeformerFramework/Source/MLDeformerFrameworkEditor/Public/MLDeformerSampler.h b/Engine/Plugins/Experimental/MLDeformer/MLDeformerFramework/Source/MLDeformerFrameworkEditor/Public/MLDeformerSampler.h
--- a/Engine/Plugins/Experimental/MLDeformer/MLDeformerFramework/Source/MLDeformerFrameworkEditor/Public/MLDeformerSampler.h
+++ b/Engine/Plugins/Experimental/MLDeformer/MLDeformerFramework/Source/MLDeformerFrameworkEditor/Public/MLDeformerSampler.h
@@ -17,10 +17,13 @@ namespace UE::MLDeformer
 {
 	class FMLDeformerEditorModel;
 
+	/**
+	 * The space in which the sampler will calculate its vertex deltas.
+	 */
 	enum class EVertexDeltaSpace
 	{
-		PreSkinning,	/** Apply the deltas before skinning. */
-		PostSkinning	/** Apply the deltas after skinning. */
+		PreSkinning,	/** Vertex deltas will be in pre-skinning space. We train on pre-skinning deltas. */
+		PostSkinning	/** Vertex deltas will be in post-skinning space. Mostly used to visualize deltas in the viewport. */
 	};
 
 	/**
@@ -40,31 +43,75 @@ namespace UE::MLDeformer
 		virtual void Sample(int32 AnimFrameIndex);
 
 		/** Get the time in seconds, at a given frame index. */
-		virtual float GetTimeAtFrame(int32 InAnimFrameIndex) const { return 0.0f; }
-
-		const TArray<FVector3f>& GetSkinnedVertexPositions() const { return SkinnedVertexPositions; }
-		const TArray<FVector3f>& GetUnskinnedVertexPositions() const { return UnskinnedVertexPositions; }
-		const TArray<float>& GetVertexDeltas() const { return VertexDeltas; }
-		const TArray<float>& GetBoneRotations() const { return BoneRotations; }
-		const TArray<float>& GetCurveValues() const { return CurveValues; }
-		int32 GetNumImportedVertices() const { return NumImportedVertices; }
+		virtual float GetTimeAtFrame(int32 InAnimFrameIndex) const		{ return 0.0f; }
+
+		/** Get the skinned vertex positions of the skeletal mesh. */
+		const TArray<FVector3f>& GetSkinnedVertexPositions() const		{ return SkinnedVertexPositions; }
+
+		/** Get the unskinned vertex positions of the skeletal mesh. */
+		const TArray<FVector3f>& GetUnskinnedVertexPositions() const	{ return UnskinnedVertexPositions; }
+
+		/** Get the calculated vertex deltas. The space they will be in depends on the setup vertex deltas space. */
+		const TArray<float>& GetVertexDeltas() const					{ return VertexDeltas; }
+
+		/** 
+		 * Get the current bone rotations that we can feed to a neural network. 
+		 * The number of floats in the buffer equals to NumBones * 6.
+		 * The six floats represent two columns of the 3x3 rotation matrix of the bone.
+		 * @return The arrray of floats that we can pass to the neural network as inputs.
+		 */
+		const TArray<float>& GetBoneRotations() const					{ return BoneRotations; }
+
+		/** Get the float values for the curves. One float per curve. */
+		const TArray<float>& GetCurveValues() const						{ return CurveValues; }
+
+		/** Get the number of imported vertices. */
+		int32 GetNumImportedVertices() const							{ return NumImportedVertices; }
+
+		/** Get the space in which we will calculate the vertex deltas. */
+		EVertexDeltaSpace GetVertexDeltaSpace() const					{ return VertexDeltaSpace; }
+
+		/** Set the space in which we should calculate the vertex deltas. */
+		void SetVertexDeltaSpace(EVertexDeltaSpace DeltaSpace)			{ VertexDeltaSpace = DeltaSpace; }
+
+		/** Have we already initialized this sampler? Is it ready for sampling? */
+		bool IsInitialized() const										{ return SkelMeshActor.Get() != nullptr; }
+
+		/** Get the number of bones. */
 		int32 GetNumBones() const;
-		EVertexDeltaSpace GetVertexDeltaSpace() const { return VertexDeltaSpace; }
+
+		/** Calculate the memory usage of the sampler. */
 		SIZE_T CalcMemUsagePerFrameInBytes() const;
-		void SetVertexDeltaSpace(EVertexDeltaSpace DeltaSpace) { VertexDeltaSpace = DeltaSpace; }
-		bool IsInitialized() const { return SkelMeshActor.Get() != nullptr; }
 
 	protected:
+		/** Create the actors used for sampling. This creates two actors, one for the base skeletal mesh and one for the target mesh. */
 		virtual void CreateActors();
+
+		/** Register the target mesh components. */
 		virtual void RegisterTargetComponents() {}
+
+		/** Extract the skinned vertex positions, of the skeletal mesh. */
 		void ExtractSkinnedPositions(int32 LODIndex, TArray<FMatrix44f>& InBoneMatrices, TArray<FVector3f>& TempPositions, TArray<FVector3f>& OutPositions) const;
+
+		/** Extract the unskinned vertex positions, of the skeletal mesh. */
 		void ExtractUnskinnedPositions(int32 LODIndex, TArray<FVector3f>& OutPositions) const;
+
+		/** Calculate the inverse skinning transform, which is used to convert from post-skinning to pre-skinning space. */
 		FMatrix44f CalcInverseSkinningTransform(int32 VertexIndex, const FSkeletalMeshLODRenderData& SkelMeshLODData, const FSkinWeightVertexBuffer& SkinWeightBuffer) const;
+
+		/** Create a new actor with a specific name, in a specified world. */
 		AActor* CreateNewActor(UWorld* InWorld, const FName& Name) const;
 
+		/** Update the skeletal mesh component's play position, and refresh its bone transforms internally. */
 		void UpdateSkeletalMeshComponent();
+
+		/** Update the skinned vertex positions. */
 		void UpdateSkinnedPositions();
+
+		/** Update the bone rotation buffers. */
 		void UpdateBoneRotations();
+
+		/** Update the curve value buffers. */
 		void UpdateCurveValues();
 
 	protected:
diff --git a/Engine/Plugins/Experimental/MLDeformer/MLDeformerFramework/Source/MLDeformerFrameworkEditor/Public/MLDeformerTrainingModel.h b/Engine/Plugins/Experimental/MLDeformer/MLDeformerFramework/Source/MLDeformerFrameworkEditor/Public/MLDeformerTrainingModel.h
--- a/Engine/Plugins/Experimental/MLDeformer/MLDeformerFramework/Source/MLDeformerFrameworkEditor/Public/MLDeformerTrainingModel.h
+++ b/Engine/Plugins/Experimental/MLDeformer/MLDeformerFramework/Source/MLDeformerFrameworkEditor/Public/MLDeformerTrainingModel.h
@@ -1,4 +1,5 @@
 // Copyright Epic Games, Inc. All Rights Reserved.
+
 #pragma once
 
 #include "CoreMinimal.h"
@@ -13,6 +14,33 @@ namespace UE::MLDeformer
 	class FMLDeformerEditorModel;
 }
 
+/**
+ * The training model base class.
+ * This class is used to interface with Python by providing some methods you can call inside your python training code.
+ * For example it allows you to get all the sampled data, such as the deltas, bones and curve values.
+ * When you create a new model you need to create a class inherited from this base class, and define a Train method inside it as follows:
+ *
+ * @code{.cpp}
+ * // Doesn't need an implementation inside cpp, just a declaration in the header file.
+ * UFUNCTION(BlueprintImplementableEvent, Category = "Training Model")
+ * int32 Train() const;
+ * @endcode
+ * 
+ * Now inside your Python class you do something like:
+ * 
+ * @code{.py}
+ * @unreal.uclass()
+ * class YourModelPythonTrainingModel(unreal.YourTrainingModel):
+ *     @unreal.ufunction(override=True)
+ *     def train(self):
+ *         # ...do training here...
+ *         return 0   # A value of 0 is success, 1 means aborted, see ETrainingResult.
+ * @endcode
+ * 
+ * The editor will execute the Train method, which will trigger the "train" method in your Python class to be executed.
+ * Keep in mind that in Unreal Engine all python code is lower case. So a "Train" method inside c++ will need to be called "train" inside the python code.
+ * Or if you have something called "PerformMyTraining" it will need to be called "perform_my_training" inside Python.
+ */
 UCLASS(Blueprintable)
 class MLDEFORMERFRAMEWORKEDITOR_API UMLDeformerTrainingModel
 	: public UObject
@@ -20,15 +48,18 @@ class MLDEFORMERFRAMEWORKEDITOR_API UMLDeformerTrainingModel
 	GENERATED_BODY()
 
 public:
+	/**
+	 * Initialize the training model.
+	 * This is automatically called by the editor.
+	 * @param InEditorModel The pointer to the editor model that this is a training model for.
+	 */
 	void Init(UE::MLDeformer::FMLDeformerEditorModel* InEditorModel);
 
-	UE::MLDeformer::FMLDeformerEditorModel* GetEditorModel() const { return EditorModel; }
-	void SetEditorModel(UE::MLDeformer::FMLDeformerEditorModel* InModel) { EditorModel = InModel; }
-
+	/** Get the runtime ML Deformer model object. */
 	UFUNCTION(BlueprintPure, Category = "Training Data")
 	UMLDeformerModel* GetModel() const;
 
-	/** Get the number of input transforms. */
+	/** Get the number of input transforms. This is the number of bones. */
 	UFUNCTION(BlueprintPure, Category = "Training Data")
 	int32 GetNumberSampleTransforms() const;
 
@@ -36,28 +67,45 @@ public:
 	UFUNCTION(BlueprintPure, Category = "Training Data")
 	int32 GetNumberSampleCurves() const;
 
-	/** Get the number of output deltas. */
+	/** Get the number of vertex deltas. */
 	UFUNCTION(BlueprintPure, Category = "Training Data")
 	int32 GetNumberSampleDeltas() const;
 
-	/** Get the number of samples in this data set. */
+	/** Get the number of samples in this data set. This is the number of sample frames we want to train on. */
 	UFUNCTION(BlueprintPure, Category = "Training Data")
     int32 NumSamples() const;
 
-	/** Set the current sample frame. This will internally call the SampleFrame method, which will update the deltas, curve values and bone rotations. */
+	/** 
+	 * Set the current sample frame. This will internally call the SampleFrame method, which will update the deltas, curve values and bone rotations. 
+	 * You call this before getting the input bone/curve and vertex delta values.
+	 * @param Index The training data frame/sample number.
+	 * @return Returns true when successful, or false when the specified sample index is out of range.
+	 */
 	UFUNCTION(BlueprintCallable, Category = "Training Data")
 	bool SetCurrentSampleIndex(int32 Index);
 
+	/** 
+	 * Check whether we need to resample the inputs and outputs, or if we can use a cached version. 
+	 * This will return true if any inputs changed, that indicate that we should regenerate any cached data.
+	 * @return Returns true when we need to regenerate any cached data, otherwise false is returned.
+	 */
 	UFUNCTION(BlueprintCallable, Category = "Training Data")
 	bool GetNeedsResampling() const;
 
+	/** Specify whether we need to resample any cached data or not, because our input assets, or any other relevant settings changed that would invalidate the cached data. */
 	UFUNCTION(BlueprintCallable, Category = "Training Data")
 	void SetNeedsResampling(bool bNeedsResampling);
 
 protected:
-	/** Sample a given frame. This updates the sample deltas, curves, and bone rotations. */
+	/** This updates the sample deltas, curves, and bone rotations. */
     virtual bool SampleFrame(int32 Index);
 
+	/** Change the pointer to the editor model. */
+	void SetEditorModel(UE::MLDeformer::FMLDeformerEditorModel* InModel);
+
+	/** Get a pointer to the editor model. */ 
+	UE::MLDeformer::FMLDeformerEditorModel* GetEditorModel() const;
+
 public:
 	// The delta values per vertex for this sample. This is updated after SetCurrentSampleIndex is called. Contains an xyz (3 floats) for each vertex.
 	UPROPERTY(EditAnywhere, BlueprintReadWrite, Category = "Training Data")
@@ -72,5 +120,6 @@ public:
 	TArray<float> SampleBoneRotations;
 
 protected:
+	/** A pointer to the editor model from which we use the sampler. */
 	UE::MLDeformer::FMLDeformerEditorModel* EditorModel = nullptr;
 };
diff --git a/Engine/Plugins/Experimental/MLDeformer/MLDeformerFramework/Source/MLDeformerFrameworkEditor/Public/MLDeformerVizSettingsDetails.h b/Engine/Plugins/Experimental/MLDeformer/MLDeformerFramework/Source/MLDeformerFrameworkEditor/Public/MLDeformerVizSettingsDetails.h
--- a/Engine/Plugins/Experimental/MLDeformer/MLDeformerFramework/Source/MLDeformerFrameworkEditor/Public/MLDeformerVizSettingsDetails.h
+++ b/Engine/Plugins/Experimental/MLDeformer/MLDeformerFramework/Source/MLDeformerFrameworkEditor/Public/MLDeformerVizSettingsDetails.h
@@ -17,6 +17,10 @@ namespace UE::MLDeformer
 {
 	class FMLDeformerEditorModel;
 
+	/**
+	 * The detail customization base class for any UMLDeformerVizSettings inherited object.
+	 * This automatically adds the settings of the UMLDeformerVizSettings class to the UI and creates some groups, inserts error messages when needed, etc.
+	 */
 	class MLDEFORMERFRAMEWORKEDITOR_API FMLDeformerVizSettingsDetails
 		: public IDetailCustomization
 	{
@@ -25,31 +29,63 @@ namespace UE::MLDeformer
 		virtual void CustomizeDetails(IDetailLayoutBuilder& DetailBuilder) override;
 		// ~END ILayoutDetails overrides.
 
+		/** Get a pointer to the detail layout builder. */
 		IDetailLayoutBuilder* GetDetailLayoutBuilder() const { return DetailLayoutBuilder; }
 
 	protected:
+		/** The filter that only shows anim sequences that are compatible with the given skeleton. */
 		bool FilterAnimSequences(const FAssetData& AssetData, USkeleton* Skeleton);
+
+		/** When we click the reset to default button next to the deformer graph property. */
 		void OnResetToDefaultDeformerGraph(TSharedPtr<IPropertyHandle> PropertyHandle);
+
+		/** Check if the reset to default button next to the deformer graph property should be visible or not. */
 		bool IsResetToDefaultDeformerGraphVisible(TSharedPtr<IPropertyHandle> PropertyHandle);
 
+		/**
+		 * Update the class member pointers, which includes the pointer to the model, its editor model, and the viz settings.
+		 * @param Objects The array of objects that the detail customization is showing.
+		 */
 		virtual bool UpdateMemberPointers(const TArray<TWeakObjectPtr<UObject>>& Objects);
+
+		/** Create the categories to which we can add properties. Update the member category pointers to those. */
 		virtual void CreateCategories();
+
+		/** Add some error message related to the test animation sequence. */
 		virtual void AddTestSequenceErrors() {}
+
+		/** Add some error related to the deformer graph. */
 		virtual void AddDeformerGraphErrors() {}
+
+		/** Add the ground truth test sequence property. */
 		virtual void AddGroundTruth() {}
+
+		/** Add any additional settings. */
 		virtual void AddAdditionalSettings() {}
 
 	protected:
 		/** Associated detail layout builder. */
 		IDetailLayoutBuilder* DetailLayoutBuilder = nullptr;
 
+		/** The ML deformer model that these visualization settings belong to. */
 		TObjectPtr<UMLDeformerModel> Model = nullptr;
+
+		/** The visualization settings we are viewing. */
 		TObjectPtr<UMLDeformerVizSettings> VizSettings = nullptr;
+
+		/** The editor model of the runtime model. */
 		FMLDeformerEditorModel* EditorModel = nullptr;
 
+		/** A pointer to the shared settings category. */
 		IDetailCategoryBuilder* SharedCategoryBuilder = nullptr;
+
+		/** A pointer to the test assets category. */
 		IDetailCategoryBuilder* TestAssetsCategory = nullptr;
+
+		/** A pointer to the live settings category. */
 		IDetailCategoryBuilder* LiveSettingsCategory = nullptr;
+
+		/** A pointer to the training mesh category. This is only visible in training mode. */
 		IDetailCategoryBuilder* TrainingMeshesCategoryBuilder = nullptr;
 	};
 }	// namespace UE::MLDeformer
diff --git a/Engine/Plugins/Experimental/MLDeformer/MLDeformerFramework/Source/MLDeformerFrameworkEditor/Public/SMLDeformerCurvePickerWidget.h b/Engine/Plugins/Experimental/MLDeformer/MLDeformerFramework/Source/MLDeformerFrameworkEditor/Public/SMLDeformerCurvePickerWidget.h
--- a/Engine/Plugins/Experimental/MLDeformer/MLDeformerFramework/Source/MLDeformerFrameworkEditor/Public/SMLDeformerCurvePickerWidget.h
+++ b/Engine/Plugins/Experimental/MLDeformer/MLDeformerFramework/Source/MLDeformerFrameworkEditor/Public/SMLDeformerCurvePickerWidget.h
@@ -15,19 +15,22 @@ class SSearchBox;
 class USkeleton;
 
 DECLARE_DELEGATE_OneParam(FOnCurveSelectionChanged, const FString&);
+DECLARE_DELEGATE_OneParam(FOnCurveNamePicked, const FString&)
 
 DECLARE_DELEGATE_RetVal(FString, FOnGetSelectedCurve);
 DECLARE_DELEGATE_RetVal(USkeleton*, FOnGetSkeleton);
 
-DECLARE_DELEGATE_OneParam(FOnCurveNamePicked, const FString&)
 
 namespace UE::MLDeformer
 {
-	class MLDEFORMERFRAMEWORKEDITOR_API SMLDeformerCurvePickerWidget
+	/**
+	 * The curve picker widget for the ML Deformer editor.
+	 */
+	class MLDEFORMERFRAMEWORKEDITOR_API SCurvePickerWidget
 		: public SCompoundWidget
 	{
 	public:
-		SLATE_BEGIN_ARGS(SMLDeformerCurvePickerWidget) {}
+		SLATE_BEGIN_ARGS(SCurvePickerWidget) {}
 		SLATE_EVENT(FOnCurveNamePicked, OnCurveNamePicked)
 		SLATE_EVENT(FOnGetSkeleton, OnGetSkeleton)
 		SLATE_END_ARGS()
@@ -39,10 +42,9 @@ namespace UE::MLDeformer
 	private:
 		void RefreshListItems();
 		void FilterAvailableCurves();
-
 		void HandleSelectionChanged(TSharedPtr<FString> InItem, ESelectInfo::Type InSelectionType);
-		TSharedRef<ITableRow> HandleGenerateRow(TSharedPtr<FString> InItem, const TSharedRef<STableViewBase>& InOwnerTable);
 		void HandleFilterTextChanged(const FText& InFilterText);
+		TSharedRef<ITableRow> HandleGenerateRow(TSharedPtr<FString> InItem, const TSharedRef<STableViewBase>& InOwnerTable);
 
 	private:
 		/** Delegate fired when a curve name is picked. */
@@ -70,6 +72,9 @@ namespace UE::MLDeformer
 		TSharedPtr<SListView<TSharedPtr<FString>>> NameListView;
 	};
 
+	/**
+	 * The curve selection widget for the ML Deformer.
+	 */
 	class SCurveSelectionWidget
 		: public SCompoundWidget
 	{
diff --git a/Engine/Plugins/Experimental/MLDeformer/NeuralMorphModel/Source/NeuralMorphModel/Public/NeuralMorphModel.h b/Engine/Plugins/Experimental/MLDeformer/NeuralMorphModel/Source/NeuralMorphModel/Public/NeuralMorphModel.h
--- a/Engine/Plugins/Experimental/MLDeformer/NeuralMorphModel/Source/NeuralMorphModel/Public/NeuralMorphModel.h
+++ b/Engine/Plugins/Experimental/MLDeformer/NeuralMorphModel/Source/NeuralMorphModel/Public/NeuralMorphModel.h
@@ -3,10 +3,6 @@
 
 #include "CoreMinimal.h"
 #include "MLDeformerMorphModel.h"
-#include "MLDeformerVizSettings.h"
-#include "MLDeformerInputInfo.h"
-#include "MLDeformerGeomCacheHelpers.h"
-#include "UObject/Object.h"
 #include "NeuralMorphModel.generated.h"
 
 class USkeletalMesh;
@@ -16,7 +12,6 @@ class UNeuralNetwork;
 class UMLDeformerAsset;
 class UMLDeformerModelInstance;
 class USkeleton;
-class IPropertyHandle;
 struct FExternalMorphSet;
 
 NEURALMORPHMODEL_API DECLARE_LOG_CATEGORY_EXTERN(LogNeuralMorphModel, Log, All);
@@ -47,6 +42,15 @@ enum class ENeuralMorphMode : uint8
 	Global
 };
 
+/**
+ * The neural morph model.
+ * This generates a set of highly compressed morph targets to approximate a target deformation based on bone rotations and/or curve inputs.
+ * The neural network inside this model runs on the CPU and outputs the morph target weights.
+ * An external morph target set is generated during training and serialized inside the ML Deformer asset that contains this model.
+ * When the ML Deformer component initializes the morph target set is registered. Most of the heavy lifting is done by the UMLDeformerMorphModel class.
+ * The neural morph model has two modes: local and global. See the ENeuralMorphMode for a description of the two.
+ * @see ENeuralMorphMode
+ */
 UCLASS()
 class NEURALMORPHMODEL_API UNeuralMorphModel 
 	: public UMLDeformerMorphModel
@@ -57,18 +61,18 @@ public:
 	UNeuralMorphModel(const FObjectInitializer& ObjectInitializer);
 
 	// UMLDeformerModel overrides.
-	virtual FString GetDisplayName() const override { return "Neural Morph Model"; }
+	virtual FString GetDisplayName() const override		{ return "Neural Morph Model"; }
 	// ~END UMLDeformerModel overrides.
 
 #if WITH_EDITORONLY_DATA
-	int32 GetLocalNumHiddenLayers() const { return LocalNumHiddenLayers; }
-	int32 GetLocalNumNeuronsPerLayer() const { return LocalNumNeuronsPerLayer; }
-	int32 GetGlobalNumHiddenLayers() const { return GlobalNumHiddenLayers; }
-	int32 GetGlobalNumNeuronsPerLayer() const { return GlobalNumNeuronsPerLayer; }
-	int32 GetNumIterations() const { return NumIterations; }
-	int32 GetBatchSize() const { return BatchSize; }
-	float GetLearningRate() const { return LearningRate; }
-	float GetRegularizationFactor() const { return RegularizationFactor;  }
+	int32 GetLocalNumHiddenLayers() const				{ return LocalNumHiddenLayers; }
+	int32 GetLocalNumNeuronsPerLayer() const			{ return LocalNumNeuronsPerLayer; }
+	int32 GetGlobalNumHiddenLayers() const				{ return GlobalNumHiddenLayers; }
+	int32 GetGlobalNumNeuronsPerLayer() const			{ return GlobalNumNeuronsPerLayer; }
+	int32 GetNumIterations() const						{ return NumIterations; }
+	int32 GetBatchSize() const							{ return BatchSize; }
+	float GetLearningRate() const						{ return LearningRate; }
+	float GetRegularizationFactor() const				{ return RegularizationFactor;  }
 
 public:
 	/**
diff --git a/Engine/Plugins/Experimental/MLDeformer/NeuralMorphModel/Source/NeuralMorphModel/Public/NeuralMorphModelVizSettings.h b/Engine/Plugins/Experimental/MLDeformer/NeuralMorphModel/Source/NeuralMorphModel/Public/NeuralMorphModelVizSettings.h
--- a/Engine/Plugins/Experimental/MLDeformer/NeuralMorphModel/Source/NeuralMorphModel/Public/NeuralMorphModelVizSettings.h
+++ b/Engine/Plugins/Experimental/MLDeformer/NeuralMorphModel/Source/NeuralMorphModel/Public/NeuralMorphModelVizSettings.h
@@ -3,7 +3,6 @@
 #pragma once
 
 #include "CoreMinimal.h"
-#include "UObject/Object.h"
 #include "MLDeformerMorphModelVizSettings.h"
 #include "NeuralMorphModelVizSettings.generated.h"
 
diff --git a/Engine/Plugins/Experimental/MLDeformer/NeuralMorphModel/Source/NeuralMorphModelEditor/Public/NeuralMorphEditorModel.h b/Engine/Plugins/Experimental/MLDeformer/NeuralMorphModel/Source/NeuralMorphModelEditor/Public/NeuralMorphEditorModel.h
--- a/Engine/Plugins/Experimental/MLDeformer/NeuralMorphModel/Source/NeuralMorphModelEditor/Public/NeuralMorphEditorModel.h
+++ b/Engine/Plugins/Experimental/MLDeformer/NeuralMorphModel/Source/NeuralMorphModelEditor/Public/NeuralMorphEditorModel.h
@@ -11,6 +11,9 @@ namespace UE::NeuralMorphModel
 {
 	using namespace UE::MLDeformer;
 
+	/**
+	 * The editor model related to the neural morph model's runtime class (UNeuralMorphModel).
+	 */
 	class NEURALMORPHMODELEDITOR_API FNeuralMorphEditorModel
 		: public UE::MLDeformer::FMLDeformerMorphModelEditorModel
 	{
@@ -19,15 +22,15 @@ namespace UE::NeuralMorphModel
 		static FMLDeformerEditorModel* MakeInstance();
 
 		// FGCObject overrides.
-		virtual FString GetReferencerName() const override { return TEXT("FNeuralMorphEditorModel"); }
+		virtual FString GetReferencerName() const override		{ return TEXT("FNeuralMorphEditorModel"); }
 		// ~END FGCObject overrides.
 	
 		// FMLDeformerEditorModel overrides.
 		virtual ETrainingResult Train() override;
-		virtual void OnPropertyChanged(FPropertyChangedEvent& PropertyChangedEvent);
+		virtual void OnPropertyChanged(FPropertyChangedEvent& PropertyChangedEvent) override;
 		// ~END FMLDeformerEditorModel overrides.
 
-		// Some helpers that cast to this model's variants of some classes.
-		UNeuralMorphModel* GetNeuralMorphModel() const { return Cast<UNeuralMorphModel>(Model); }
+		/** Get a pointer to the neural morph runtime model. */
+		UNeuralMorphModel* GetNeuralMorphModel() const			{ return Cast<UNeuralMorphModel>(Model); }
 	};
 }	// namespace UE::NeuralMorphModel
diff --git a/Engine/Plugins/Experimental/MLDeformer/NeuralMorphModel/Source/NeuralMorphModelEditor/Public/NeuralMorphModelDetails.h b/Engine/Plugins/Experimental/MLDeformer/NeuralMorphModel/Source/NeuralMorphModelEditor/Public/NeuralMorphModelDetails.h
--- a/Engine/Plugins/Experimental/MLDeformer/NeuralMorphModel/Source/NeuralMorphModelEditor/Public/NeuralMorphModelDetails.h
+++ b/Engine/Plugins/Experimental/MLDeformer/NeuralMorphModel/Source/NeuralMorphModelEditor/Public/NeuralMorphModelDetails.h
@@ -7,12 +7,12 @@
 #include "IDetailCustomization.h"
 
 class IDetailLayoutBuilder;
-class UNeuralMorphModel;
 
 namespace UE::NeuralMorphModel
 {
-	class FNeuralMorphEditorModel;
-
+	/**
+	 * The detail customization for the neural morph model.
+	 */
 	class NEURALMORPHMODELEDITOR_API FNeuralMorphModelDetails
 		: public UE::MLDeformer::FMLDeformerMorphModelDetails
 	{
diff --git a/Engine/Plugins/Experimental/MLDeformer/NeuralMorphModel/Source/NeuralMorphModelEditor/Public/NeuralMorphTrainingModel.h b/Engine/Plugins/Experimental/MLDeformer/NeuralMorphModel/Source/NeuralMorphModelEditor/Public/NeuralMorphTrainingModel.h
--- a/Engine/Plugins/Experimental/MLDeformer/NeuralMorphModel/Source/NeuralMorphModelEditor/Public/NeuralMorphTrainingModel.h
+++ b/Engine/Plugins/Experimental/MLDeformer/NeuralMorphModel/Source/NeuralMorphModelEditor/Public/NeuralMorphTrainingModel.h
@@ -6,6 +6,10 @@
 #include "MLDeformerTrainingModel.h"
 #include "NeuralMorphTrainingModel.generated.h"
 
+/**
+ * The training model for the neural morph model.
+ * This class is our link to the Python training.
+ */
 UCLASS(Blueprintable)
 class NEURALMORPHMODELEDITOR_API UNeuralMorphTrainingModel
 	: public UMLDeformerTrainingModel
@@ -13,7 +17,11 @@ class NEURALMORPHMODELEDITOR_API UNeuralMorphTrainingModel
 	GENERATED_BODY()
 
 public:
-	/** Main training function, with implementation in python. */
+	/**
+	 * Main training function, with implementation in python.
+	 * You need to implement this method. See the UMLDeformerTrainingModel documentation for more details.
+	 * @see UMLDeformerTrainingModel
+	 */
 	UFUNCTION(BlueprintImplementableEvent, Category = "Training Model")
 	int32 Train() const;
 };
diff --git a/Engine/Plugins/Experimental/MLDeformer/VertexDeltaModel/Source/VertexDeltaModel/Public/VertexDeltaModel.h b/Engine/Plugins/Experimental/MLDeformer/VertexDeltaModel/Source/VertexDeltaModel/Public/VertexDeltaModel.h
--- a/Engine/Plugins/Experimental/MLDeformer/VertexDeltaModel/Source/VertexDeltaModel/Public/VertexDeltaModel.h
+++ b/Engine/Plugins/Experimental/MLDeformer/VertexDeltaModel/Source/VertexDeltaModel/Public/VertexDeltaModel.h
@@ -1,14 +1,21 @@
 // Copyright Epic Games, Inc. All Rights Reserved.
+
 #pragma once
 
 #include "CoreMinimal.h"
 #include "MLDeformerGeomCacheModel.h"
 #include "VertexDeltaModelVizSettings.h"
-#include "UObject/Object.h"
 #include "VertexDeltaModel.generated.h"
 
 VERTEXDELTAMODEL_API DECLARE_LOG_CATEGORY_EXTERN(LogVertexDeltaModel, Log, All);
 
+/** 
+ * The vertex delta model, which uses a GPU based neural network.
+ * This model acts more as an example of how to implement a model that only uses the GPU.
+ * It is not as efficient as the Neural Morph Model, even though that model runs on the CPU.
+ * Unlike the Neural Morph Model this Vertex Delta Model does not generate any morph targets. All the vertex deltas are directly output
+ * by the neural network. This buffer of output deltas remains on the GPU and is then used as input buffer directly inside an optimus deformer graph.
+ */
 UCLASS()
 class VERTEXDELTAMODEL_API UVertexDeltaModel 
 	: public UMLDeformerGeomCacheModel
@@ -19,16 +26,16 @@ public:
 	UVertexDeltaModel(const FObjectInitializer& ObjectInitializer);
 
 	// UMLDeformerModel overrides.
-	virtual FString GetDisplayName() const override { return "Vertex Delta Model"; }
-	virtual bool IsNeuralNetworkOnGPU() const override { return true; }	// GPU neural network.
+	virtual FString GetDisplayName() const override			{ return "Vertex Delta Model"; }
+	virtual bool IsNeuralNetworkOnGPU() const override		{ return true; }	// GPU neural network.
 	// ~END UMLDeformerModel overrides.
 
 #if WITH_EDITORONLY_DATA
-	int32 GetNumHiddenLayers() const { return NumHiddenLayers; }
-	int32 GetNumNeuronsPerLayer() const { return NumNeuronsPerLayer; }
-	int32 GetNumIterations() const { return NumIterations; }
-	int32 GetBatchSize() const { return BatchSize; }
-	float GetLearningRate() const { return LearningRate; }
+	int32 GetNumHiddenLayers() const						{ return NumHiddenLayers; }
+	int32 GetNumNeuronsPerLayer() const						{ return NumNeuronsPerLayer; }
+	int32 GetNumIterations() const							{ return NumIterations; }
+	int32 GetBatchSize() const								{ return BatchSize; }
+	float GetLearningRate() const							{ return LearningRate; }
 
 public:
 	/** The number of hidden layers that the neural network model will have.\nHigher numbers will slow down performance but can deal with more complex deformations. */
diff --git a/Engine/Plugins/Experimental/MLDeformer/VertexDeltaModel/Source/VertexDeltaModel/Public/VertexDeltaModelVizSettings.h b/Engine/Plugins/Experimental/MLDeformer/VertexDeltaModel/Source/VertexDeltaModel/Public/VertexDeltaModelVizSettings.h
--- a/Engine/Plugins/Experimental/MLDeformer/VertexDeltaModel/Source/VertexDeltaModel/Public/VertexDeltaModelVizSettings.h
+++ b/Engine/Plugins/Experimental/MLDeformer/VertexDeltaModel/Source/VertexDeltaModel/Public/VertexDeltaModelVizSettings.h
@@ -9,6 +9,9 @@
 
 /**
  * The vizualization settings specific to the the vertex delta model.
+ * This is inherited from the UMLDeformerGeomCacheVizSettings.
+ * We cannot directly use that as we have to register the detail customization using a unique class for our model.
+ * This is the reason why we just inherit a new class from it.
  */
 UCLASS()
 class VERTEXDELTAMODEL_API UVertexDeltaModelVizSettings 
diff --git a/Engine/Plugins/Experimental/MLDeformer/VertexDeltaModel/Source/VertexDeltaModelEditor/Public/VertexDeltaEditorModel.h b/Engine/Plugins/Experimental/MLDeformer/VertexDeltaModel/Source/VertexDeltaModelEditor/Public/VertexDeltaEditorModel.h
--- a/Engine/Plugins/Experimental/MLDeformer/VertexDeltaModel/Source/VertexDeltaModelEditor/Public/VertexDeltaEditorModel.h
+++ b/Engine/Plugins/Experimental/MLDeformer/VertexDeltaModel/Source/VertexDeltaModelEditor/Public/VertexDeltaEditorModel.h
@@ -9,6 +9,9 @@ namespace UE::VertexDeltaModel
 {
 	using namespace UE::MLDeformer;
 
+	/** 
+	 * The editor model for a UVertexDeltaModel.
+	 */
 	class VERTEXDELTAMODELEDITOR_API FVertexDeltaEditorModel 
 		: public UE::MLDeformer::FMLDeformerGeomCacheEditorModel
 	{
@@ -17,7 +20,7 @@ namespace UE::VertexDeltaModel
 		static FMLDeformerEditorModel* MakeInstance();
 
 		// FGCObject overrides.
-		virtual FString GetReferencerName() const override { return TEXT("FVertexDeltaEditorModel"); }
+		virtual FString GetReferencerName() const override	{ return TEXT("FVertexDeltaEditorModel"); }
 		// ~END FGCObject overrides.
 	
 		// FMLDeformerEditorModel overrides.
diff --git a/Engine/Plugins/Experimental/MLDeformer/VertexDeltaModel/Source/VertexDeltaModelEditor/Public/VertexDeltaModelDetails.h b/Engine/Plugins/Experimental/MLDeformer/VertexDeltaModel/Source/VertexDeltaModelEditor/Public/VertexDeltaModelDetails.h
--- a/Engine/Plugins/Experimental/MLDeformer/VertexDeltaModel/Source/VertexDeltaModelEditor/Public/VertexDeltaModelDetails.h
+++ b/Engine/Plugins/Experimental/MLDeformer/VertexDeltaModel/Source/VertexDeltaModelEditor/Public/VertexDeltaModelDetails.h
@@ -7,6 +7,12 @@
 
 namespace UE::VertexDeltaModel
 {
+	/**
+	 * The detail customization for the vertex delta model.
+	 * This basically just adds the properties specific to the vertex delta model to the trainings settings.
+	 * But as we inherit from the FMLDeformerGeomCacheModelDetails class, it directly creates all the other properties for us 
+	 * as well, and inserts error messages when something is wrong, creates some groups, etc.
+	 */
 	class VERTEXDELTAMODELEDITOR_API FVertexDeltaModelDetails
 		: public UE::MLDeformer::FMLDeformerGeomCacheModelDetails
 	{
diff --git a/Engine/Plugins/Experimental/MLDeformer/VertexDeltaModel/Source/VertexDeltaModelEditor/Public/VertexDeltaModelVizSettingsDetails.h b/Engine/Plugins/Experimental/MLDeformer/VertexDeltaModel/Source/VertexDeltaModelEditor/Public/VertexDeltaModelVizSettingsDetails.h
--- a/Engine/Plugins/Experimental/MLDeformer/VertexDeltaModel/Source/VertexDeltaModelEditor/Public/VertexDeltaModelVizSettingsDetails.h
+++ b/Engine/Plugins/Experimental/MLDeformer/VertexDeltaModel/Source/VertexDeltaModelEditor/Public/VertexDeltaModelVizSettingsDetails.h
@@ -7,6 +7,11 @@
 
 namespace UE::VertexDeltaModel
 {
+	/**
+	 * The detail customization for the vertex delta model visualization settings.
+	 * We inherit from a base class as that already takes care of some nice grouping of properties, some error reporting
+	 * in case of issues, and some nice grouping of settings.
+	 */
 	class VERTEXDELTAMODELEDITOR_API FVertexDeltaModelVizSettingsDetails
 		: public UE::MLDeformer::FMLDeformerGeomCacheVizSettingsDetails
 	{
diff --git a/Engine/Plugins/Experimental/MLDeformer/VertexDeltaModel/Source/VertexDeltaModelEditor/Public/VertexDeltaTrainingModel.h b/Engine/Plugins/Experimental/MLDeformer/VertexDeltaModel/Source/VertexDeltaModelEditor/Public/VertexDeltaTrainingModel.h
--- a/Engine/Plugins/Experimental/MLDeformer/VertexDeltaModel/Source/VertexDeltaModelEditor/Public/VertexDeltaTrainingModel.h
+++ b/Engine/Plugins/Experimental/MLDeformer/VertexDeltaModel/Source/VertexDeltaModelEditor/Public/VertexDeltaTrainingModel.h
@@ -6,7 +6,10 @@
 #include "MLDeformerTrainingModel.h"
 #include "VertexDeltaTrainingModel.generated.h"
 
-
+/**
+ * The training model for the vertex delta model.
+ * This class is our link to the Python training.
+ */
 UCLASS(Blueprintable)
 class VERTEXDELTAMODELEDITOR_API UVertexDeltaTrainingModel
 	: public UMLDeformerTrainingModel
@@ -14,7 +17,11 @@ class VERTEXDELTAMODELEDITOR_API UVertexDeltaTrainingModel
 	GENERATED_BODY()
 
 public:
-	/** Main training function, with implementation in python. */
+	/**
+	 * Main training function, with implementation in python.
+	 * You need to implement this method. See the UMLDeformerTrainingModel documentation for more details.
+	 * @see UMLDeformerTrainingModel
+	 */
 	UFUNCTION(BlueprintImplementableEvent, Category = "Python")
 	int32 Train() const;
 };
