[MLDeformer] Moved out of Experimental, and marked as beta.

diff --git a/Engine/Plugins/Animation/MLDeformer/MLDeformerFramework/Shaders/Private/MLDeformerGraphDataInterface.ush b/Engine/Plugins/Animation/MLDeformer/MLDeformerFramework/Shaders/Private/MLDeformerGraphDataInterface.ush
new file mode 100644
--- /dev/null
+++ b/Engine/Plugins/Animation/MLDeformer/MLDeformerFramework/Shaders/Private/MLDeformerGraphDataInterface.ush
@@ -0,0 +1,22 @@
+// Copyright Epic Games, Inc. All Rights Reserved.
+
+#include "/Plugin/ComputeFramework/Private/ComputeKernelCommon.ush"
+
+uint DI_LOCAL(NumVertices);
+uint DI_LOCAL(InputStreamStart);
+float DI_LOCAL(Weight);
+
+Buffer<float> DI_LOCAL(PositionDeltaBuffer);
+Buffer<uint> DI_LOCAL(VertexMapBuffer);
+
+DI_IMPL_READ(ReadNumVertices, uint, )
+{
+	return DI_LOCAL(NumVertices);
+}
+
+DI_IMPL_READ(ReadPositionDelta, float3, uint VertexIndex)
+{
+	uint BufferIndex = DI_LOCAL(VertexMapBuffer)[VertexIndex + DI_LOCAL(InputStreamStart)] * 3;
+	float3 Delta = float3(DI_LOCAL(PositionDeltaBuffer)[BufferIndex], DI_LOCAL(PositionDeltaBuffer)[BufferIndex + 1], DI_LOCAL(PositionDeltaBuffer)[BufferIndex + 2]);
+	return Delta * DI_LOCAL(Weight);
+}
diff --git a/Engine/Plugins/Animation/MLDeformer/MLDeformerFramework/Shaders/Private/MLDeformerGraphHeatMapDataInterface.ush b/Engine/Plugins/Animation/MLDeformer/MLDeformerFramework/Shaders/Private/MLDeformerGraphHeatMapDataInterface.ush
new file mode 100644
--- /dev/null
+++ b/Engine/Plugins/Animation/MLDeformer/MLDeformerFramework/Shaders/Private/MLDeformerGraphHeatMapDataInterface.ush
@@ -0,0 +1,43 @@
+// Copyright Epic Games, Inc. All Rights Reserved.
+
+#include "/Plugin/ComputeFramework/Private/ComputeKernelCommon.ush"
+
+uint DI_LOCAL(NumVertices);
+uint DI_LOCAL(InputStreamStart);
+uint DI_LOCAL(HeatMapMode);
+float DI_LOCAL(HeatMapMax);
+float DI_LOCAL(GroundTruthLerp);
+uint DI_LOCAL(GroundTruthBufferSize);
+
+StructuredBuffer<float3> DI_LOCAL(PositionGroundTruthBuffer);
+Buffer<uint> DI_LOCAL(VertexMapBuffer);
+
+DI_IMPL_READ(ReadNumVertices, uint, )
+{
+	return DI_LOCAL(NumVertices);
+}
+
+DI_IMPL_READ(ReadPositionGroundTruth, float3, uint VertexIndex)
+{
+	uint BufferIndex = DI_LOCAL(VertexMapBuffer)[VertexIndex + DI_LOCAL(InputStreamStart)];
+	if (BufferIndex >= DI_LOCAL(GroundTruthBufferSize))
+	{
+		return float3(0, 0, 0);
+	}
+	return DI_LOCAL(PositionGroundTruthBuffer)[BufferIndex];
+}
+
+DI_IMPL_READ(ReadHeatMapMode, uint, )
+{
+	return DI_LOCAL(HeatMapMode);
+}
+
+DI_IMPL_READ(ReadHeatMapMax, float, )
+{
+	return DI_LOCAL(HeatMapMax);
+}
+
+DI_IMPL_READ(ReadGroundTruthLerp, float, )
+{
+	return DI_LOCAL(GroundTruthLerp);
+}
diff --git a/Engine/Plugins/Animation/MLDeformer/MLDeformerFramework/Source/MLDeformerFramework/Private/MLDeformerGraphDataInterface.cpp b/Engine/Plugins/Animation/MLDeformer/MLDeformerFramework/Source/MLDeformerFramework/Private/MLDeformerGraphDataInterface.cpp
new file mode 100644
--- /dev/null
+++ b/Engine/Plugins/Animation/MLDeformer/MLDeformerFramework/Source/MLDeformerFramework/Private/MLDeformerGraphDataInterface.cpp
@@ -0,0 +1,109 @@
+// Copyright Epic Games, Inc. All Rights Reserved.
+
+#include "MLDeformerGraphDataInterface.h"
+
+#include "Components/SkeletalMeshComponent.h"
+#include "ComputeFramework/ShaderParamTypeDefinition.h"
+#include "MLDeformerModel.h"
+#include "MLDeformerModelInstance.h"
+#include "MLDeformerAsset.h"
+#include "MLDeformerComponent.h"
+#include "NeuralNetwork.h"
+#include "OptimusDataDomain.h"
+#include "RenderGraphBuilder.h"
+#include "RenderGraphUtils.h"
+#include "RenderGraphResources.h"
+#include "ShaderParameterMetadataBuilder.h"
+#include "SkeletalRenderPublic.h"
+
+TArray<FOptimusCDIPinDefinition> UMLDeformerGraphDataInterface::GetPinDefinitions() const
+{
+	TArray<FOptimusCDIPinDefinition> Defs;
+	Defs.Add({ "PositionDelta", "ReadPositionDelta", Optimus::DomainName::Vertex, "ReadNumVertices" });
+	return Defs;
+}
+
+TSubclassOf<UActorComponent> UMLDeformerGraphDataInterface::GetRequiredComponentClass() const
+{
+	return UMLDeformerComponent::StaticClass();
+}
+
+void UMLDeformerGraphDataInterface::GetSupportedInputs(TArray<FShaderFunctionDefinition>& OutFunctions) const
+{
+	OutFunctions.AddDefaulted_GetRef()
+		.SetName(TEXT("ReadNumVertices"))
+		.AddReturnType(EShaderFundamentalType::Uint);
+
+	OutFunctions.AddDefaulted_GetRef()
+		.SetName(TEXT("ReadPositionDelta"))
+		.AddReturnType(EShaderFundamentalType::Float, 3)
+		.AddParam(EShaderFundamentalType::Uint);
+}
+
+BEGIN_SHADER_PARAMETER_STRUCT(FMLDeformerGraphDataInterfaceParameters, )
+	MLDEFORMER_SHADER_PARAMETERS()
+END_SHADER_PARAMETER_STRUCT()
+
+MLDEFORMER_GRAPH_IMPLEMENT_BASICS(
+	UMLDeformerGraphDataInterface,
+	UMLDeformerGraphDataProvider,
+	UE::MLDeformer::FMLDeformerGraphDataProviderProxy,
+	FMLDeformerGraphDataInterfaceParameters,
+	TEXT("#include \"/Plugin/MLDeformerFramework/Private/MLDeformerGraphDataInterface.ush\"\n"),
+	TEXT("ML Deformer"))
+
+
+bool UMLDeformerGraphDataProvider::IsValid() const
+{
+	if (DeformerComponent == nullptr || DeformerComponent->GetDeformerAsset() == nullptr)
+	{
+		return false;
+	}
+
+	return DeformerComponent->GetModelInstance()->IsValidForDataProvider();
+}
+
+namespace UE::MLDeformer
+{
+	FMLDeformerGraphDataProviderProxy::FMLDeformerGraphDataProviderProxy(UMLDeformerComponent* DeformerComponent)
+	{
+		using namespace UE::MLDeformer;
+
+		const UMLDeformerAsset* DeformerAsset = DeformerComponent->GetDeformerAsset();
+		if (DeformerAsset)
+		{
+			const UMLDeformerModel* Model = DeformerAsset->GetModel();
+			const UMLDeformerModelInstance* ModelInstance = DeformerComponent->GetModelInstance();
+		
+			SkeletalMeshObject = ModelInstance->GetSkeletalMeshComponent()->MeshObject;
+			NeuralNetwork = Model->GetNeuralNetwork();
+			NeuralNetworkInferenceHandle = ModelInstance->GetNeuralNetworkInferenceHandle();
+			bCanRunNeuralNet = ModelInstance->IsCompatible();
+			Weight = DeformerComponent->GetWeight();
+			VertexMapBufferSRV = Model->GetVertexMapBuffer().ShaderResourceViewRHI;
+		}
+	}
+
+	void FMLDeformerGraphDataProviderProxy::AllocateResources(FRDGBuilder& GraphBuilder)
+	{
+		if (bCanRunNeuralNet)
+		{
+			check(NeuralNetwork != nullptr);
+			Buffer = GraphBuilder.RegisterExternalBuffer(NeuralNetwork->GetOutputTensorForContext(NeuralNetworkInferenceHandle).GetPooledBuffer());
+		}
+		else
+		{
+			// TODO: use an actual buffer that is of the right size, and filled with zero's.
+			Buffer = GraphBuilder.RegisterExternalBuffer(GWhiteVertexBufferWithRDG->Buffer);
+		}
+
+		BufferSRV = GraphBuilder.CreateSRV(Buffer, PF_R32_FLOAT);
+	}
+
+	void FMLDeformerGraphDataProviderProxy::GatherDispatchData(FDispatchSetup const& InDispatchSetup, FCollectedDispatchData& InOutDispatchData)
+	{
+		MLDEFORMER_GRAPH_DISPATCH_START(FMLDeformerGraphDataInterfaceParameters, InDispatchSetup, InOutDispatchData)
+		MLDEFORMER_GRAPH_DISPATCH_DEFAULT_PARAMETERS()
+		MLDEFORMER_GRAPH_DISPATCH_END()
+	}
+}	// namespace UE::MLDeformer
diff --git a/Engine/Plugins/Animation/MLDeformer/MLDeformerFramework/Source/MLDeformerFramework/Private/MLDeformerGraphDebugDataInterface.cpp b/Engine/Plugins/Animation/MLDeformer/MLDeformerFramework/Source/MLDeformerFramework/Private/MLDeformerGraphDebugDataInterface.cpp
new file mode 100644
--- /dev/null
+++ b/Engine/Plugins/Animation/MLDeformer/MLDeformerFramework/Source/MLDeformerFramework/Private/MLDeformerGraphDebugDataInterface.cpp
@@ -0,0 +1,145 @@
+// Copyright Epic Games, Inc. All Rights Reserved.
+
+#include "MLDeformerGraphDebugDataInterface.h"
+#include "MLDeformerGraphDataInterface.h"
+#include "MLDeformerModelInstance.h"
+#include "MLDeformerAsset.h"
+#include "MLDeformerComponent.h"
+#include "MLDeformerVizSettings.h"
+#include "Animation/AnimSequence.h"
+#include "Components/SkeletalMeshComponent.h"
+#include "ComputeFramework/ComputeKernelPermutationSet.h"
+#include "ComputeFramework/ShaderParamTypeDefinition.h"
+#include "GeometryCache.h"
+#include "GeometryCacheMeshData.h"
+#include "GeometryCacheTrack.h"
+#include "OptimusDataDomain.h"
+#include "RenderGraphBuilder.h"
+#include "RenderGraphUtils.h"
+#include "RenderGraphResources.h"
+#include "Rendering/SkeletalMeshModel.h"
+#include "ShaderParameterMetadataBuilder.h"
+#include "SkeletalRenderPublic.h"
+#include "NeuralNetwork.h"
+
+TArray<FOptimusCDIPinDefinition> UMLDeformerGraphDebugDataInterface::GetPinDefinitions() const
+{
+	TArray<FOptimusCDIPinDefinition> Defs;
+	Defs.Add({ "HeatMapMode", "ReadHeatMapMode" });
+	Defs.Add({ "HeatMapMax", "ReadHeatMapMax" });
+	Defs.Add({ "GroundTruthLerp", "ReadGroundTruthLerp" });
+	Defs.Add({ "PositionGroundTruth", "ReadPositionGroundTruth", Optimus::DomainName::Vertex, "ReadNumVertices" });
+	return Defs;
+}
+
+TSubclassOf<UActorComponent> UMLDeformerGraphDebugDataInterface::GetRequiredComponentClass() const
+{
+	return UMLDeformerComponent::StaticClass();
+}
+
+void UMLDeformerGraphDebugDataInterface::GetSupportedInputs(TArray<FShaderFunctionDefinition>& OutFunctions) const
+{
+	OutFunctions.AddDefaulted_GetRef()
+		.SetName(TEXT("ReadNumVertices"))
+		.AddReturnType(EShaderFundamentalType::Uint);
+
+	OutFunctions.AddDefaulted_GetRef()
+		.SetName(TEXT("ReadHeatMapMode"))
+		.AddReturnType(EShaderFundamentalType::Uint);
+
+	OutFunctions.AddDefaulted_GetRef()
+		.SetName(TEXT("ReadHeatMapMax"))
+		.AddReturnType(EShaderFundamentalType::Float);
+
+	OutFunctions.AddDefaulted_GetRef()
+		.SetName(TEXT("ReadGroundTruthLerp"))
+		.AddReturnType(EShaderFundamentalType::Float);
+
+	OutFunctions.AddDefaulted_GetRef()
+		.SetName(TEXT("ReadPositionGroundTruth"))
+		.AddReturnType(EShaderFundamentalType::Float, 3)
+		.AddParam(EShaderFundamentalType::Uint);
+}
+
+BEGIN_SHADER_PARAMETER_STRUCT(FMLDeformerGraphDebugDataInterfaceParameters, )
+	MLDEFORMER_DEBUG_SHADER_PARAMETERS()
+END_SHADER_PARAMETER_STRUCT()
+
+MLDEFORMER_GRAPH_IMPLEMENT_DEBUG_BASICS_WITH_PROXY(
+	UMLDeformerGraphDebugDataInterface,
+	UMLDeformerGraphDebugDataProvider,
+	UE::MLDeformer::FMLDeformerGraphDebugDataProviderProxy,
+	FMLDeformerGraphDebugDataInterfaceParameters,
+	TEXT("#include \"/Plugin/MLDeformerFramework/Private/MLDeformerGraphHeatMapDataInterface.ush\"\n"),
+	TEXT("ML Deformer Debug"))
+
+bool UMLDeformerGraphDebugDataProvider::IsValid() const
+{
+#if WITH_EDITORONLY_DATA
+	if (DeformerComponent == nullptr || DeformerComponent->GetDeformerAsset() == nullptr)
+	{
+		return false;
+	}
+
+	return DeformerComponent->GetModelInstance()->IsValidForDataProvider();
+#else
+	return false; // This data interface is only valid in editor.
+#endif
+}
+
+#if WITH_EDITORONLY_DATA
+namespace UE::MLDeformer
+{
+	FMLDeformerGraphDebugDataProviderProxy::FMLDeformerGraphDebugDataProviderProxy(UMLDeformerComponent* DeformerComponent, UMLDeformerAsset* DeformerAsset, UMLDeformerGraphDebugDataProvider* InProvider)
+		: FComputeDataProviderRenderProxy()
+	{
+		Provider = InProvider;
+
+		if (DeformerAsset)
+		{
+			UMLDeformerModel* Model = DeformerAsset->GetModel();	
+			UMLDeformerVizSettings* VizSettings = Model->GetVizSettings();
+			const UMLDeformerModelInstance* ModelInstance = DeformerComponent->GetModelInstance();
+
+			SkeletalMeshObject = ModelInstance->GetSkeletalMeshComponent()->MeshObject;
+			VertexMapBufferSRV = Model->GetVertexMapBuffer().ShaderResourceViewRHI;
+			HeatMapMode = (int32)VizSettings->GetHeatMapMode();
+			HeatMapMax = 1.0f / FMath::Max(VizSettings->GetHeatMapMax(), 0.00001f);
+			GroundTruthLerp = VizSettings->GetGroundTruthLerp();
+		}
+	}
+
+	void FMLDeformerGraphDebugDataProviderProxy::HandleZeroGroundTruthPositions()
+	{
+		if (GroundTruthPositions.Num() == 0)
+		{	
+			// We didn't get valid ground truth vertices.
+			// Make non empty array for later buffer generation.
+			GroundTruthPositions.Add(FVector3f::ZeroVector);
+
+			// Silently disable relevant debug things.
+			if (HeatMapMode == (int32)EMLDeformerHeatMapMode::GroundTruth)
+			{
+				HeatMapMode = -1;
+				HeatMapMax = 0.0f;
+				GroundTruthLerp = 0.0f;
+			}
+		}
+	}
+
+	void FMLDeformerGraphDebugDataProviderProxy::AllocateResources(FRDGBuilder& GraphBuilder)
+	{
+		GroundTruthBuffer = GraphBuilder.CreateBuffer(FRDGBufferDesc::CreateStructuredDesc(sizeof(FVector3f), GroundTruthPositions.Num()), TEXT("MLDeformer.GroundTruthPositions"));
+		GroundTruthBufferSRV = GraphBuilder.CreateSRV(GroundTruthBuffer);
+		GraphBuilder.QueueBufferUpload(GroundTruthBuffer, GroundTruthPositions.GetData(), sizeof(FVector3f) * GroundTruthPositions.Num(), ERDGInitialDataFlags::None);
+	}
+
+	void FMLDeformerGraphDebugDataProviderProxy::GatherDispatchData(FDispatchSetup const& InDispatchSetup, FCollectedDispatchData& InOutDispatchData)
+	{
+		MLDEFORMER_GRAPH_DISPATCH_START(FMLDeformerGraphDebugDataInterfaceParameters, InDispatchSetup, InOutDispatchData)
+		MLDEFORMER_GRAPH_DISPATCH_DEFAULT_DEBUG_PARAMETERS()
+		MLDEFORMER_GRAPH_DISPATCH_END()
+	}
+}	// namespace UE::MLDeformer
+#endif // WITH_EDITORONLY_DATA
+
diff --git a/Engine/Plugins/Animation/MLDeformer/MLDeformerFramework/Source/MLDeformerFramework/Private/MLDeformerModel.cpp b/Engine/Plugins/Animation/MLDeformer/MLDeformerFramework/Source/MLDeformerFramework/Private/MLDeformerModel.cpp
new file mode 100644
--- /dev/null
+++ b/Engine/Plugins/Animation/MLDeformer/MLDeformerFramework/Source/MLDeformerFramework/Private/MLDeformerModel.cpp
@@ -0,0 +1,227 @@
+// Copyright Epic Games, Inc. All Rights Reserved.
+
+#include "MLDeformerModel.h"
+#include "MLDeformerInputInfo.h"
+#include "MLDeformerModelInstance.h"
+#include "MLDeformerVizSettings.h"
+#include "MLDeformerAsset.h"
+#include "MLDeformerComponent.h"
+#include "Engine/SkeletalMesh.h"
+#include "Rendering/SkeletalMeshModel.h"
+#include "Rendering/SkeletalMeshLODModel.h"
+#include "NeuralNetwork.h"
+#include "Animation/AnimSequence.h"
+#include "UObject/UObjectGlobals.h"
+
+namespace UE::MLDeformer
+{
+	void FVertexMapBuffer::InitRHI()
+	{
+		if (VertexMap.Num() > 0)
+		{
+			FRHIResourceCreateInfo CreateInfo(TEXT("UMLDeformerModel::FVertexMapBuffer"));
+
+			VertexBufferRHI = RHICreateVertexBuffer(VertexMap.Num() * sizeof(uint32), BUF_Static | BUF_ShaderResource, CreateInfo);
+			uint32* Data = reinterpret_cast<uint32*>(RHILockBuffer(VertexBufferRHI, 0, VertexMap.Num() * sizeof(uint32), RLM_WriteOnly));
+			for (int32 Index = 0; Index < VertexMap.Num(); ++Index)
+			{
+				Data[Index] = static_cast<uint32>(VertexMap[Index]);
+			}
+			RHIUnlockBuffer(VertexBufferRHI);
+			VertexMap.Empty();
+
+			ShaderResourceViewRHI = RHICreateShaderResourceView(VertexBufferRHI, 4, PF_R32_UINT);
+		}
+		else
+		{
+			VertexBufferRHI = nullptr;
+			ShaderResourceViewRHI = nullptr;
+		}
+	}
+}	// namespace UE::MLDeformer
+
+UMLDeformerInputInfo* UMLDeformerModel::CreateInputInfo()
+{ 
+	return NewObject<UMLDeformerInputInfo>(this);
+}
+
+UMLDeformerModelInstance* UMLDeformerModel::CreateModelInstance(UMLDeformerComponent* Component)
+{ 
+	return NewObject<UMLDeformerModelInstance>(Component);
+}
+
+void UMLDeformerModel::Init(UMLDeformerAsset* InDeformerAsset) 
+{ 
+	check(InDeformerAsset);
+	DeformerAsset = InDeformerAsset; 
+	if (InputInfo == nullptr)
+	{
+		InputInfo = CreateInputInfo();
+	}
+}
+
+void UMLDeformerModel::Serialize(FArchive& Archive)
+{
+	TRACE_CPUPROFILER_EVENT_SCOPE(UMLDeformerModel::Serialize)
+	#if WITH_EDITOR
+		if (Archive.IsSaving() && Archive.IsPersistent())
+		{
+			InitVertexMap();
+			UpdateCachedNumVertices();
+		}
+	#endif
+
+	Super::Serialize(Archive);
+}
+
+UMLDeformerAsset* UMLDeformerModel::GetDeformerAsset() const
+{ 
+	return DeformerAsset.Get(); 
+}
+
+void UMLDeformerModel::PostLoad()
+{
+	TRACE_CPUPROFILER_EVENT_SCOPE(UMLDeformerModel::PostLoad)
+
+	Super::PostLoad();
+
+	InitGPUData();
+
+	#if WITH_EDITOR
+		UpdateCachedNumVertices();
+	#endif
+
+	UMLDeformerAsset* MLDeformerAsset = Cast<UMLDeformerAsset>(GetOuter());
+	Init(MLDeformerAsset);
+
+	if (InputInfo)
+	{
+		InputInfo->OnPostLoad();
+	}
+
+	if (NeuralNetwork)
+	{
+		// If we run the neural network on the GPU.
+		if (IsNeuralNetworkOnGPU())
+		{
+			NeuralNetwork->SetDeviceType(ENeuralDeviceType::GPU, ENeuralDeviceType::CPU, ENeuralDeviceType::GPU);
+			if (NeuralNetwork->GetDeviceType() != ENeuralDeviceType::GPU || NeuralNetwork->GetOutputDeviceType() != ENeuralDeviceType::GPU || NeuralNetwork->GetInputDeviceType() != ENeuralDeviceType::CPU)
+			{
+				UE_LOG(LogMLDeformer, Error, TEXT("Neural net in ML Deformer '%s' cannot run on the GPU, it will not be active."), *GetDeformerAsset()->GetName());
+			}
+		}
+		else // We run our neural network on the CPU.
+		{
+			NeuralNetwork->SetDeviceType(ENeuralDeviceType::CPU, ENeuralDeviceType::CPU, ENeuralDeviceType::CPU);
+		}
+	}
+
+#if WITH_EDITOR
+	SetAssetEditorOnlyFlags();
+#endif
+}
+
+void UMLDeformerModel::SetNeuralNetwork(UNeuralNetwork* InNeuralNetwork)
+{
+	NeuralNetworkModifyDelegate.Broadcast();
+	NeuralNetwork = InNeuralNetwork;
+}
+
+// Used for the FBoenReference, so it knows what skeleton to pick bones from.
+USkeleton* UMLDeformerModel::GetSkeleton(bool& bInvalidSkeletonIsError, const IPropertyHandle* PropertyHandle)
+{
+	#if WITH_EDITOR
+		bInvalidSkeletonIsError = false;
+		if (SkeletalMesh)
+		{
+			return SkeletalMesh->GetSkeleton();
+		}
+	#endif
+	return nullptr;
+}
+
+void UMLDeformerModel::BeginDestroy()
+{
+	PostEditPropertyDelegate.RemoveAll(this);
+
+	BeginReleaseResource(&VertexMapBuffer);
+	RenderResourceDestroyFence.BeginFence();
+	Super::BeginDestroy();
+}
+
+bool UMLDeformerModel::IsReadyForFinishDestroy()
+{
+	return Super::IsReadyForFinishDestroy() && RenderResourceDestroyFence.IsFenceComplete();
+}
+
+void UMLDeformerModel::InitGPUData()
+{
+	BeginReleaseResource(&VertexMapBuffer);
+	VertexMapBuffer.Init(VertexMap);
+	BeginInitResource(&VertexMapBuffer);
+}
+
+void UMLDeformerModel::FloatArrayToVector3Array(const TArray<float>& FloatArray, TArray<FVector3f>& OutVectorArray)
+{
+	check(FloatArray.Num() % 3 == 0);
+	const int32 NumVerts = FloatArray.Num() / 3;
+	OutVectorArray.Reset();
+	OutVectorArray.SetNumUninitialized(NumVerts);
+	for (int32 VertexIndex = 0; VertexIndex < NumVerts; ++VertexIndex)
+	{
+		const int32 FloatBufferOffset = VertexIndex * 3;
+		OutVectorArray[VertexIndex] = FVector3f(FloatArray[FloatBufferOffset + 0], FloatArray[FloatBufferOffset + 1], FloatArray[FloatBufferOffset + 2]);
+	}
+}
+
+#if WITH_EDITOR
+	void UMLDeformerModel::UpdateNumBaseMeshVertices()
+	{
+		NumBaseMeshVerts = UMLDeformerModel::ExtractNumImportedSkinnedVertices(SkeletalMesh);
+	}
+
+	void UMLDeformerModel::PostEditChangeProperty(FPropertyChangedEvent& PropertyChangedEvent)
+	{
+		const FProperty* Property = PropertyChangedEvent.Property;
+		if (Property == nullptr)
+		{
+			return;
+		}
+
+		OnPostEditChangeProperty().Broadcast(PropertyChangedEvent);
+	}
+
+	void UMLDeformerModel::UpdateCachedNumVertices()
+	{
+		UpdateNumBaseMeshVertices();
+		UpdateNumTargetMeshVertices();
+	}
+
+	int32 UMLDeformerModel::ExtractNumImportedSkinnedVertices(const USkeletalMesh* SkeletalMesh)
+	{
+		return SkeletalMesh ? SkeletalMesh->GetNumImportedVertices() : 0;
+	}
+
+	void UMLDeformerModel::SetAssetEditorOnlyFlags()
+	{
+		if (AnimSequence)
+		{
+			AnimSequence->GetPackage()->SetPackageFlags(PKG_EditorOnly);
+		}
+	}
+#endif	// #if WITH_EDITOR
+
+#if WITH_EDITORONLY_DATA
+	void UMLDeformerModel::InitVertexMap()
+	{
+		VertexMap.Empty();
+		if (SkeletalMesh)
+		{
+			FSkeletalMeshModel* SkeletalMeshModel = SkeletalMesh->GetImportedModel();
+			if (SkeletalMeshModel)
+			{
+				VertexMap = SkeletalMeshModel->LODModels[0].MeshToImportVertexMap;
+			}
+		}
+	}
+#endif	// WITH_EDITORDATA_ONLY
diff --git a/Engine/Plugins/Animation/MLDeformer/MLDeformerFramework/Source/MLDeformerFramework/Private/Tests/MLDeformerTestModel.cpp b/Engine/Plugins/Animation/MLDeformer/MLDeformerFramework/Source/MLDeformerFramework/Private/Tests/MLDeformerTestModel.cpp
new file mode 100644
--- /dev/null
+++ b/Engine/Plugins/Animation/MLDeformer/MLDeformerFramework/Source/MLDeformerFramework/Private/Tests/MLDeformerTestModel.cpp
@@ -0,0 +1,22 @@
+// Copyright Epic Games, Inc. All Rights Reserved.
+
+#include "MLDeformerTestModel.h"
+#include "MLDeformerComponent.h"
+#include "UObject/UObjectGlobals.h"
+
+UMLDeformerModelInstance* UMLDeformerTestModel::CreateModelInstance(UMLDeformerComponent* Component)
+{
+	return NewObject<UTestModelInstance>(Component);
+}
+
+FString UMLDeformerTestModel::GetDisplayName() const 
+{ 
+	return FString("Test Model");
+}
+
+#if WITH_EDITOR
+	void UMLDeformerTestModel::UpdateNumTargetMeshVertices()
+	{
+		NumTargetMeshVerts = NumBaseMeshVerts;
+	}
+#endif
diff --git a/Engine/Plugins/Animation/MLDeformer/MLDeformerFramework/Source/MLDeformerFramework/Private/Tests/MLDeformerTestModel.h b/Engine/Plugins/Animation/MLDeformer/MLDeformerFramework/Source/MLDeformerFramework/Private/Tests/MLDeformerTestModel.h
new file mode 100644
--- /dev/null
+++ b/Engine/Plugins/Animation/MLDeformer/MLDeformerFramework/Source/MLDeformerFramework/Private/Tests/MLDeformerTestModel.h
@@ -0,0 +1,34 @@
+// Copyright Epic Games, Inc. All Rights Reserved.
+
+#pragma once
+
+#include "CoreMinimal.h"
+#include "MLDeformerModel.h"
+#include "MLDeformerModelInstance.h"
+#include "MLDeformerTestModel.generated.h"
+
+class UMLDeformerComponent;
+
+UCLASS()
+class UMLDeformerTestModel 
+	: public UMLDeformerModel
+{
+	GENERATED_BODY()
+
+public:
+	UMLDeformerModelInstance* CreateModelInstance(UMLDeformerComponent* Component) override;
+	FString GetDisplayName() const override;
+#if WITH_EDITOR
+	void UpdateNumTargetMeshVertices() override;
+#endif
+};
+
+UCLASS()
+class UTestModelInstance
+	: public UMLDeformerModelInstance
+{
+	GENERATED_BODY()
+
+public:
+	void Tick(float DeltaTime, float ModelWeight) override { }
+};
diff --git a/Engine/Plugins/Animation/MLDeformer/MLDeformerFramework/Source/MLDeformerFramework/Private/Tests/ModelTests.cpp b/Engine/Plugins/Animation/MLDeformer/MLDeformerFramework/Source/MLDeformerFramework/Private/Tests/ModelTests.cpp
new file mode 100644
--- /dev/null
+++ b/Engine/Plugins/Animation/MLDeformer/MLDeformerFramework/Source/MLDeformerFramework/Private/Tests/ModelTests.cpp
@@ -0,0 +1,97 @@
+// Copyright Epic Games, Inc. All Rights Reserved.
+
+#include "CoreMinimal.h"
+#include "Misc/AutomationTest.h"
+
+#if WITH_DEV_AUTOMATION_TESTS
+
+#include "MLDeformerTestModel.h"
+#include "MLDeformerAsset.h"
+#include "MLDeformerComponent.h"
+#include "MLDeformerInputInfo.h"
+#include "UObject/GCObjectScopeGuard.h"
+#include "PreviewScene.h"
+#include "Components/SkeletalMeshComponent.h"
+
+namespace UE::MLDeformerTests
+{
+	IMPLEMENT_SIMPLE_AUTOMATION_TEST(FModelMainTest, "MLDeformer.Model.MainTest", EAutomationTestFlags::ApplicationContextMask | EAutomationTestFlags::EngineFilter);
+	bool FModelMainTest::RunTest(const FString& Parameters)
+	{
+		using namespace UE::MLDeformer;
+
+#if WITH_EDITORONLY_DATA
+		UMLDeformerAsset* DeformerAsset = NewObject<UMLDeformerAsset>();
+		FGCObjectScopeGuard AssetScopeGuard(DeformerAsset);
+		UMLDeformerTestModel* TestModel = NewObject<UMLDeformerTestModel>(DeformerAsset);
+		FGCObjectScopeGuard ModelScopeGuard(TestModel);
+		UTEST_NOT_NULL(TEXT("DeformerAsset check"), DeformerAsset);
+		UTEST_NOT_NULL(TEXT("Creating test model"), TestModel);
+		UTEST_EQUAL(TEXT("Display name check"), TestModel->GetDisplayName(), FString("Test Model"));
+
+		USkeletalMesh* SkeletalMesh = LoadObject<USkeletalMesh>(nullptr, TEXT("SkeletalMesh'/MLDeformerFramework/Tests/SKM_Cylinder.SKM_Cylinder'"));
+		UTEST_NOT_NULL(TEXT("Loading skeletal mesh"), SkeletalMesh);
+
+		TestModel->SetSkeletalMesh(SkeletalMesh);
+		UTEST_EQUAL(TEXT("GetSkeletalMesh check"), TestModel->GetSkeletalMesh(), SkeletalMesh);
+
+		DeformerAsset->SetModel(TestModel);
+		UTEST_TRUE(TEXT("GetModel check"), DeformerAsset->GetModel() == TestModel);
+
+		TestModel->Init(DeformerAsset);
+		UTEST_TRUE(TEXT("GetDeformerAsset check"), TestModel->GetDeformerAsset() == DeformerAsset);
+
+		const int32 ExpectedNumImportedVerts = 1178;
+		const int32 ImportedVerts = UMLDeformerModel::ExtractNumImportedSkinnedVertices(SkeletalMesh);
+		UTEST_EQUAL(TEXT("Imported verts check"), ImportedVerts, ExpectedNumImportedVerts);
+
+		TestModel->UpdateCachedNumVertices();
+		UTEST_EQUAL(TEXT("Imported verts check"), TestModel->GetNumBaseMeshVerts(), ImportedVerts);
+
+		const USkeleton* Skeleton = SkeletalMesh->GetSkeleton();
+		bool bErrorOnInvalid = false;
+		const USkeleton* ModelSkeleton = TestModel->GetSkeleton(bErrorOnInvalid, nullptr);
+		UTEST_EQUAL(TEXT("Skeleton check"), ModelSkeleton, Skeleton);
+
+		// Create a preview scene and spawn an actor.
+		TUniquePtr<FPreviewScene> PreviewScene = MakeUnique<FPreviewScene>(FPreviewScene::ConstructionValues());
+		UWorld* World = PreviewScene->GetWorld();
+		UTEST_NOT_NULL(TEXT("World check"), World);
+		FActorSpawnParameters SpawnParams;
+		SpawnParams.Name = MakeUniqueObjectName(World, AActor::StaticClass(), FName(TEXT("ML Deformer Test Actor")));	
+		AActor* Actor = World->SpawnActor<AActor>(SpawnParams);
+		Actor->SetFlags(RF_Transient);
+		UTEST_NOT_NULL(TEXT("Actor check"), Actor);
+
+		// Add a skeletal mesh component.
+		USkeletalMeshComponent* SkelMeshComponent = NewObject<USkeletalMeshComponent>(Actor);
+		SkelMeshComponent->SetSkeletalMesh(SkeletalMesh);
+		Actor->SetRootComponent(SkelMeshComponent);
+		SkelMeshComponent->RegisterComponent();
+		
+		// Add an ML Deformer component and activate it.
+		UMLDeformerComponent* MLDeformerComponent = NewObject<UMLDeformerComponent>(Actor);
+		MLDeformerComponent->SetDeformerAsset(DeformerAsset);
+		MLDeformerComponent->RegisterComponent();
+		MLDeformerComponent->Activate();
+		UTEST_EQUAL(TEXT("MLDeformerComponent weight check"), MLDeformerComponent->GetWeight(), 1.0f);
+		UTEST_TRUE(TEXT("MLDeformerComponent SkelMeshComponent check"), MLDeformerComponent->GetSkeletalMeshComponent() == SkelMeshComponent);
+
+		// Create and init the model instance.
+		UMLDeformerModelInstance* ModelInstance = MLDeformerComponent->GetModelInstance();
+		UTEST_NOT_NULL(TEXT("CreateModelInstance check"), ModelInstance);
+		UTEST_TRUE(TEXT("ModelInstance GetModel check"), ModelInstance->GetModel() == TestModel);
+		UTEST_TRUE(TEXT("ModelInstance compatible check"), ModelInstance->IsCompatible());
+
+		// Do some input info checks.
+		UMLDeformerInputInfo* InputInfo = TestModel->GetInputInfo();
+		UTEST_NOT_NULL(TEXT("Get InputInfo"), InputInfo);
+		UTEST_EQUAL(TEXT("InputInfo NumBaseMeshVerts check"), InputInfo->GetNumBaseMeshVertices(), 0);
+		UTEST_EQUAL(TEXT("InputInfo NumBones check"), InputInfo->GetNumBones(), 0);
+		UTEST_EQUAL(TEXT("InputInfo NumCurves check"), InputInfo->GetNumCurves(), 0);
+#endif	// #if WITH_EDITORONLY_DATA
+		return true;
+	}
+} // namespace UE::MLDeformerTests
+
+#endif // WITH_DEV_AUTOMATION_TESTS
diff --git a/Engine/Plugins/Animation/MLDeformer/MLDeformerFramework/Source/MLDeformerFramework/Public/MLDeformerGraphDataInterface.h b/Engine/Plugins/Animation/MLDeformer/MLDeformerFramework/Source/MLDeformerFramework/Public/MLDeformerGraphDataInterface.h
new file mode 100644
--- /dev/null
+++ b/Engine/Plugins/Animation/MLDeformer/MLDeformerFramework/Source/MLDeformerFramework/Public/MLDeformerGraphDataInterface.h
@@ -0,0 +1,146 @@
+// Copyright Epic Games, Inc. All Rights Reserved.
+
+#pragma once
+
+#include "CoreMinimal.h"
+#include "OptimusComputeDataInterface.h"
+#include "ComputeFramework/ComputeDataProvider.h"
+#include "Rendering/SkeletalMeshRenderData.h"
+#include "Rendering/SkeletalMeshLODRenderData.h"
+#include "SkeletalRenderPublic.h"
+#include "MLDeformerGraphDataInterface.generated.h"
+
+class FRDGBuffer;
+class FRDGBufferSRV;
+class FRHIShaderResourceView;
+class FSkeletalMeshObject;
+class UMLDeformerComponent;
+class UNeuralNetwork;
+class USkeletalMeshComponent;
+class UMLDeformerModel;
+
+/** Defines the default base ML Deformer shader parameters. */
+#define MLDEFORMER_SHADER_PARAMETERS() \
+	SHADER_PARAMETER(uint32, NumVertices) \
+	SHADER_PARAMETER(uint32, InputStreamStart) \
+	SHADER_PARAMETER(float, Weight) \
+	SHADER_PARAMETER_RDG_BUFFER_SRV(Buffer<float>, PositionDeltaBuffer) \
+	SHADER_PARAMETER_SRV(Buffer<uint>, VertexMapBuffer)
+
+#define MLDEFORMER_GRAPH_DISPATCH_START(ParameterStructType, InDispatchSetup, InOutDispatchData) \
+	if (!ensure(InDispatchSetup.ParameterStructSizeForValidation == sizeof(ParameterStructType))) \
+	{ \
+		return; \
+	} \
+	const FSkeletalMeshRenderData& SkeletalMeshRenderData = SkeletalMeshObject->GetSkeletalMeshRenderData(); \
+	const FSkeletalMeshLODRenderData* LodRenderData = SkeletalMeshRenderData.GetPendingFirstLOD(0); \
+	for (int32 InvocationIndex = 0; InvocationIndex < InDispatchSetup.NumInvocations; ++InvocationIndex) \
+	{ \
+		const FSkelMeshRenderSection& RenderSection = LodRenderData->RenderSections[InvocationIndex]; \
+		ParameterStructType* Parameters = (ParameterStructType*)(InOutDispatchData.ParameterBuffer + InDispatchSetup.ParameterBufferOffset + InDispatchSetup.ParameterBufferStride * InvocationIndex);
+
+#define MLDEFORMER_GRAPH_DISPATCH_DEFAULT_PARAMETERS() \
+	Parameters->NumVertices = 0; \
+	Parameters->InputStreamStart = RenderSection.BaseVertexIndex; \
+	Parameters->Weight = Weight; \
+	Parameters->PositionDeltaBuffer = BufferSRV; \
+	Parameters->VertexMapBuffer = VertexMapBufferSRV;
+
+#define MLDEFORMER_GRAPH_DISPATCH_END() }
+
+#define MLDEFORMER_GRAPH_IMPLEMENT_BASICS(InterfaceClassName, DataProviderClassName, DataProviderProxyClassName, ParameterStructType, HLSLText, DisplayName) \
+	FString InterfaceClassName::GetDisplayName() const \
+	{ \
+		return DisplayName; \
+	} \
+	void InterfaceClassName::GetShaderParameters(TCHAR const* UID, FShaderParametersMetadataBuilder& InOutBuilder, FShaderParametersMetadataAllocations& InOutAllocations) const \
+	{ \
+		InOutBuilder.AddNestedStruct<ParameterStructType>(UID); \
+	} \
+	void InterfaceClassName::GetHLSL(FString& OutHLSL) const \
+	{ \
+		OutHLSL += HLSLText; \
+	} \
+	UComputeDataProvider* InterfaceClassName::CreateDataProvider(TObjectPtr<UObject> InBinding, uint64 InInputMask, uint64 InOutputMask) const \
+	{ \
+		DataProviderClassName* Provider = NewObject<DataProviderClassName>(); \
+		Provider->DeformerComponent = Cast<UMLDeformerComponent>(InBinding); \
+		return Provider; \
+	} \
+	FComputeDataProviderRenderProxy* DataProviderClassName::GetRenderProxy() \
+	{ \
+		return new DataProviderProxyClassName(DeformerComponent); \
+	}
+
+#if WITH_EDITORONLY_DATA
+	#define MLDEFORMER_EDITORDATA_ONLY(Statement, ElseStatement) Statement
+#else
+	#define MLDEFORMER_EDITORDATA_ONLY(X, ElseStatement) ElseStatement
+#endif
+
+/** Compute Framework Data Interface for MLDefomer data. */
+UCLASS(Category = ComputeFramework)
+class MLDEFORMERFRAMEWORK_API UMLDeformerGraphDataInterface
+	: public UOptimusComputeDataInterface
+{
+	GENERATED_BODY()
+
+public:
+	// UOptimusComputeDataInterface overrides.
+	virtual FString GetDisplayName() const override;
+	virtual TArray<FOptimusCDIPinDefinition> GetPinDefinitions() const override;
+	virtual TSubclassOf<UActorComponent> GetRequiredComponentClass() const override;
+	// ~END UOptimusComputeDataInterface overrides.
+
+	// UComputeDataInterface overrides.
+	TCHAR const* GetClassName() const override { return TEXT("MLDeformer"); }
+	virtual void GetSupportedInputs(TArray<FShaderFunctionDefinition>& OutFunctions) const override;
+	virtual void GetShaderParameters(TCHAR const* UID, FShaderParametersMetadataBuilder& InOutBuilder, FShaderParametersMetadataAllocations& InOutAllocations) const override;
+	virtual void GetHLSL(FString& OutHLSL) const override;
+	virtual UComputeDataProvider* CreateDataProvider(TObjectPtr<UObject> InBinding, uint64 InInputMask, uint64 InOutputMask) const override;
+	// ~END UComputeDataInterface overrides.
+};
+
+/** Compute Framework Data Provider for MLDeformer data. */
+UCLASS(BlueprintType, EditInlineNew, Category = ComputeFramework)
+class MLDEFORMERFRAMEWORK_API UMLDeformerGraphDataProvider
+	: public UComputeDataProvider
+{
+	GENERATED_BODY()
+
+public:
+	/** The deformer component that this data provider works on. */
+	UPROPERTY(EditAnywhere, BlueprintReadWrite, Category = Binding)
+	TObjectPtr<UMLDeformerComponent> DeformerComponent = nullptr;
+
+	// UComputeDataProvider overrides.
+	virtual bool IsValid() const override;
+	virtual FComputeDataProviderRenderProxy* GetRenderProxy() override;
+	// ~END UComputeDataProvider overrides.
+};
+
+namespace UE::MLDeformer
+{
+	/** Compute Framework Data Provider Proxy for MLDeformer data. */
+	class MLDEFORMERFRAMEWORK_API FMLDeformerGraphDataProviderProxy
+		: public FComputeDataProviderRenderProxy
+	{
+	public:
+		FMLDeformerGraphDataProviderProxy(UMLDeformerComponent* DeformerComponent);
+
+		// FComputeDataProviderRenderProxy overrides.
+		virtual void AllocateResources(FRDGBuilder& GraphBuilder) override;
+		virtual void GatherDispatchData(FDispatchSetup const& InDispatchSetup, FCollectedDispatchData& InOutDispatchData) override;
+		// ~END FComputeDataProviderRenderProxy overrides.
+
+	protected:
+		FSkeletalMeshObject* SkeletalMeshObject = nullptr;
+		TObjectPtr<UNeuralNetwork> NeuralNetwork = nullptr;
+		FRHIShaderResourceView* VertexMapBufferSRV = nullptr;
+		FRDGBuffer* Buffer = nullptr;
+		FRDGBufferSRV* BufferSRV = nullptr;
+		int32 NeuralNetworkInferenceHandle = -1;
+		float Weight = 1.0f;
+		bool bCanRunNeuralNet = false;
+	};
+}	// namespace UE::MLDeformer
diff --git a/Engine/Plugins/Animation/MLDeformer/MLDeformerFramework/Source/MLDeformerFramework/Public/MLDeformerGraphDebugDataInterface.h b/Engine/Plugins/Animation/MLDeformer/MLDeformerFramework/Source/MLDeformerFramework/Public/MLDeformerGraphDebugDataInterface.h
new file mode 100644
--- /dev/null
+++ b/Engine/Plugins/Animation/MLDeformer/MLDeformerFramework/Source/MLDeformerFramework/Public/MLDeformerGraphDebugDataInterface.h
@@ -0,0 +1,169 @@
+// Copyright Epic Games, Inc. All Rights Reserved.
+
+#pragma once
+
+#include "CoreMinimal.h"
+#include "MLDeformerModel.h"
+#include "MLDeformerGraphDataInterface.h"
+#include "OptimusComputeDataInterface.h"
+#include "ComputeFramework/ComputeDataProvider.h"
+#include "MLDeformerGraphDebugDataInterface.generated.h"
+
+class FRDGBuffer;
+class FRDGBufferSRV;
+class FRHIShaderResourceView;
+class FSkeletalMeshObject;
+class UMLDeformerAsset;
+class USkeletalMeshComponent;
+
+#define MLDEFORMER_DEBUG_SHADER_PARAMETERS() \
+	SHADER_PARAMETER(uint32, NumVertices) \
+	SHADER_PARAMETER(uint32, InputStreamStart) \
+	SHADER_PARAMETER(int32, HeatMapMode) \
+	SHADER_PARAMETER(float, HeatMapMax) \
+	SHADER_PARAMETER(float, GroundTruthLerp) \
+	SHADER_PARAMETER(uint32, GroundTruthBufferSize) \
+	SHADER_PARAMETER_RDG_BUFFER_SRV(StructuredBuffer<float3>, PositionGroundTruthBuffer) \
+	SHADER_PARAMETER_SRV(Buffer<uint>, VertexMapBuffer)
+
+#define MLDEFORMER_GRAPH_DISPATCH_DEFAULT_DEBUG_PARAMETERS() \
+	Parameters->NumVertices = 0; \
+	Parameters->InputStreamStart = RenderSection.BaseVertexIndex; \
+	Parameters->HeatMapMode = HeatMapMode; \
+	Parameters->HeatMapMax = HeatMapMax; \
+	Parameters->GroundTruthLerp = GroundTruthLerp; \
+	Parameters->GroundTruthBufferSize = GroundTruthPositions.Num(); \
+	Parameters->PositionGroundTruthBuffer = GroundTruthBufferSRV; \
+	Parameters->VertexMapBuffer = VertexMapBufferSRV;
+
+#define MLDEFORMER_GRAPH_IMPLEMENT_DEBUG_BASICS(InterfaceClassName, DataProviderClassName, ParameterStructType, HLSLText, DisplayName) \
+	FString InterfaceClassName::GetDisplayName() const \
+	{ \
+		return DisplayName; \
+	} \
+	void InterfaceClassName::GetShaderParameters(TCHAR const* UID, FShaderParametersMetadataBuilder& InOutBuilder, FShaderParametersMetadataAllocations& InOutAllocations) const \
+	{ \
+		InOutBuilder.AddNestedStruct<ParameterStructType>(UID); \
+	} \
+	void InterfaceClassName::GetHLSL(FString& OutHLSL) const \
+	{ \
+		OutHLSL += HLSLText; \
+	} \
+	UComputeDataProvider* InterfaceClassName::CreateDataProvider(TObjectPtr<UObject> InBinding, uint64 InInputMask, uint64 InOutputMask) const \
+	{ \
+		DataProviderClassName* Provider = NewObject<DataProviderClassName>(); \
+		Provider->DeformerComponent = Cast<UMLDeformerComponent>(InBinding); \
+		if (Provider->DeformerComponent) \
+		{ \
+			Provider->DeformerAsset = Provider->DeformerComponent->GetDeformerAsset(); \
+		} \
+		MLDEFORMER_EDITORDATA_ONLY( \
+			if (Provider->DeformerAsset != nullptr) \
+			{ \
+				Provider->Init(); \
+			} \
+		,) \
+		return Provider; \
+	}
+
+#define MLDEFORMER_GRAPH_IMPLEMENT_DEBUG_BASICS_WITH_PROXY(InterfaceClassName, DataProviderClassName, DataProviderProxyClassName, ParameterStructType, HLSLText, DisplayName) \
+	MLDEFORMER_GRAPH_IMPLEMENT_DEBUG_BASICS(InterfaceClassName, DataProviderClassName, ParameterStructType, HLSLText, DisplayName) \
+	FComputeDataProviderRenderProxy* DataProviderClassName::GetRenderProxy() \
+	{ \
+		MLDEFORMER_EDITORDATA_ONLY( \
+			DataProviderProxyClassName* Proxy = new DataProviderProxyClassName(DeformerComponent, DeformerAsset, this); \
+			if (DeformerAsset) \
+			{ \
+				const float SampleTime = DeformerComponent->GetModelInstance()->GetSkeletalMeshComponent()->GetPosition(); \
+				UMLDeformerModel* Model = DeformerAsset->GetModel(); \
+				Model->SampleGroundTruthPositions(SampleTime, Proxy->GetGroundTruthPositions()); \
+				Proxy->HandleZeroGroundTruthPositions(); \
+				return Proxy; \
+			} \
+			return nullptr; \
+			, \
+			return nullptr; \
+		) \
+	}
+
+/** 
+ * Compute Framework Data Interface for MLDeformer debugging data. 
+ * This interfaces to editor only data, and so will only give valid results in that context.
+ */
+UCLASS(Category = ComputeFramework)
+class MLDEFORMERFRAMEWORK_API UMLDeformerGraphDebugDataInterface
+	: public UOptimusComputeDataInterface
+{
+	GENERATED_BODY()
+
+public:
+	// UOptimusComputeDataInterface overrides.
+	virtual FString GetDisplayName() const override;
+	virtual TArray<FOptimusCDIPinDefinition> GetPinDefinitions() const override;
+	virtual TSubclassOf<UActorComponent> GetRequiredComponentClass() const override;
+	// ~END UOptimusComputeDataInterface overrides.
+
+	// UComputeDataInterface overrides.
+	TCHAR const* GetClassName() const override { return TEXT("MLDeformerDebug"); }
+	virtual void GetSupportedInputs(TArray<FShaderFunctionDefinition>& OutFunctions) const override;
+	virtual void GetShaderParameters(TCHAR const* UID, FShaderParametersMetadataBuilder& OutBuilder, FShaderParametersMetadataAllocations& InOutAllocations) const override;
+	virtual void GetHLSL(FString& OutHLSL) const override;
+	virtual UComputeDataProvider* CreateDataProvider(TObjectPtr<UObject> InBinding, uint64 InInputMask, uint64 InOutputMask) const override;
+	// ~END UComputeDataInterface overrides.
+};
+
+/** Compute Framework Data Provider for MLDeformer debugging data. */
+UCLASS(BlueprintType, EditInlineNew, Category = ComputeFramework)
+class MLDEFORMERFRAMEWORK_API UMLDeformerGraphDebugDataProvider
+	: public UComputeDataProvider
+{
+	GENERATED_BODY()
+
+public:
+	virtual void Init() {};
+
+	// UComputeDataProvider overrides.
+	virtual bool IsValid() const override;
+	virtual FComputeDataProviderRenderProxy* GetRenderProxy() override;
+	// ~END UComputeDataProvider overrides.
+
+public:
+	UPROPERTY(EditAnywhere, BlueprintReadWrite, Category = Binding)
+	TObjectPtr<UMLDeformerComponent> DeformerComponent = nullptr;
+
+	UPROPERTY(EditAnywhere, BlueprintReadWrite, Category = Binding)
+	TObjectPtr<UMLDeformerAsset> DeformerAsset = nullptr;
+};
+
+/** Compute Framework Data Provider Proxy for MLDeformer debugging data. */
+#if WITH_EDITORONLY_DATA
+namespace UE::MLDeformer
+{
+	class MLDEFORMERFRAMEWORK_API FMLDeformerGraphDebugDataProviderProxy
+		: public FComputeDataProviderRenderProxy
+	{
+	public:
+		FMLDeformerGraphDebugDataProviderProxy(UMLDeformerComponent* DeformerComponent, UMLDeformerAsset* DeformerAsset, UMLDeformerGraphDebugDataProvider* InProvider);
+
+		virtual void HandleZeroGroundTruthPositions();
+
+		// FComputeDataProviderRenderProxy overrides.
+		virtual void AllocateResources(FRDGBuilder& GraphBuilder) override;
+		virtual void GatherDispatchData(FDispatchSetup const& InDispatchSetup, FCollectedDispatchData& InOutDispatchData) override;
+		// ~END FComputeDataProviderRenderProxy overrides.
+
+		TArray<FVector3f>& GetGroundTruthPositions() { return GroundTruthPositions; }
+
+	protected:
+		TObjectPtr<UMLDeformerGraphDebugDataProvider> Provider = nullptr;
+		FSkeletalMeshObject* SkeletalMeshObject;
+		TArray<FVector3f> GroundTruthPositions;
+		FRHIShaderResourceView* VertexMapBufferSRV = nullptr;
+		FRDGBuffer* GroundTruthBuffer = nullptr;
+		FRDGBufferSRV* GroundTruthBufferSRV = nullptr;
+		int32 HeatMapMode = 0;
+		float HeatMapMax = 0.0f;
+		float GroundTruthLerp = 0.0f;
+	};
+}	// namespace UE::MLDeformer
+#endif
