[MLDeformer] Moved out of Experimental, and marked as beta.

diff --git a/Engine/Plugins/Animation/MLDeformer/MLDeformerFramework/Source/MLDeformerFrameworkEditor/Private/MLDeformerEditorModel.cpp b/Engine/Plugins/Animation/MLDeformer/MLDeformerFramework/Source/MLDeformerFrameworkEditor/Private/MLDeformerEditorModel.cpp
new file mode 100644
--- /dev/null
+++ b/Engine/Plugins/Animation/MLDeformer/MLDeformerFramework/Source/MLDeformerFrameworkEditor/Private/MLDeformerEditorModel.cpp
@@ -0,0 +1,1971 @@
+// Copyright Epic Games, Inc. All Rights Reserved.
+
+#include "MLDeformerEditorModel.h"
+#include "MLDeformerModule.h"
+#include "MLDeformerModel.h"
+#include "MLDeformerEditorStyle.h"
+#include "MLDeformerEditorModule.h"
+#include "MLDeformerModelRegistry.h"
+#include "MLDeformerEditorToolkit.h"
+#include "MLDeformerEditorActor.h"
+#include "MLDeformerAsset.h"
+#include "MLDeformerComponent.h"
+#include "MLDeformerInputInfo.h"
+#include "MLDeformerSampler.h"
+#include "MLDeformerModelInstance.h"
+#include "AnimationEditorPreviewActor.h"
+#include "AnimPreviewInstance.h"
+#include "Animation/MeshDeformer.h"
+#include "Animation/DebugSkelMeshComponent.h"
+#include "Animation/AnimSequence.h"
+#include "Rendering/SkeletalMeshModel.h"
+#include "Rendering/SkeletalMeshLODModel.h"
+#include "IPersonaPreviewScene.h"
+#include "IPersonaToolkit.h"
+#include "Modules/ModuleManager.h"
+#include "Engine/World.h"
+#include "GameFramework/Actor.h"
+#include "EditorViewportClient.h"
+#include "Components/TextRenderComponent.h"
+#include "Materials/Material.h"
+#include "GeometryCache.h"
+#include "GeometryCacheMeshData.h"
+#include "GeometryCacheTrack.h"
+#include "NeuralNetwork.h"
+
+#define LOCTEXT_NAMESPACE "MLDeformerEditorModel"
+
+namespace UE::MLDeformer
+{
+	FMLDeformerEditorModel::~FMLDeformerEditorModel()
+	{
+		if (PostEditPropertyDelegateHandle.IsValid())
+		{
+			Model->OnPostEditChangeProperty().Remove(PostEditPropertyDelegateHandle);
+		}
+
+		ClearWorld();
+
+		FMLDeformerEditorModule& EditorModule = FModuleManager::GetModuleChecked<FMLDeformerEditorModule>("MLDeformerFrameworkEditor");
+		EditorModule.GetModelRegistry().RemoveEditorModelInstance(this);
+	}
+
+	void FMLDeformerEditorModel::Init(const InitSettings& Settings)
+	{
+		check(Settings.Editor);
+		check(Settings.Model);
+
+		Editor = Settings.Editor;
+		Model = Settings.Model;
+
+		EditorInputInfo = Model->CreateInputInfo();
+		check(EditorInputInfo);
+
+		Sampler = CreateSampler();
+		check(Sampler);
+		Sampler->Init(this);
+
+		ViewRange = TRange<double>(0.0, 100.0);
+		WorkingRange = TRange<double>(0.0, 100.0);
+		PlaybackRange = TRange<double>(0.0, 100.0);
+
+		PostEditPropertyDelegateHandle = Model->OnPostEditChangeProperty().AddRaw(this, &FMLDeformerEditorModel::OnPostEditChangeProperty);
+	}
+
+	void FMLDeformerEditorModel::AddReferencedObjects(FReferenceCollector& Collector)
+	{
+		Collector.AddReferencedObject(EditorInputInfo);
+	}
+
+	void FMLDeformerEditorModel::UpdateEditorInputInfo()
+	{
+		InitInputInfo(EditorInputInfo);
+	}
+
+	UWorld* FMLDeformerEditorModel::GetWorld() const
+	{
+		check(Editor);
+		return Editor->GetPersonaToolkit()->GetPreviewScene()->GetWorld();
+	}
+
+	FMLDeformerSampler* FMLDeformerEditorModel::CreateSampler() const
+	{
+		return new FMLDeformerSampler();
+	}
+
+	void FMLDeformerEditorModel::CreateTrainingLinearSkinnedActor(const TSharedRef<IPersonaPreviewScene>& InPersonaPreviewScene)
+	{
+		UWorld* World = InPersonaPreviewScene->GetWorld();
+
+		// Spawn the linear skinned actor.
+		FActorSpawnParameters BaseSpawnParams;
+		BaseSpawnParams.Name = MakeUniqueObjectName(World, AActor::StaticClass(), "Train Base Actor");
+		AActor* Actor = World->SpawnActor<AActor>(AActor::StaticClass(), BaseSpawnParams);
+		Actor->SetFlags(RF_Transient);
+
+		// Create the preview skeletal mesh component.
+		const FLinearColor BaseWireColor = FMLDeformerEditorStyle::Get().GetColor("MLDeformer.BaseMesh.WireframeColor");
+		UDebugSkelMeshComponent* SkelMeshComponent = NewObject<UDebugSkelMeshComponent>(Actor);
+		SkelMeshComponent->SetWireframeMeshOverlayColor(BaseWireColor);
+		SkelMeshComponent->SetVisibility(false);
+		SkelMeshComponent->MarkRenderStateDirty();
+
+		// Set the skeletal mesh on the component.
+		// NOTE: This must be done AFTER setting the AnimInstance so that the correct root anim node is loaded.
+		USkeletalMesh* Mesh = Model->GetSkeletalMesh();
+		SkelMeshComponent->SetSkeletalMesh(Mesh);
+
+		// Update the persona scene.
+		InPersonaPreviewScene->SetActor(Actor);
+		InPersonaPreviewScene->SetPreviewMeshComponent(SkelMeshComponent);
+		InPersonaPreviewScene->AddComponent(SkelMeshComponent, FTransform::Identity);
+		InPersonaPreviewScene->SetAdditionalMeshesSelectable(false);
+		InPersonaPreviewScene->SetPreviewMesh(Mesh);
+
+		// Register the editor actor.
+		const FLinearColor LabelColor = FMLDeformerEditorStyle::Get().GetColor("MLDeformer.BaseMesh.LabelColor");
+		FMLDeformerEditorActor::FConstructSettings Settings;
+		Settings.Actor = Actor;
+		Settings.TypeID = ActorID_Train_Base;
+		Settings.LabelColor = LabelColor;
+		Settings.LabelText = LOCTEXT("TrainBaseActorLabelText", "Training Base");
+		Settings.bIsTrainingActor = true;
+		FMLDeformerEditorActor* EditorActor = CreateEditorActor(Settings);
+		EditorActor->SetSkeletalMeshComponent(SkelMeshComponent);
+		EditorActor->SetMeshOffsetFactor(0.0f);
+		EditorActors.Add(EditorActor);
+	}
+
+	void FMLDeformerEditorModel::CreateTestLinearSkinnedActor(UWorld* World)
+	{
+		FActorSpawnParameters SpawnParams;
+		SpawnParams.Name = MakeUniqueObjectName(World, AActor::StaticClass(), "Test Linear Skinned Actor");
+		AActor* Actor = World->SpawnActor<AActor>(AActor::StaticClass(), FTransform::Identity, SpawnParams);
+		Actor->SetFlags(RF_Transient);
+
+		const FLinearColor BaseWireColor = FMLDeformerEditorStyle::Get().GetColor("MLDeformer.BaseMesh.WireframeColor");
+		UDebugSkelMeshComponent* SkelMeshComponent = NewObject<UDebugSkelMeshComponent>(Actor);
+		SkelMeshComponent->SetWireframeMeshOverlayColor(BaseWireColor);
+		SkelMeshComponent->SetSkeletalMesh(Model->GetSkeletalMesh());
+		Actor->SetRootComponent(SkelMeshComponent);
+		SkelMeshComponent->RegisterComponent();
+		SkelMeshComponent->SetVisibility(false);
+		SkelMeshComponent->MarkRenderStateDirty();
+
+		// Register the editor actor.
+		const FLinearColor LabelColor = FMLDeformerEditorStyle::Get().GetColor("MLDeformer.BaseMesh.LabelColor");
+		FMLDeformerEditorActor::FConstructSettings Settings;
+		Settings.Actor = Actor;
+		Settings.TypeID = ActorID_Test_Base;
+		Settings.LabelColor = LabelColor;
+		Settings.LabelText = LOCTEXT("TestBaseActorLabelText", "Linear Skinned");
+		Settings.bIsTrainingActor = false;
+		FMLDeformerEditorActor* EditorActor = CreateEditorActor(Settings);
+		EditorActor->SetSkeletalMeshComponent(SkelMeshComponent);
+		EditorActor->SetMeshOffsetFactor(0.0f);
+		EditorActors.Add(EditorActor);
+	}
+
+	void FMLDeformerEditorModel::CreateTestMLDeformedActor(UWorld* World)
+	{
+		// Create the ML deformed actor.
+		FActorSpawnParameters SpawnParams;
+		SpawnParams.Name = MakeUniqueObjectName(World, AActor::StaticClass(), "Test ML Deformed");
+		AActor* Actor = World->SpawnActor<AActor>(SpawnParams);
+		Actor->SetFlags(RF_Transient);
+
+		// Create the skeletal mesh component.
+		const FLinearColor MLDeformedWireColor = FMLDeformerEditorStyle::Get().GetColor("MLDeformer.MLDeformedMesh.WireframeColor");
+		UDebugSkelMeshComponent* SkelMeshComponent = NewObject<UDebugSkelMeshComponent>(Actor);
+		SkelMeshComponent->SetSkeletalMesh(Model->GetSkeletalMesh());
+		Actor->SetRootComponent(SkelMeshComponent);
+		SkelMeshComponent->RegisterComponent();
+		SkelMeshComponent->SetWireframeMeshOverlayColor(MLDeformedWireColor);
+		SkelMeshComponent->SetVisibility(false);
+		SkelMeshComponent->MarkRenderStateDirty();
+
+		// Create the ML Deformer component.
+		UMLDeformerAsset* DeformerAsset = Model->GetDeformerAsset();
+		UMLDeformerComponent* MLDeformerComponent = NewObject<UMLDeformerComponent>(Actor);
+		MLDeformerComponent->SetDeformerAsset(DeformerAsset);
+		MLDeformerComponent->RegisterComponent();
+		MLDeformerComponent->SetupComponent(DeformerAsset, SkelMeshComponent);
+
+		// Create the editor actor.
+		const FLinearColor LabelColor = FMLDeformerEditorStyle::Get().GetColor("MLDeformer.MLDeformedMesh.LabelColor");
+		FMLDeformerEditorActor::FConstructSettings Settings;
+		Settings.Actor = Actor;
+		Settings.TypeID = ActorID_Test_MLDeformed;
+		Settings.LabelColor = LabelColor;
+		Settings.LabelText = LOCTEXT("TestMLDeformedActorLabelText", "ML Deformed");
+		Settings.bIsTrainingActor = false;
+		FMLDeformerEditorActor* EditorActor = static_cast<FMLDeformerEditorActor*>(CreateEditorActor(Settings));
+		EditorActor->SetSkeletalMeshComponent(SkelMeshComponent);
+		EditorActor->SetMLDeformerComponent(MLDeformerComponent);
+		EditorActor->SetMeshOffsetFactor(1.0f);
+		EditorActors.Add(EditorActor);
+	}
+
+	void FMLDeformerEditorModel::CreateActors(const TSharedRef<IPersonaPreviewScene>& InPersonaPreviewScene)
+	{
+		UWorld* World = InPersonaPreviewScene->GetWorld();
+		CreateTrainingLinearSkinnedActor(InPersonaPreviewScene);
+		CreateTestLinearSkinnedActor(World);
+		CreateTestMLDeformedActor(World);
+		CreateTrainingGroundTruthActor(World);
+		CreateTestGroundTruthActor(World);
+
+		// Set the default mesh translation offsets for our ground truth actors.
+		for (FMLDeformerEditorActor* EditorActor : EditorActors)
+		{
+			if (EditorActor && EditorActor->IsGroundTruthActor())
+			{			
+				// The mesh offset factor basically just offsets the actor position by a given factor.
+				// The amount the actor is moved from the origin is: (MeshSpacing * MeshOffsetFactor).
+				// In test mode we have 3 actors (Linear, ML Deformed, Ground Truth), so it's mesh offset factor will be 2.0 for the ground truth.
+				// It is 2.0 because the ground truth actor in testing mode is all the way on the right, next to the ML Deformed model.
+				// In training mode we have only the Linear Skinned actor and the ground truth, so there the spacing factor is 1.0.
+				// TLDR: Basically 1.0 means its the first actor next to the linear skinned actor while 2.0 means its the second character, etc.
+				EditorActor->SetMeshOffsetFactor(EditorActor->IsTestActor() ? 2.0f : 1.0f);
+			}
+		}
+	}
+
+	void FMLDeformerEditorModel::ClearWorld()
+	{
+		TSharedRef<IPersonaPreviewScene> PreviewScene = Editor->GetPersonaToolkit()->GetPreviewScene();
+
+		// First remove actors from the world and add them to the destroy list.
+		TArray<AActor*> ActorsToDestroy;
+		ActorsToDestroy.Reserve(EditorActors.Num());
+		UWorld* World = PreviewScene->GetWorld();
+		for (FMLDeformerEditorActor* EditorActor : EditorActors)
+		{
+			if (EditorActor)
+			{				
+				World->RemoveActor(EditorActor->GetActor(), true);
+				ActorsToDestroy.Add(EditorActor->GetActor());
+			}
+		}
+
+		// Now delete the editor actors, which removes components.
+		DeleteEditorActors();
+
+		PreviewScene->DeselectAll();
+		PreviewScene->SetPreviewAnimationAsset(nullptr);
+		PreviewScene->SetPreviewAnimationBlueprint(nullptr, nullptr);
+		PreviewScene->SetPreviewMesh(nullptr);
+		PreviewScene->SetPreviewMeshComponent(nullptr);
+		PreviewScene->SetActor(nullptr);
+
+		for (AActor* Actor : ActorsToDestroy)
+		{
+			Actor->Destroy();
+		}
+	}
+
+	FMLDeformerEditorActor* FMLDeformerEditorModel::CreateEditorActor(const FMLDeformerEditorActor::FConstructSettings& Settings) const
+	{ 
+		return new FMLDeformerEditorActor(Settings);
+	}
+
+	void FMLDeformerEditorModel::DeleteEditorActors()
+	{
+		for (FMLDeformerEditorActor* EditorActor : EditorActors)
+		{
+			delete EditorActor;
+		}
+		EditorActors.Empty();
+	}
+
+	FMLDeformerEditorActor* FMLDeformerEditorModel::FindEditorActor(int32 ActorTypeID) const
+	{
+		for (FMLDeformerEditorActor* EditorActor : EditorActors)
+		{
+			if (EditorActor && EditorActor->GetTypeID() == ActorTypeID)
+			{
+				return EditorActor;
+			}
+		}
+
+		return nullptr;
+	}
+
+	void FMLDeformerEditorModel::Tick(FEditorViewportClient* ViewportClient, float DeltaTime)
+	{
+		// Force the training sequence to use Step interpolation and sample raw animation data.
+		// We do this in the Tick, as this is reset to false in the engine sometimes.
+		UAnimSequence* TrainingAnimSequence = Model->GetAnimSequence();
+		if (TrainingAnimSequence)
+		{
+			TrainingAnimSequence->bUseRawDataOnly = true;
+			TrainingAnimSequence->Interpolation = EAnimInterpolationType::Step;
+		}
+
+		// Do the same for the test anim sequence.
+		UMLDeformerVizSettings* VizSettings = Model->GetVizSettings();
+		{
+			UAnimSequence* TestAnimSequence = VizSettings->GetTestAnimSequence();
+			if (TestAnimSequence)
+			{
+				TestAnimSequence->bUseRawDataOnly = true;
+
+				// Enable step interpolation when doing showing a heatmap vs ground truth.
+				if (VizSettings->HasTestGroundTruth() && VizSettings->GetShowHeatMap() && VizSettings->GetHeatMapMode() == EMLDeformerHeatMapMode::GroundTruth)
+				{
+					TestAnimSequence->Interpolation = EAnimInterpolationType::Step;
+				}
+				else
+				{
+					TestAnimSequence->Interpolation = EAnimInterpolationType::Linear;
+				}
+			}
+		}
+
+		UpdateActorTransforms();
+		UpdateLabels();
+		CheckTrainingDataFrameChanged();
+
+		// Update the ML Deformer component's weight.
+		FMLDeformerEditorActor* EditorActor = FindEditorActor(ActorID_Test_MLDeformed);
+		if (EditorActor)
+		{
+			UMLDeformerComponent* DeformerComponent = EditorActor->GetMLDeformerComponent();
+			if (DeformerComponent)
+			{		
+				DeformerComponent->SetWeight(Model->GetVizSettings()->GetWeight());
+			}
+		}
+	}
+
+	void FMLDeformerEditorModel::UpdateLabels()
+	{
+		const UMLDeformerVizSettings* VizSettings = Model->GetVizSettings();
+		const bool bDrawTrainingActors = (VizSettings->GetVisualizationMode() == EMLDeformerVizMode::TrainingData);
+		const bool bDrawTestActors = (VizSettings->GetVisualizationMode() == EMLDeformerVizMode::TestData);
+
+		for (FMLDeformerEditorActor* EditorActor : EditorActors)
+		{
+			if (EditorActor == nullptr)
+			{
+				continue;
+			}
+
+			UTextRenderComponent* LabelComponent = EditorActor->GetLabelComponent();
+			if (LabelComponent == nullptr)
+			{
+				continue;
+			}
+
+			if (VizSettings->GetDrawLabels())
+			{
+				const AActor* Actor = EditorActor->GetActor();
+				const FVector ActorLocation = Actor->GetActorLocation();
+				const FVector AlignmentOffset = EditorActor->IsGroundTruthActor() ? Model->GetAlignmentTransform().GetTranslation() : FVector::ZeroVector;
+
+				LabelComponent->SetRelativeLocation(ActorLocation + FVector(0.0f, 0.0f, VizSettings->GetLabelHeight()) - AlignmentOffset);
+				LabelComponent->SetRelativeRotation(FQuat(FVector(0.0f, 0.0f, 1.0f), FMath::DegreesToRadians(90.0f)));
+				LabelComponent->SetRelativeScale3D(FVector(VizSettings->GetLabelScale() * 0.3f));
+
+				// Update visibility.
+				const bool bLabelIsVisible = (bDrawTrainingActors && EditorActor->IsTrainingActor()) || (bDrawTestActors && EditorActor->IsTestActor());
+				LabelComponent->SetVisibility(bLabelIsVisible);
+
+				// Handle test ground truth, disable its label when no ground truth asset was selected.
+				if (EditorActor->GetTypeID() == ActorID_Test_GroundTruth && !VizSettings->HasTestGroundTruth())
+				{
+					LabelComponent->SetVisibility(false);
+				}
+			}
+			else
+			{
+				LabelComponent->SetVisibility(false);
+			}
+		}
+	}
+
+	void FMLDeformerEditorModel::UpdateActorTransforms()
+	{
+		const FVector MeshSpacingVector = Model->GetVizSettings()->GetMeshSpacingOffsetVector();
+		for (FMLDeformerEditorActor* EditorActor : EditorActors)
+		{
+			if (EditorActor)
+			{
+				FTransform Transform = EditorActor->IsGroundTruthActor() ? Model->GetAlignmentTransform() : FTransform::Identity;
+				Transform.AddToTranslation(MeshSpacingVector * EditorActor->GetMeshOffsetFactor());
+				EditorActor->GetActor()->SetActorTransform(Transform);
+			}
+		}
+	}
+
+	void FMLDeformerEditorModel::UpdateActorVisibility()
+	{
+		const UMLDeformerVizSettings* VizSettings = Model->GetVizSettings();
+		const bool bShowTrainingData = (VizSettings->GetVisualizationMode() == EMLDeformerVizMode::TrainingData);
+		const bool bShowTestData = (VizSettings->GetVisualizationMode() == EMLDeformerVizMode::TestData);
+		for (FMLDeformerEditorActor* EditorActor : EditorActors)
+		{
+			if (EditorActor)
+			{
+				bool bIsVisible = (EditorActor->IsTestActor() && bShowTestData) || (EditorActor->IsTrainingActor() && bShowTrainingData);
+				if (EditorActor->GetTypeID() == ActorID_Test_Base)
+				{
+					bIsVisible &= VizSettings->GetDrawLinearSkinnedActor();
+				}
+				else if (EditorActor->GetTypeID() == ActorID_Test_MLDeformed)
+				{
+					bIsVisible &= VizSettings->GetDrawMLDeformedActor();
+				}
+				else if (EditorActor->GetTypeID() == ActorID_Test_GroundTruth)
+				{
+					bIsVisible &= VizSettings->GetDrawGroundTruthActor();
+				}
+				EditorActor->SetVisibility(bIsVisible);
+			}
+		}
+	}
+
+	void FMLDeformerEditorModel::OnInputAssetsChanged()
+	{
+		// Force the training sequence to use Step interpolation and sample raw animation data.
+		UAnimSequence* TrainingAnimSequence = Model->GetAnimSequence();
+		if (TrainingAnimSequence)
+		{
+			TrainingAnimSequence->bUseRawDataOnly = true;
+			TrainingAnimSequence->Interpolation = EAnimInterpolationType::Step;
+		}
+
+		// Update the training base actor.
+		UDebugSkelMeshComponent* SkeletalMeshComponent = FindEditorActor(ActorID_Train_Base)->GetSkeletalMeshComponent();
+		UMLDeformerVizSettings* VizSettings = Model->GetVizSettings();
+		UAnimSequence* TestAnimSequence = VizSettings->GetTestAnimSequence();
+		const float TestAnimSpeed = VizSettings->GetAnimPlaySpeed();
+		{
+			check(SkeletalMeshComponent);
+			SkeletalMeshComponent->SetSkeletalMesh(Model->GetSkeletalMesh());
+			if (GetEditor()->GetPersonaToolkitPointer())
+			{
+				GetEditor()->GetPersonaToolkit()->GetPreviewScene()->SetPreviewMesh(Model->GetSkeletalMesh());
+			}
+			SkeletalMeshComponent->SetAnimationMode(EAnimationMode::AnimationSingleNode);
+			const float CurrentPlayTime = SkeletalMeshComponent->GetPosition();
+			SkeletalMeshComponent->SetAnimation(TrainingAnimSequence);
+			SkeletalMeshComponent->SetPosition(CurrentPlayTime);
+			SkeletalMeshComponent->SetPlayRate(TestAnimSpeed);
+			SkeletalMeshComponent->Play(true);
+		}
+
+		// Update the test base model.
+		SkeletalMeshComponent = FindEditorActor(ActorID_Test_Base)->GetSkeletalMeshComponent();
+		if (SkeletalMeshComponent)
+		{
+			SkeletalMeshComponent->SetSkeletalMesh(Model->GetSkeletalMesh());
+			SkeletalMeshComponent->SetAnimationMode(EAnimationMode::AnimationSingleNode);
+			const float CurrentPlayTime = SkeletalMeshComponent->GetPosition();
+			SkeletalMeshComponent->SetAnimation(TestAnimSequence);
+			SkeletalMeshComponent->SetPosition(CurrentPlayTime);
+			SkeletalMeshComponent->SetPlayRate(TestAnimSpeed);
+			SkeletalMeshComponent->Play(true);
+		}
+
+		// Update the test ML Deformed skeletal mesh component.
+		SkeletalMeshComponent = FindEditorActor(ActorID_Test_MLDeformed)->GetSkeletalMeshComponent();
+		if (SkeletalMeshComponent)
+		{
+			SkeletalMeshComponent->SetSkeletalMesh(Model->GetSkeletalMesh());
+			SkeletalMeshComponent->SetAnimationMode(EAnimationMode::AnimationSingleNode);
+			const float CurrentPlayTime = SkeletalMeshComponent->GetPosition();
+			SkeletalMeshComponent->SetAnimation(TestAnimSequence);
+			SkeletalMeshComponent->SetPosition(CurrentPlayTime);
+			SkeletalMeshComponent->SetPlayRate(TestAnimSpeed);
+			SkeletalMeshComponent->Play(true);
+		}
+
+		SetResamplingInputOutputsNeeded(true);
+	}
+
+	void FMLDeformerEditorModel::OnPostInputAssetChanged()
+	{
+		CurrentTrainingFrame = -1;
+		Editor->UpdateTimeSliderRange();
+		Model->UpdateCachedNumVertices();
+		UpdateDeformerGraph();
+		RefreshMLDeformerComponents();
+		UpdateIsReadyForTrainingState();
+
+		const int32 TrainingFrame = GetTrainingFrameAtTime(CalcTrainingTimelinePosition());
+		SetTrainingFrame(TrainingFrame);
+
+		const int32 TestFrame = GetTestFrameAtTime(CalcTestTimelinePosition());
+		SetTestFrame(TestFrame);
+
+		UpdateEditorInputInfo();
+		CheckTrainingDataFrameChanged();
+	}
+
+	void FMLDeformerEditorModel::OnTimeSliderScrubPositionChanged(double NewScrubTime, bool bIsScrubbing)
+	{
+		float PlayOffset = static_cast<float>(NewScrubTime);
+
+		UMLDeformerVizSettings* VizSettings = Model->GetVizSettings();
+		if (VizSettings->GetVisualizationMode() == EMLDeformerVizMode::TrainingData)
+		{
+			const int32 TargetFrame = GetTrainingFrameAtTime(NewScrubTime);
+			for (FMLDeformerEditorActor* EditorActor : EditorActors)
+			{
+				if (EditorActor && EditorActor->IsTrainingActor())
+				{
+					if (Model->HasTrainingGroundTruth())
+					{
+						PlayOffset = GetTrainingTimeAtFrame(TargetFrame);
+					}
+					EditorActor->SetPlayPosition(PlayOffset);
+				}
+			}
+			VizSettings->SetTrainingFrameNumber(TargetFrame);
+		}
+		else if (VizSettings->GetVisualizationMode() == EMLDeformerVizMode::TestData)
+		{
+			const int32 TargetFrame = GetTestFrameAtTime(NewScrubTime);
+			for (FMLDeformerEditorActor* EditorActor : EditorActors)
+			{
+				if (EditorActor && EditorActor->IsTestActor())
+				{
+					if (Model->GetVizSettings()->HasTestGroundTruth())
+					{
+						PlayOffset = GetTestTimeAtFrame(TargetFrame);
+					}
+					EditorActor->SetPlayPosition(PlayOffset);
+				}
+			}
+			VizSettings->SetTestingFrameNumber(TargetFrame);
+		}
+	}
+
+	double FMLDeformerEditorModel::GetTrainingTimeAtFrame(int32 FrameNumber) const
+	{
+		return Model->GetAnimSequence() ? Model->GetAnimSequence()->GetTimeAtFrame(FrameNumber) : 0.0;
+	}
+
+	int32 FMLDeformerEditorModel::GetTrainingFrameAtTime(double TimeInSeconds) const
+	{
+		return Model->GetAnimSequence() ? Model->GetAnimSequence()->GetFrameAtTime(TimeInSeconds) : 0;
+	}
+
+	double FMLDeformerEditorModel::GetTestTimeAtFrame(int32 FrameNumber) const
+	{
+		return Model->GetVizSettings()->GetTestAnimSequence() ? Model->GetVizSettings()->GetTestAnimSequence()->GetTimeAtFrame(FrameNumber) : 0.0;
+	}
+
+	int32 FMLDeformerEditorModel::GetTestFrameAtTime(double TimeInSeconds) const
+	{
+		return Model->GetVizSettings()->GetTestAnimSequence() ? Model->GetVizSettings()->GetTestAnimSequence()->GetFrameAtTime(TimeInSeconds) : 0;
+	}
+
+	void FMLDeformerEditorModel::SetTrainingFrame(int32 FrameNumber)
+	{
+		UMLDeformerVizSettings* VizSettings = Model->GetVizSettings();
+		VizSettings->SetTrainingFrameNumber(FrameNumber);
+		ClampCurrentTrainingFrameIndex();
+		if (VizSettings->GetVisualizationMode() == EMLDeformerVizMode::TrainingData)
+		{
+			OnTimeSliderScrubPositionChanged(GetTrainingTimeAtFrame(FrameNumber), false);
+		}
+	}
+
+	void FMLDeformerEditorModel::SetTestFrame(int32 FrameNumber)
+	{
+		UMLDeformerVizSettings* VizSettings = Model->GetVizSettings();
+		VizSettings->SetTestingFrameNumber(FrameNumber);
+		ClampCurrentTestFrameIndex();
+		if (VizSettings->GetVisualizationMode() == EMLDeformerVizMode::TestData)
+		{
+			OnTimeSliderScrubPositionChanged(GetTestTimeAtFrame(FrameNumber), false);
+		}
+	}
+
+	void FMLDeformerEditorModel::OnPropertyChanged(FPropertyChangedEvent& PropertyChangedEvent)
+	{
+		const FProperty* Property = PropertyChangedEvent.Property;
+		if (Property == nullptr)
+		{
+			return;
+		}
+
+		// When we change one of these properties below, restart animations etc.
+		if (Property->GetFName() == UMLDeformerModel::GetSkeletalMeshPropertyName())
+		{
+			TriggerInputAssetChanged();
+			SetResamplingInputOutputsNeeded(true);
+			Model->InitVertexMap();
+			Model->InitGPUData();
+			UpdateDeformerGraph();
+		}
+		else if (Property->GetFName() == UMLDeformerVizSettings::GetTestAnimSequencePropertyName())
+		{
+			TriggerInputAssetChanged(true);
+		}
+		else if (Property->GetFName() == UMLDeformerModel::GetAnimSequencePropertyName())
+		{
+			TriggerInputAssetChanged();
+			SetResamplingInputOutputsNeeded(true);
+		}
+		if (Property->GetFName() == UMLDeformerModel::GetAlignmentTransformPropertyName() ||
+		    Property->GetFName() == UMLDeformerModel::GetDeltaCutoffLengthPropertyName())
+		{
+			if (PropertyChangedEvent.ChangeType == EPropertyChangeType::ValueSet)
+			{
+				SetResamplingInputOutputsNeeded(true);
+				SampleDeltas();
+			}
+		}
+		else if (Property->GetFName() == UMLDeformerModel::GetMaxTrainingFramesPropertyName())
+		{
+			TriggerInputAssetChanged();
+			SetResamplingInputOutputsNeeded(true);
+		}
+		else
+		if (Property->GetFName() == UMLDeformerModel::GetShouldIncludeBonesPropertyName() ||
+			Property->GetFName() == UMLDeformerModel::GetShouldIncludeCurvesPropertyName())
+		{
+			if (PropertyChangedEvent.ChangeType == EPropertyChangeType::ValueSet)
+			{
+				SetResamplingInputOutputsNeeded(true);
+				UpdateEditorInputInfo();
+				UpdateIsReadyForTrainingState();
+				GetEditor()->GetModelDetailsView()->ForceRefresh();
+			}
+		}
+		else
+		if (Property->GetFName() == UMLDeformerModel::GetBoneIncludeListPropertyName() ||
+			Property->GetFName() == UMLDeformerModel::GetCurveIncludeListPropertyName())
+		{
+			SetResamplingInputOutputsNeeded(true);
+			UpdateEditorInputInfo();
+			GetEditor()->GetModelDetailsView()->ForceRefresh();
+		}
+		else if (Property->GetFName() == UMLDeformerVizSettings::GetAnimPlaySpeedPropertyName())
+		{
+			UpdateTestAnimPlaySpeed();
+		}
+		else if (Property->GetFName() == UMLDeformerVizSettings::GetTrainingFrameNumberPropertyName())
+		{
+			ClampCurrentTrainingFrameIndex();
+			const int32 CurrentFrameNumber = Model->GetVizSettings()->GetTrainingFrameNumber();
+			OnTimeSliderScrubPositionChanged(GetTrainingTimeAtFrame(CurrentFrameNumber), false);
+		}
+		else if (Property->GetFName() == UMLDeformerVizSettings::GetTestingFrameNumberPropertyName())
+		{
+			ClampCurrentTestFrameIndex();
+			const int32 CurrentFrameNumber = Model->GetVizSettings()->GetTestingFrameNumber();
+			OnTimeSliderScrubPositionChanged(GetTestTimeAtFrame(CurrentFrameNumber), false);
+		}
+		else if (Property->GetFName() == UMLDeformerVizSettings::GetShowHeatMapPropertyName())
+		{
+			SetHeatMapMaterialEnabled(Model->GetVizSettings()->GetShowHeatMap());
+			UpdateDeformerGraph();
+		}
+		else
+		if (Property->GetFName() == UMLDeformerVizSettings::GetDrawLinearSkinnedActorPropertyName() ||
+			Property->GetFName() == UMLDeformerVizSettings::GetDrawMLDeformedActorPropertyName() ||
+			Property->GetFName() == UMLDeformerVizSettings::GetDrawGroundTruthActorPropertyName())
+		{
+			UpdateActorVisibility();
+		}
+		else if (Property->GetFName() == UMLDeformerVizSettings::GetDrawVertexDeltasPropertyName())
+		{
+			SampleDeltas();
+		} 
+		else if (Property->GetFName() == UMLDeformerVizSettings::GetDeformerGraphPropertyName())
+		{
+			UpdateDeformerGraph();
+			GetEditor()->GetVizSettingsDetailsView()->ForceRefresh();
+		};
+	}
+
+	FMLDeformerEditorActor* FMLDeformerEditorModel::GetVisualizationModeBaseActor() const
+	{
+		const UMLDeformerVizSettings* VizSettings = Model->GetVizSettings();
+		const bool bTestMode = VizSettings->GetVisualizationMode() == EMLDeformerVizMode::TestData;
+		return bTestMode ? FindEditorActor(ActorID_Test_Base) : FindEditorActor(ActorID_Train_Base);
+	}
+
+	void FMLDeformerEditorModel::OnPlayPressed()
+	{
+		// The play button shouldn't do anything in training mode.
+		const UMLDeformerVizSettings* VizSettings = Model->GetVizSettings();
+		const FMLDeformerEditorActor* BaseActor = GetVisualizationModeBaseActor();
+		const bool bMustPause = (BaseActor && BaseActor->IsPlaying());
+		const bool bTestMode = (VizSettings->GetVisualizationMode() == EMLDeformerVizMode::TestData);
+
+		// Update the current frame numbers.
+		if (bMustPause)
+		{
+			if (VizSettings->GetVisualizationMode() == EMLDeformerVizMode::TrainingData)
+			{
+				SetTrainingFrame(GetTrainingFrameAtTime(CalcTimelinePosition()));
+			}
+			else if (VizSettings->GetVisualizationMode() == EMLDeformerVizMode::TestData)
+			{
+				SetTestFrame(GetTestFrameAtTime(CalcTimelinePosition()));
+			}
+		}
+
+		for (FMLDeformerEditorActor* EditorActor : EditorActors)
+		{
+			if (EditorActor && ((EditorActor->IsTestActor() && bTestMode) || (EditorActor->IsTrainingActor() && !bTestMode)))
+			{
+				EditorActor->SetPlaySpeed(VizSettings->GetAnimPlaySpeed());
+				EditorActor->Pause(bMustPause);
+			}
+		}
+	}
+
+	bool FMLDeformerEditorModel::IsPlayingAnim() const
+	{
+		const FMLDeformerEditorActor* BaseActor = GetVisualizationModeBaseActor();
+		if (BaseActor)
+		{
+			UDebugSkelMeshComponent* SkeletalMeshComponent = BaseActor->GetSkeletalMeshComponent();
+			if (SkeletalMeshComponent)
+			{
+				return !SkeletalMeshComponent->bPauseAnims;
+			}
+		}
+		return false;
+	}
+
+	bool FMLDeformerEditorModel::IsPlayingForward() const
+	{
+		FMLDeformerEditorActor* BaseActor = GetVisualizationModeBaseActor();
+		if (BaseActor)
+		{
+			return BaseActor->GetPlaySpeed() > 0.0f;
+		}
+		return false;
+	}
+
+
+	double FMLDeformerEditorModel::CalcTrainingTimelinePosition() const
+	{
+		FMLDeformerEditorActor* EditorActor = FindEditorActor(ActorID_Train_GroundTruth);
+		if (EditorActor && EditorActor->HasVisualMesh())
+		{
+			return EditorActor->GetPlayPosition();
+		}
+
+		EditorActor = FindEditorActor(ActorID_Train_Base);
+		if (EditorActor && EditorActor->HasVisualMesh())
+		{
+			return EditorActor->GetPlayPosition();
+		}
+
+		return 0.0;
+	}
+
+	double FMLDeformerEditorModel::CalcTestTimelinePosition() const
+	{
+		FMLDeformerEditorActor* EditorActor = FindEditorActor(ActorID_Test_GroundTruth);
+		if (EditorActor && EditorActor->HasVisualMesh())
+		{
+			return EditorActor->GetPlayPosition();
+		}
+
+		EditorActor = FindEditorActor(ActorID_Test_Base);
+		if (EditorActor && EditorActor->HasVisualMesh())
+		{
+			return EditorActor->GetPlayPosition();
+		}
+
+		return 0.0;
+	}
+
+	void FMLDeformerEditorModel::UpdateTestAnimPlaySpeed()
+	{
+		UMLDeformerVizSettings* VizSettings = Model->GetVizSettings();
+		const float Speed = VizSettings->GetAnimPlaySpeed();
+		for (FMLDeformerEditorActor* EditorActor : EditorActors)
+		{
+			if (EditorActor && EditorActor->IsTestActor()) // Only do test actors, no training actors.
+			{
+				EditorActor->SetPlaySpeed(Speed);
+			}
+		}
+	}
+
+	void FMLDeformerEditorModel::ClampCurrentTrainingFrameIndex()
+	{
+		UMLDeformerVizSettings* VizSettings = Model->GetVizSettings();
+		if (GetNumTrainingFrames() > 0)
+		{
+			const int32 ClampedNumber = FMath::Min(VizSettings->GetTrainingFrameNumber(), GetNumTrainingFrames() - 1);
+			VizSettings->SetTrainingFrameNumber(ClampedNumber);
+		}
+		else
+		{
+			VizSettings->SetTrainingFrameNumber(0);
+		}
+	}
+
+	void FMLDeformerEditorModel::ClampCurrentTestFrameIndex()
+	{
+		UMLDeformerVizSettings* VizSettings = Model->GetVizSettings();
+		if (GetNumTestFrames() > 0)
+		{
+			const int32 ClampedNumber = FMath::Min(VizSettings->GetTestingFrameNumber(), GetNumTestFrames() - 1);
+			VizSettings->SetTestingFrameNumber(ClampedNumber);
+		}
+		else
+		{
+			VizSettings->SetTestingFrameNumber(0);
+		}
+	}
+
+	int32 FMLDeformerEditorModel::GetNumTestFrames() const
+	{
+		const UAnimSequence* Sequence = Model->GetVizSettings()->GetTestAnimSequence();
+		if (Sequence)
+		{
+			return Sequence->GetNumberOfSampledKeys();
+		}
+		return 0;
+	}
+
+	int32 FMLDeformerEditorModel::GetNumFramesForTraining() const 
+	{ 
+		return FMath::Min(GetNumTrainingFrames(), Model->GetTrainingFrameLimit());
+	}
+
+	FText FMLDeformerEditorModel::GetBaseAssetChangedErrorText() const
+	{
+		FText Result;
+		if (Model->GetSkeletalMesh() && Model->GetInputInfo())
+		{
+			if (Model->GetNumBaseMeshVerts() != Model->GetInputInfo()->GetNumBaseMeshVertices() &&
+				Model->GetNumBaseMeshVerts() > 0 && Model->GetInputInfo()->GetNumBaseMeshVertices() > 0)
+			{
+				Result = FText::Format(LOCTEXT("BaseMeshMismatch", "Number of vertices in base mesh has changed from {0} to {1} vertices since this ML Deformer Asset was saved! {2}"),
+					Model->GetInputInfo()->GetNumBaseMeshVertices(),
+					Model->GetNumBaseMeshVerts(),
+					IsTrained() ? LOCTEXT("BaseMeshMismatchNN", "Neural network needs to be retrained!") : FText());
+			}
+		}
+		return Result;
+	}
+
+	FText FMLDeformerEditorModel::GetVertexMapChangedErrorText() const
+	{
+		FText Result;
+		if (Model->GetSkeletalMesh())
+		{
+			bool bVertexMapMatch = true;
+			const FSkeletalMeshModel* ImportedModel = Model->GetSkeletalMesh()->GetImportedModel();
+			if (ImportedModel)
+			{
+				const TArray<int32>& MeshVertexMap = ImportedModel->LODModels[0].MeshToImportVertexMap;
+				const TArray<int32>& ModelVertexMap = Model->GetVertexMap();
+				if (MeshVertexMap.Num() == ModelVertexMap.Num())
+				{
+					for (int32 Index = 0; Index < ModelVertexMap.Num(); ++Index)
+					{
+						if (MeshVertexMap[Index] != ModelVertexMap[Index])
+						{
+							bVertexMapMatch = false;
+							break;
+						}
+					}
+					
+					if (!bVertexMapMatch)
+					{
+						Result = FText(LOCTEXT("VertexMapMismatch", "The vertex order of your Skeletal Mesh changed."));
+					}
+				}
+			}
+		}
+		return Result;
+	}
+
+	FText FMLDeformerEditorModel::GetSkeletalMeshNeedsReimportErrorText() const
+	{
+		FText Result;
+
+		if (Model->GetSkeletalMesh())
+		{
+			FSkeletalMeshModel* ImportedModel = Model->GetSkeletalMesh()->GetImportedModel();
+			check(ImportedModel);
+
+			const TArray<FSkelMeshImportedMeshInfo>& SkelMeshInfos = ImportedModel->LODModels[0].ImportedMeshInfos;
+			if (SkelMeshInfos.IsEmpty())
+			{
+				Result = LOCTEXT("SkelMeshNeedsReimport", "Skeletal Mesh asset needs to be reimported.");
+			}
+		}
+
+		return Result;
+	}
+
+	FText FMLDeformerEditorModel::GetInputsErrorText() const
+	{
+		if (Model->GetSkeletalMesh() && GetEditorInputInfo()->IsEmpty())
+		{
+			FString ErrorString;
+			if (Model->DoesSupportBones() && Model->ShouldIncludeBonesInTraining())
+			{
+				ErrorString += FText(LOCTEXT("InputsEmptyBonesErrorText", "Your base mesh has no bones to train on.\n")).ToString();
+			}
+			else
+			if (Model->DoesSupportCurves() && Model->ShouldIncludeCurvesInTraining())
+			{
+				ErrorString += FText(LOCTEXT("InputsNoCurvesToTrainText", "Your base mesh has no curves to train on.\n")).ToString();
+			}
+			else
+			{
+				ErrorString += FText(LOCTEXT("InputsNeededErrorText", "The training process needs inputs.\n")).ToString();
+			}
+
+			ErrorString.RemoveFromEnd("\n");
+			return FText::FromString(ErrorString);
+		}
+
+		return FText();
+	}
+
+	FText FMLDeformerEditorModel::GetIncompatibleSkeletonErrorText(USkeletalMesh* InSkelMesh, UAnimSequence* InAnimSeq) const
+	{
+		FText Result;
+		if (InSkelMesh && InAnimSeq)
+		{
+			if (!InSkelMesh->GetSkeleton()->IsCompatible(InAnimSeq->GetSkeleton()))
+			{
+				Result = LOCTEXT("SkeletonMismatch", "The base skeletal mesh and anim sequence use different skeletons. The animation might not play correctly.");
+			}
+		}
+		return Result;
+	}
+
+	FText FMLDeformerEditorModel::GetTargetAssetChangedErrorText() const
+	{
+		FText Result;
+
+		UMLDeformerInputInfo* InputInfo = Model->GetInputInfo();
+		if (Model->HasTrainingGroundTruth() && InputInfo)
+		{
+			if (Model->GetNumTargetMeshVerts() != InputInfo->GetNumTargetMeshVertices() &&
+				Model->GetNumTargetMeshVerts() > 0 && InputInfo->GetNumTargetMeshVertices() > 0)
+			{
+				Result = FText::Format(LOCTEXT("TargetMeshMismatch", "Number of vertices in target mesh has changed from {0} to {1} vertices since this ML Deformer Asset was saved! {2}"),
+					InputInfo->GetNumTargetMeshVertices(),
+					Model->GetNumTargetMeshVerts(),
+					IsTrained() ? LOCTEXT("BaseMeshMismatchModel", "Model needs to be retrained!") : FText());
+			}
+		}
+
+		return Result;
+	}
+
+	void FMLDeformerEditorModel::InitInputInfo(UMLDeformerInputInfo* InputInfo)
+	{
+		InputInfo->Reset();
+
+		TArray<FString>& BoneNameStrings = InputInfo->GetBoneNameStrings();
+		TArray<FString>& CurveNameStrings = InputInfo->GetCurveNameStrings();
+		TArray<FName>& BoneNames = InputInfo->GetBoneNames();
+		TArray<FName>& CurveNames = InputInfo->GetCurveNames();
+		USkeletalMesh* SkeletalMesh = Model->GetSkeletalMesh();
+
+		BoneNames.Reset();
+		BoneNameStrings.Reset();
+		CurveNames.Reset();
+
+		InputInfo->SetNumBaseVertices(Model->GetNumBaseMeshVerts());
+		InputInfo->SetNumTargetVertices(Model->GetNumTargetMeshVerts());
+
+		const bool bIncludeBones = Model->DoesSupportBones() && Model->ShouldIncludeBonesInTraining();
+		const bool bIncludeCurves = Model->DoesSupportCurves() && Model->ShouldIncludeCurvesInTraining();
+		const USkeleton* Skeleton = Model->GetSkeletalMesh() ? Model->GetSkeletalMesh()->GetSkeleton() : nullptr;
+
+		// Handle bones.
+		if (bIncludeBones && SkeletalMesh)
+		{
+			// Include all the bones when no list was provided.
+			const FReferenceSkeleton& RefSkeleton = SkeletalMesh->GetRefSkeleton();
+			if (Model->GetBoneIncludeList().IsEmpty())
+			{
+				// Grab all bone names.
+				const int32 NumBones = RefSkeleton.GetNum();
+				BoneNameStrings.Reserve(NumBones);
+				for (int32 Index = 0; Index < NumBones; ++Index)
+				{
+					const FName BoneName = RefSkeleton.GetBoneName(Index);
+					BoneNameStrings.Add(BoneName.ToString());
+					BoneNames.Add(BoneName);
+				}
+			}
+			else // A list of bones to include was provided.
+			{
+				for (const FBoneReference& BoneReference : Model->GetBoneIncludeList())
+				{
+					if (BoneReference.BoneName.IsValid())
+					{
+						const FName BoneName = BoneReference.BoneName;
+						if (RefSkeleton.FindBoneIndex(BoneName) == INDEX_NONE)
+						{
+							UE_LOG(LogMLDeformer, Warning, TEXT("Bone '%s' in the bones include list doesn't exist, ignoring it."), *BoneName.ToString());
+							continue;
+						}
+
+						BoneNameStrings.Add(BoneName.ToString());
+						BoneNames.Add(BoneName);
+					}
+				}
+			}
+		}
+
+		// Handle curves.
+		if (bIncludeCurves && SkeletalMesh)
+		{
+			// Anim curves.
+			const FSmartNameMapping* SmartNameMapping = Skeleton ? Skeleton->GetSmartNameContainer(USkeleton::AnimCurveMappingName) : nullptr;
+			if (SmartNameMapping) // When there are curves.
+			{
+				// Include all curves when no list was provided.
+				if (Model->GetCurveIncludeList().IsEmpty())
+				{
+					SmartNameMapping->FillNameArray(CurveNames);
+					CurveNameStrings.Reserve(CurveNames.Num());
+					for (const FName Name : CurveNames)
+					{
+						CurveNameStrings.Add(Name.ToString());
+					}
+				}
+				else // A list of curve names was provided.
+				{
+					for (const FMLDeformerCurveReference& CurveReference : Model->GetCurveIncludeList())
+					{
+						if (CurveReference.CurveName.IsValid())
+						{
+							const FName CurveName = CurveReference.CurveName;
+							if (!SmartNameMapping->Exists(CurveName))
+							{
+								UE_LOG(LogMLDeformer, Warning, TEXT("Curve '%s' doesn't exist, ignoring it."), *CurveName.ToString());
+								continue;
+							}
+
+							CurveNameStrings.Add(CurveName.ToString());
+							CurveNames.Add(CurveName);
+						}
+					}
+				}
+			}
+		}
+	}
+
+	void FMLDeformerEditorModel::InitBoneIncludeListToAnimatedBonesOnly()
+	{
+		if (!Model->GetAnimSequence())
+		{
+			UE_LOG(LogMLDeformer, Warning, TEXT("Cannot initialize bone list as no Anim Sequence has been picked."));
+			return;
+		}
+
+		const UAnimDataModel* DataModel = Model->GetAnimSequence()->GetDataModel();
+		if (!DataModel)
+		{
+			UE_LOG(LogMLDeformer, Warning, TEXT("Anim sequence has no data model."));
+			return;
+		}
+
+		if (!Model->GetSkeletalMesh())
+		{
+			UE_LOG(LogMLDeformer, Warning, TEXT("Skeletal Mesh has not been set."));
+			return;
+		}
+
+		USkeleton* Skeleton = Model->GetSkeletalMesh()->GetSkeleton();
+		if (!Skeleton)
+		{
+			UE_LOG(LogMLDeformer, Warning, TEXT("Skeletal Mesh has no skeleton."));
+			return;
+		}
+
+		// Iterate over all bones that are both in the skeleton and the animation.
+		TArray<FName> AnimatedBoneList;
+		const FReferenceSkeleton& RefSkeleton = Skeleton->GetReferenceSkeleton();
+		const int32 NumBones = RefSkeleton.GetNum();
+		for (int32 Index = 0; Index < NumBones; ++Index)
+		{
+			const FName BoneName = RefSkeleton.GetBoneName(Index);
+			const int32 BoneTrackIndex = DataModel->GetBoneTrackIndexByName(BoneName);
+			if (BoneTrackIndex == INDEX_NONE)
+			{
+				continue;
+			}
+
+			// Check if there is actually animation data.
+			const FBoneAnimationTrack& BoneAnimTrack = DataModel->GetBoneTrackByIndex(BoneTrackIndex);
+			const TArray<FQuat4f>& Rotations = BoneAnimTrack.InternalTrackData.RotKeys;
+			bool bIsAnimated = false;
+			if (!Rotations.IsEmpty())
+			{
+				const FQuat4f FirstQuat = Rotations[0];
+				for (const FQuat4f KeyValue : Rotations)
+				{
+					if (!KeyValue.Equals(FirstQuat))
+					{
+						bIsAnimated = true;
+						break;
+					}
+				}
+
+				if (!bIsAnimated)
+				{
+					UE_LOG(LogMLDeformer, Display, TEXT("Bone '%s' has keyframes but isn't animated."), *BoneName.ToString());
+				}
+			}
+
+			if (bIsAnimated)
+			{
+				AnimatedBoneList.Add(BoneName);
+			}
+		}
+
+		// Init the bone include list using the animated bones.
+		if (!AnimatedBoneList.IsEmpty())
+		{
+			TArray<FBoneReference>& BoneList = Model->GetBoneIncludeList();
+			BoneList.Empty();
+			BoneList.Reserve(AnimatedBoneList.Num());
+			for (FName BoneName : AnimatedBoneList)
+			{
+				BoneList.AddDefaulted();
+				FBoneReference& BoneRef = BoneList.Last();
+				BoneRef.BoneName = BoneName;
+			}
+		}
+		else
+		{
+			Model->GetBoneIncludeList().Empty();
+			UE_LOG(LogMLDeformer, Warning, TEXT("There are no animated bone rotations in Anim Sequence '%s'."), *(Model->GetAnimSequence()->GetName()));
+		}
+		UpdateEditorInputInfo();
+	}
+
+	void FMLDeformerEditorModel::InitCurveIncludeListToAnimatedCurvesOnly()
+	{
+		if (!Model->GetAnimSequence())
+		{
+			UE_LOG(LogMLDeformer, Warning, TEXT("Cannot initialize curve list as no Anim Sequence has been picked."));
+			return;
+		}
+
+		const UAnimDataModel* DataModel = Model->GetAnimSequence()->GetDataModel();
+		if (!DataModel)
+		{
+			UE_LOG(LogMLDeformer, Warning, TEXT("Anim sequence has no data model."));
+			return;
+		}
+
+		if (!Model->GetSkeletalMesh())
+		{
+			UE_LOG(LogMLDeformer, Warning, TEXT("Skeletal Mesh has not been set."));
+			return;
+		}
+
+		USkeleton* Skeleton = Model->GetSkeletalMesh()->GetSkeleton();
+		if (!Skeleton)
+		{
+			UE_LOG(LogMLDeformer, Warning, TEXT("Skeletal Mesh has no skeleton."));
+			return;
+		}
+
+		// Iterate over all curves that are both in the skeleton and the animation.
+		TArray<FName> AnimatedCurveList;
+		const FSmartNameMapping* Mapping = Skeleton->GetSmartNameContainer(USkeleton::AnimCurveMappingName);
+		if (Mapping)
+		{
+			TArray<FName> SkeletonCurveNames;
+			Mapping->FillNameArray(SkeletonCurveNames);
+			for (const FName& SkeletonCurveName : SkeletonCurveNames)
+			{
+				const TArray<FFloatCurve>& AnimCurves = DataModel->GetFloatCurves();
+				for (const FFloatCurve& AnimCurve : AnimCurves)
+				{
+					if (AnimCurve.Name.IsValid() && AnimCurve.Name.DisplayName == SkeletonCurveName)
+					{
+						TArray<float> TimeValues;
+						TArray<float> KeyValues;
+						AnimCurve.GetKeys(TimeValues, KeyValues);
+						if (KeyValues.Num() > 0)
+						{
+							const float FirstKeyValue = KeyValues[0];					
+							for (float CurKeyValue : KeyValues)
+							{
+								if (CurKeyValue != FirstKeyValue)
+								{
+									AnimatedCurveList.Add(SkeletonCurveName);
+									break;
+								}
+							}
+						}
+						break;
+					}
+				}
+			}
+		}
+
+		// Init the bone include list using the animated bones.
+		if (!AnimatedCurveList.IsEmpty())
+		{
+			TArray<FMLDeformerCurveReference>& CurveList = Model->GetCurveIncludeList();
+			CurveList.Empty();
+			CurveList.Reserve(AnimatedCurveList.Num());
+			for (FName CurveName : AnimatedCurveList)
+			{
+				CurveList.AddDefaulted();
+				FMLDeformerCurveReference& CurveRef = CurveList.Last();
+				CurveRef.CurveName = CurveName;
+			}
+		}
+		else
+		{
+			Model->GetCurveIncludeList().Empty();
+			UE_LOG(LogMLDeformer, Warning, TEXT("There are no animated curves in Anim Sequence '%s'."), *(Model->GetAnimSequence()->GetName()));
+		}
+		UpdateEditorInputInfo();
+	}
+
+	void FMLDeformerEditorModel::Render(const FSceneView* View, FViewport* Viewport, FPrimitiveDrawInterface* PDI)
+	{
+		// Make sure that before we render anything, that our sampler is ready.
+		if (!Sampler->IsInitialized())
+		{
+			Sampler->Init(this); // This can still fail.
+			Sampler->SetVertexDeltaSpace(EVertexDeltaSpace::PostSkinning);
+			if (Sampler->IsInitialized()) // If we actually managed to initialize this frame.
+			{
+				SampleDeltas(); // Update the deltas.
+			}
+		}
+
+		const UMLDeformerVizSettings* VizSettings = Model->GetVizSettings();
+		if (VizSettings->GetVisualizationMode() == EMLDeformerVizMode::TrainingData)
+		{
+			const TArray<float>& VertexDeltas = Sampler->GetVertexDeltas();
+			const TArray<FVector3f>& LinearSkinnedPositions = Sampler->GetSkinnedVertexPositions();
+			bool bDrawDeltas = (VizSettings->GetDrawVertexDeltas() && (VertexDeltas.Num() / 3) == LinearSkinnedPositions.Num());
+
+			// Disable drawing deltas when we're playing the training anim sequence, as that can get too slow.
+			const FMLDeformerEditorActor* BaseActor = FindEditorActor(ActorID_Train_Base);
+			if (BaseActor)
+			{
+				bDrawDeltas &= !BaseActor->IsPlaying();
+			}
+
+			// Draw the deltas for the current frame.
+			if (bDrawDeltas)
+			{
+				const FLinearColor DeltasColor = FMLDeformerEditorStyle::Get().GetColor("MLDeformer.Deltas.Color");
+				const FLinearColor DebugVectorsColor = FMLDeformerEditorStyle::Get().GetColor("MLDeformer.DebugVectors.Color");
+				const FLinearColor DebugVectorsColor2 = FMLDeformerEditorStyle::Get().GetColor("MLDeformer.DebugVectors.Color2");
+				const uint8 DepthGroup = VizSettings->GetXRayDeltas() ? 100 : 0;
+				for (int32 Index = 0; Index < LinearSkinnedPositions.Num(); ++Index)
+				{
+					const int32 ArrayIndex = 3 * Index;
+					const FVector Delta(
+						VertexDeltas[ArrayIndex], 
+						VertexDeltas[ArrayIndex + 1], 
+						VertexDeltas[ArrayIndex + 2]);
+					const FVector VertexPos = (FVector)LinearSkinnedPositions[Index];
+					PDI->DrawLine(VertexPos, VertexPos + Delta, DeltasColor, DepthGroup);
+				}
+			}
+		}
+	}
+
+	void FMLDeformerEditorModel::SampleDeltas()
+	{
+		UMLDeformerVizSettings* VizSettings = Model->GetVizSettings();
+		ClampCurrentTrainingFrameIndex();
+
+		// If we have no Persona toolkit yet, then it is not yet safe to init the sampler.
+		if (Editor->GetPersonaToolkitPointer() != nullptr)
+		{
+			Sampler->Init(this);
+		}
+
+		if (Sampler->IsInitialized())
+		{
+			Sampler->SetVertexDeltaSpace(EVertexDeltaSpace::PostSkinning);
+			Sampler->Sample(VizSettings->GetTrainingFrameNumber());
+		}
+	}
+
+	void FMLDeformerEditorModel::CheckTrainingDataFrameChanged()
+	{
+		UMLDeformerVizSettings* VizSettings = Model->GetVizSettings();
+		ClampCurrentTrainingFrameIndex();
+		if (CurrentTrainingFrame != VizSettings->GetTrainingFrameNumber())
+		{
+			OnTrainingDataFrameChanged();
+		}
+	}
+
+	void FMLDeformerEditorModel::OnTrainingDataFrameChanged()
+	{
+		UMLDeformerVizSettings* VizSettings = Model->GetVizSettings();
+
+		// If the current frame number changed, re-sample the deltas if needed.
+		if (CurrentTrainingFrame != VizSettings->GetTrainingFrameNumber())
+		{
+			CurrentTrainingFrame = VizSettings->GetTrainingFrameNumber();
+			if (VizSettings->GetDrawVertexDeltas() && VizSettings->GetVisualizationMode() == EMLDeformerVizMode::TrainingData)
+			{
+				SampleDeltas();
+			}
+		}
+	}
+
+	void FMLDeformerEditorModel::RefreshMLDeformerComponents()
+	{
+		for (FMLDeformerEditorActor* EditorActor : EditorActors)
+		{
+			if (EditorActor && EditorActor->GetMLDeformerComponent())
+			{
+				USkeletalMeshComponent* SkelMeshComponent = EditorActor ? EditorActor->GetSkeletalMeshComponent() : nullptr;
+				UMLDeformerAsset* DeformerAsset = GetModel()->GetDeformerAsset();
+				EditorActor->GetMLDeformerComponent()->SetupComponent(DeformerAsset, SkelMeshComponent);
+				UMLDeformerModelInstance* ModelInstance = EditorActor->GetMLDeformerComponent()->GetModelInstance();
+				if (ModelInstance)
+				{
+					ModelInstance->UpdateCompatibilityStatus();
+				}
+			}
+		}
+	}
+
+	void FMLDeformerEditorModel::CreateHeatMapMaterial()
+	{
+		const FString HeatMapMaterialPath = GetHeatMapMaterialPath();
+		UObject* MaterialObject = StaticLoadObject(UMaterial::StaticClass(), nullptr, *HeatMapMaterialPath);
+		HeatMapMaterial = Cast<UMaterial>(MaterialObject);
+	}
+
+	void FMLDeformerEditorModel::CreateHeatMapDeformerGraph()
+	{
+		const FString HeatMapDeformerPath = GetHeatMapDeformerGraphPath();
+		UObject* DeformerObject = StaticLoadObject(UMeshDeformer::StaticClass(), nullptr, *HeatMapDeformerPath);
+		HeatMapDeformerGraph = Cast<UMeshDeformer>(DeformerObject);
+	}
+
+	void FMLDeformerEditorModel::CreateHeatMapAssets()
+	{
+		CreateHeatMapMaterial();
+		CreateHeatMapDeformerGraph();
+	}
+
+	void FMLDeformerEditorModel::SetHeatMapMaterialEnabled(bool bEnabled)
+	{
+		FMLDeformerEditorActor* EditorActor = FindEditorActor(ActorID_Test_MLDeformed);
+		if (EditorActor == nullptr)
+		{
+			return;
+		}
+
+		USkeletalMeshComponent* Component = EditorActor->GetSkeletalMeshComponent();
+		if (Component)
+		{
+			if (bEnabled)
+			{
+				for (int32 Index = 0; Index < Component->GetNumMaterials(); ++Index)
+				{
+					Component->SetMaterial(Index, HeatMapMaterial);
+				}
+			}
+			else
+			{
+				Component->EmptyOverrideMaterials();
+			}
+		}
+
+		UpdateDeformerGraph();
+	}
+
+	UMeshDeformer* FMLDeformerEditorModel::LoadDefaultDeformerGraph()
+	{
+		const FString GraphAssetPath = GetDefaultDeformerGraphAssetPath();
+		UObject* Object = StaticLoadObject(UMeshDeformer::StaticClass(), nullptr, *GraphAssetPath);
+		UMeshDeformer* DeformerGraph = Cast<UMeshDeformer>(Object);
+		if (DeformerGraph == nullptr)
+		{
+			UE_LOG(LogMLDeformer, Warning, TEXT("Failed to load default ML Deformer compute graph from: %s"), *GraphAssetPath);
+		}
+		else
+		{
+			UE_LOG(LogMLDeformer, Verbose, TEXT("Loaded default ML Deformer compute graph from: %s"), *GraphAssetPath);
+		}
+
+		return DeformerGraph;
+	}
+
+	void FMLDeformerEditorModel::SetDefaultDeformerGraphIfNeeded()
+	{
+		// Initialize the asset on the default plugin deformer graph.
+		UMLDeformerVizSettings* VizSettings = Model->GetVizSettings();
+		if (VizSettings && VizSettings->GetDeformerGraph() == nullptr)
+		{
+			UMeshDeformer* DefaultGraph = LoadDefaultDeformerGraph();
+			VizSettings->SetDeformerGraph(DefaultGraph);
+		}
+	}
+
+	FText FMLDeformerEditorModel::GetOverlayText() const
+	{
+		const FMLDeformerEditorActor* EditorActor = FindEditorActor(ActorID_Test_MLDeformed);
+		const UMLDeformerComponent* DeformerComponent = EditorActor ? EditorActor->GetMLDeformerComponent() : nullptr;
+		if (DeformerComponent)
+		{
+			const UMLDeformerModelInstance* ModelInstance = DeformerComponent->GetModelInstance();
+			if (ModelInstance &&
+				ModelInstance->GetSkeletalMeshComponent() && 
+				ModelInstance->GetSkeletalMeshComponent()->GetSkeletalMeshAsset() &&
+				!ModelInstance->IsCompatible() )
+			{
+				return FText::FromString( ModelInstance->GetCompatibilityErrorText() );
+			}
+		}
+	 
+		return FText::GetEmpty();
+	}
+
+	void FMLDeformerEditorModel::UpdateDeformerGraph()
+	{	
+		for (FMLDeformerEditorActor* EditorActor : EditorActors)
+		{
+			if (EditorActor == nullptr || EditorActor->GetMLDeformerComponent() == nullptr)
+			{
+				continue;
+			}
+
+			USkeletalMeshComponent* SkelMeshComponent = EditorActor->GetSkeletalMeshComponent();
+			if (SkelMeshComponent == nullptr)
+			{
+				continue;
+			}
+
+			if (SkelMeshComponent)
+			{
+				const UMLDeformerVizSettings* VizSettings = Model->GetVizSettings();
+				UMeshDeformer* MeshDeformer = IsTrained() ? VizSettings->GetDeformerGraph() : nullptr;	
+				const bool bUseHeatMapDeformer = VizSettings->GetShowHeatMap();
+				SkelMeshComponent->SetMeshDeformer(bUseHeatMapDeformer ? HeatMapDeformerGraph.Get() : MeshDeformer);
+			}
+		}
+	}
+
+	void FMLDeformerEditorModel::OnPostTraining(ETrainingResult TrainingResult, bool bUsePartiallyTrainedWhenAborted)
+	{
+		for (FMLDeformerEditorActor* EditorActor : EditorActors)
+		{
+			if (EditorActor && EditorActor->GetMLDeformerComponent())
+			{
+				USkeletalMeshComponent* SkelMeshComponent = EditorActor->GetSkeletalMeshComponent();
+				UMLDeformerAsset* DeformerAsset = Model->GetDeformerAsset();
+				EditorActor->GetMLDeformerComponent()->SetupComponent(DeformerAsset, SkelMeshComponent);
+			}
+		}
+
+		Sampler->SetVertexDeltaSpace(EVertexDeltaSpace::PostSkinning);
+		SampleDeltas();
+		Model->InitGPUData();
+	}
+
+	FMLDeformerEditorActor* FMLDeformerEditorModel::GetTimelineEditorActor() const
+	{
+		const UMLDeformerVizSettings* VizSettings = Model->GetVizSettings();
+		if (VizSettings->GetVisualizationMode() == EMLDeformerVizMode::TrainingData)
+		{
+			return FindEditorActor(ActorID_Train_GroundTruth);
+		}
+		else if (VizSettings->GetVisualizationMode() == EMLDeformerVizMode::TestData)
+		{
+			return FindEditorActor(ActorID_Test_GroundTruth);
+		}
+		return nullptr;
+	}
+
+	UNeuralNetwork* FMLDeformerEditorModel::LoadNeuralNetworkFromOnnx(const FString& Filename) const
+	{
+		const FString OnnxFile = FPaths::ConvertRelativePathToFull(Filename);
+		if (FPaths::FileExists(OnnxFile))
+		{
+			UE_LOG(LogMLDeformer, Display, TEXT("Loading Onnx file '%s'..."), *OnnxFile);
+			UNeuralNetwork* Result = NewObject<UNeuralNetwork>(Model, UNeuralNetwork::StaticClass());		
+			if (Result->Load(OnnxFile))
+			{
+				if (Model->IsNeuralNetworkOnGPU())
+				{
+					Result->SetDeviceType(ENeuralDeviceType::GPU, ENeuralDeviceType::CPU, ENeuralDeviceType::GPU);	
+					if (Result->GetDeviceType() != ENeuralDeviceType::GPU || Result->GetOutputDeviceType() != ENeuralDeviceType::GPU || Result->GetInputDeviceType() != ENeuralDeviceType::CPU)
+					{
+						UE_LOG(LogMLDeformer, Error, TEXT("Neural net in ML Deformer '%s' cannot run on the GPU, it will not be active."), *Model->GetDeformerAsset()->GetName());
+					}
+				}
+				else
+				{
+					Result->SetDeviceType(ENeuralDeviceType::CPU, ENeuralDeviceType::CPU, ENeuralDeviceType::CPU);
+				}
+				UE_LOG(LogMLDeformer, Display, TEXT("Successfully loaded Onnx file '%s'..."), *OnnxFile);
+				return Result;
+			}
+			else
+			{
+				UE_LOG(LogMLDeformer, Error, TEXT("Failed to load Onnx file '%s'"), *OnnxFile);
+			}
+		}
+		else
+		{
+			UE_LOG(LogMLDeformer, Error, TEXT("Onnx file '%s' does not exist!"), *OnnxFile);
+		}
+
+		return nullptr;
+	}
+
+	bool FMLDeformerEditorModel::IsEditorReadyForTrainingBasicChecks()
+	{
+		// Make sure we have picked required assets.
+		if (!Model->HasTrainingGroundTruth() ||
+			Model->GetAnimSequence() == nullptr ||
+			Model->GetSkeletalMesh() == nullptr ||
+			GetNumTrainingFrames() == 0)
+		{
+			return false;
+		}
+
+		// Make sure we have inputs.
+		UpdateEditorInputInfo();
+		if (GetEditorInputInfo()->IsEmpty())
+		{
+			return false;
+		}
+
+		return true;
+	}
+
+	bool FMLDeformerEditorModel::LoadTrainedNetwork() const
+	{
+		const FString OnnxFile = GetTrainedNetworkOnnxFile();
+		UNeuralNetwork* Network = LoadNeuralNetworkFromOnnx(OnnxFile);
+		if (Network)
+		{
+			Model->SetNeuralNetwork(Network);
+			return true;
+		}
+
+		return false;
+	}
+
+	bool FMLDeformerEditorModel::IsTrained() const
+	{
+		return (Model->GetNeuralNetwork() != nullptr);
+	}
+
+	void FMLDeformerEditorModel::TriggerInputAssetChanged(bool bRefreshVizSettings)
+	{
+		OnInputAssetsChanged();
+		OnPostInputAssetChanged();
+		Model->SetAssetEditorOnlyFlags();
+		GetEditor()->GetModelDetailsView()->ForceRefresh();
+		if (bRefreshVizSettings)
+		{
+			GetEditor()->GetVizSettingsDetailsView()->ForceRefresh();
+		}
+	}
+
+	void FMLDeformerEditorModel::ZeroDeltasByThreshold(TArray<FVector3f>& Deltas, float Threshold)
+	{
+		for (int32 VertexIndex = 0; VertexIndex < Deltas.Num(); ++VertexIndex)
+		{
+			if (Deltas[VertexIndex].Length() <= Threshold)
+			{
+				Deltas[VertexIndex] = FVector3f::ZeroVector;
+			}
+		}
+	}
+
+	void FMLDeformerEditorModel::CreateEngineMorphTargets(TArray<UMorphTarget*>& OutMorphTargets, const TArray<FVector3f>& Deltas, const FString& NamePrefix, int32 LOD, float DeltaThreshold)
+	{
+		OutMorphTargets.Reset();
+
+		const int32 NumBaseMeshVerts = Model->GetNumBaseMeshVerts();
+		check(Deltas.Num() % NumBaseMeshVerts == 0);
+		const int32 NumMorphTargets = Deltas.Num() / NumBaseMeshVerts;
+		check((Deltas.Num() / NumMorphTargets) == NumBaseMeshVerts);
+		check(!Model->GetVertexMap().IsEmpty());
+
+		USkeletalMesh* SkelMesh = Model->GetSkeletalMesh();
+		FSkeletalMeshRenderData* RenderData = SkelMesh->GetResourceForRendering();
+		check(RenderData);
+		check(!RenderData->LODRenderData.IsEmpty());
+		const int32 NumRenderVertices = RenderData->LODRenderData[LOD].GetNumVertices();
+
+		// Initialize an engine morph target for each model morph target.
+		UE_LOG(LogMLDeformer, Display, TEXT("Initializing %d engine morph targets of %d vertices each"), NumMorphTargets, Deltas.Num() / NumMorphTargets);
+		for (int32 BlendShapeIndex = 0; BlendShapeIndex < NumMorphTargets; ++BlendShapeIndex)
+		{
+			const FName MorphName = *FString::Printf(TEXT("%s%.3d"), *NamePrefix, BlendShapeIndex);
+			UMorphTarget* MorphTarget = NewObject<UMorphTarget>(SkelMesh, MorphName);
+			MorphTarget->BaseSkelMesh = SkelMesh;
+			OutMorphTargets.Add(MorphTarget);
+
+			// Create a new LOD model for this morph.
+			TArray<FMorphTargetLODModel>& MorphLODs = MorphTarget->GetMorphLODModels();
+			MorphLODs.AddDefaulted();
+			FMorphTargetLODModel& MorphLODModel = MorphLODs.Last();
+
+			// Initialize the morph target LOD level.
+			MorphLODModel.Reset();
+			MorphLODModel.bGeneratedByEngine = true;
+			MorphLODModel.NumBaseMeshVerts = NumRenderVertices;
+			MorphLODModel.NumVertices = NumRenderVertices;
+
+			// Init sections.
+			const int32 NumSections = RenderData->LODRenderData[LOD].RenderSections.Num();
+			MorphLODModel.SectionIndices.AddUninitialized(NumSections);
+			for (int32 SectionIndex = 0; SectionIndex < NumSections; ++SectionIndex)
+			{
+				MorphLODModel.SectionIndices[SectionIndex] = SectionIndex;
+			}
+
+			// Init deltas for this morph target.
+			const TArray<int32>& VertexMap = Model->GetVertexMap();
+			MorphLODModel.Vertices.Reserve(NumRenderVertices);
+			for (int32 VertexIndex = 0; VertexIndex < NumRenderVertices; ++VertexIndex)
+			{
+				const int32 ImportedVertexNumber = VertexMap[VertexIndex];
+				if (ImportedVertexNumber != INDEX_NONE)
+				{
+					const FVector3f Delta = Deltas[ImportedVertexNumber + BlendShapeIndex * NumBaseMeshVerts];
+					if (Delta.Length() > DeltaThreshold)
+					{
+						MorphLODModel.Vertices.AddDefaulted();
+						FMorphTargetDelta& MorphTargetDelta = MorphLODModel.Vertices.Last();
+						MorphTargetDelta.PositionDelta = Delta;
+						MorphTargetDelta.SourceIdx = VertexIndex;
+						MorphTargetDelta.TangentZDelta = FVector3f::ZeroVector;
+					}
+				}
+			}
+
+			MorphLODModel.Vertices.Shrink();
+		}
+	}
+
+	void FMLDeformerEditorModel::CompressEngineMorphTargets(FMorphTargetVertexInfoBuffers& OutMorphBuffers, const TArray<UMorphTarget*>& MorphTargets, int32 LOD, float MorphErrorTolerance)
+	{
+		USkeletalMesh* SkelMesh = Model->GetSkeletalMesh();
+		FSkeletalMeshRenderData* RenderData = SkelMesh->GetResourceForRendering();
+		check(RenderData);
+		check(!RenderData->LODRenderData.IsEmpty());
+		const int32 NumRenderVertices = RenderData->LODRenderData[LOD].GetNumVertices();
+
+		// Release any existing morph buffer data.
+		if (OutMorphBuffers.IsRHIIntialized() && OutMorphBuffers.IsInitialized())
+		{
+			ReleaseResourceAndFlush(&OutMorphBuffers);
+		}
+
+		// Don't empty the array of morph target data when we init the RHI buffers, as we need them to serialize later on.
+		OutMorphBuffers = FMorphTargetVertexInfoBuffers();
+		OutMorphBuffers.SetEmptyMorphCPUDataOnInitRHI(false);
+
+		// Initialize the compressed morph target buffers.
+		OutMorphBuffers.InitMorphResources
+		(
+			GMaxRHIShaderPlatform,
+			TArray<FSkelMeshRenderSection>(),	// Empty array, as we don't need tangents, since we recalculate them in the shader.
+			MorphTargets,
+			NumRenderVertices,
+			LOD,
+			MorphErrorTolerance
+		);
+	}
+
+	void FMLDeformerEditorModel::DrawMorphTarget(FPrimitiveDrawInterface* PDI, const TArray<FVector3f>& MorphDeltas, float DeltaThreshold, int32 MorphTargetIndex, const FVector& DrawOffset)
+	{
+		UMLDeformerVizSettings* VizSettings = Model->GetVizSettings();
+		if (!MorphDeltas.IsEmpty())
+		{
+			const int32 NumVerts = Model->GetNumBaseMeshVerts();
+			check(MorphDeltas.Num() % NumVerts == 0);
+
+			const TArray<FVector3f>& SkinnedPositions = Sampler->GetUnskinnedVertexPositions();
+			check(NumVerts == SkinnedPositions.Num());
+
+			// Draw all deltas.			
+			const int32 NumMorphTargets = MorphDeltas.Num() / NumVerts;
+			const int32 FinalMorphTargetIndex = FMath::Clamp<int32>(MorphTargetIndex, 0, NumMorphTargets - 1);
+			const FLinearColor IncludedColor(1.0f, 0.0f, 1.0f);
+			const FLinearColor ExcludedColor(0.1f, 0.1f, 0.1f);
+			for (int32 VertexIndex = 0; VertexIndex < NumVerts; ++VertexIndex)
+			{
+				const FVector StartPoint = FVector(SkinnedPositions[VertexIndex].X, SkinnedPositions[VertexIndex].Y, SkinnedPositions[VertexIndex].Z) + DrawOffset;
+				const int32 DeltaArrayOffset = NumVerts * FinalMorphTargetIndex + VertexIndex;
+				const FVector Delta(MorphDeltas[DeltaArrayOffset]);
+
+				if (Delta.Length() >= DeltaThreshold)
+				{
+					PDI->DrawPoint(StartPoint, IncludedColor, 1.0f, 0);
+					PDI->DrawLine(StartPoint, StartPoint + Delta, IncludedColor, 0);
+				}
+				else
+				{
+					PDI->DrawPoint(StartPoint + Delta, ExcludedColor, 0.75f, 0);
+				}
+			}
+		}
+	}
+
+	int32 FMLDeformerEditorModel::GetNumCurvesOnSkeletalMesh(USkeletalMesh* SkelMesh) const
+	{
+		if (SkelMesh)
+		{
+			const USkeleton* Skeleton = SkelMesh->GetSkeleton();
+			if (Skeleton)
+			{
+				const FSmartNameMapping* Mapping = Skeleton->GetSmartNameContainer(USkeleton::AnimCurveMappingName);
+				if (Mapping)
+				{
+					TArray<FName> SkeletonCurveNames;
+					Mapping->FillNameArray(SkeletonCurveNames);
+					return SkeletonCurveNames.Num();
+				}
+			}
+		}
+
+		return 0;
+	}
+
+	FString FMLDeformerEditorModel::GetHeatMapMaterialPath() const
+	{
+		return FString(TEXT("/MLDeformerFramework/Materials/MLDeformerHeatMapMat.MLDeformerHeatMapMat"));
+	}
+
+	FString FMLDeformerEditorModel::GetHeatMapDeformerGraphPath() const
+	{
+		return FString(TEXT("/MLDeformerFramework/Deformers/DG_MLDeformerModel_HeatMap.DG_MLDeformerModel_HeatMap"));
+	}
+
+	FString FMLDeformerEditorModel::GetDefaultDeformerGraphAssetPath() const
+	{
+		return FString(TEXT("/MLDeformerFramework/Deformers/DG_MLDeformerModel.DG_MLDeformerModel"));
+	}
+
+	const UAnimSequence* FMLDeformerEditorModel::GetAnimSequence() const
+	{
+		if (GetModel())
+		{
+			const UMLDeformerVizSettings* VizSettings = GetModel()->GetVizSettings();
+			if (VizSettings->GetVisualizationMode() == EMLDeformerVizMode::TrainingData)
+			{
+				return GetModel()->GetAnimSequence();
+			}
+			else if (VizSettings->GetVisualizationMode() == EMLDeformerVizMode::TestData)
+			{
+				return GetModel()->GetVizSettings()->GetTestAnimSequence();
+			}
+		}
+		return nullptr;
+	}
+
+	void FMLDeformerEditorModel::UpdateRanges()
+	{
+		ViewRange = TRange<double>(0.0, 100.0);
+		WorkingRange = TRange<double>(0.0, 100.0);
+		PlaybackRange = TRange<double>(0.0, 100.0);
+		const UAnimSequence* AnimSeq = GetAnimSequence();
+		const double Duration = AnimSeq ? AnimSeq->GetPlayLength() : 0.0;
+		ViewRange = TRange<double>(0.0, Duration);
+		WorkingRange = ViewRange;
+		PlaybackRange = WorkingRange;
+		ClampCurrentTrainingFrameIndex();
+	}
+
+	void FMLDeformerEditorModel::HandleModelChanged()
+	{
+		UpdateRanges();
+	}
+
+	void FMLDeformerEditorModel::HandleVizModeChanged(EMLDeformerVizMode Mode)
+	{
+		UpdateRanges();
+	}
+
+	/** Get the current view range */
+	TRange<double> FMLDeformerEditorModel::GetViewRange() const
+	{
+		return ViewRange;
+	}
+
+	void FMLDeformerEditorModel::SetViewRange(TRange<double> InRange)
+	{
+		ViewRange = InRange;
+
+		if (WorkingRange.HasLowerBound() && WorkingRange.HasUpperBound())
+		{
+			WorkingRange = TRange<double>::Hull(WorkingRange, ViewRange);
+		}
+		else
+		{
+			WorkingRange = ViewRange;
+		}
+	}
+
+	TRange<double> FMLDeformerEditorModel::GetWorkingRange() const
+	{
+		return WorkingRange;
+	}
+
+	TRange<FFrameNumber> FMLDeformerEditorModel::GetPlaybackRange() const
+	{
+		const int32 Resolution = GetTickResolution();
+		return TRange<FFrameNumber>(FFrameNumber(FMath::RoundToInt32(PlaybackRange.GetLowerBoundValue() * (double)Resolution)), FFrameNumber(FMath::RoundToInt32(PlaybackRange.GetUpperBoundValue() * (double)Resolution)));
+
+	}
+
+	int32 FMLDeformerEditorModel::GetTicksPerFrame() const
+	{
+		// Default to millisecond resolution
+		const int32 TimelineScrubSnapValue = 1000;
+		return TimelineScrubSnapValue;
+	}
+
+	FFrameNumber FMLDeformerEditorModel::GetTickResScrubPosition() const
+	{
+		const double Resolution = GetTickResolution();
+		const UMLDeformerVizSettings* VizSettings = GetModel()->GetVizSettings();
+		double TrainingFrameNumber = 0.0;
+		if (VizSettings->GetVisualizationMode() == EMLDeformerVizMode::TrainingData)
+		{
+			TrainingFrameNumber = GetTrainingFrameAtTime(CalcTimelinePosition());
+		}
+		else if (VizSettings->GetVisualizationMode() == EMLDeformerVizMode::TestData)
+		{
+			TrainingFrameNumber = GetTestFrameAtTime(CalcTimelinePosition());
+		}
+		FFrameNumber FrameNumber = (int32)(TrainingFrameNumber * Resolution / (double)GetFrameRate());
+		return FrameNumber;
+	}
+
+	float FMLDeformerEditorModel::GetScrubTime() const
+	{
+		return (float)ScrubPosition.AsDecimal() / (float)GetFrameRate();
+	}
+
+	void FMLDeformerEditorModel::SetScrubPosition(FFrameTime NewScrubPostion)
+	{
+		const double FrameNumber = NewScrubPostion.AsDecimal() * (double)GetFrameRate() / (double)GetTickResolution();
+		ScrubPosition.FrameNumber = (int32)FrameNumber;
+
+		const UMLDeformerVizSettings* VizSettings = GetModel()->GetVizSettings();
+		if (VizSettings->GetVisualizationMode() == EMLDeformerVizMode::TrainingData)
+		{
+			SetTrainingFrame(ScrubPosition.FrameNumber.Value);
+		}
+		else if (VizSettings->GetVisualizationMode() == EMLDeformerVizMode::TestData)
+		{
+			SetTestFrame(ScrubPosition.FrameNumber.Value);
+		}
+	}
+
+	void FMLDeformerEditorModel::SetScrubPosition(FFrameNumber NewScrubPostion)
+	{
+		const double Resolution = GetTickResolution();
+		const double ClampedValue = FMath::Clamp((double)NewScrubPostion.Value, PlaybackRange.GetLowerBoundValue() * Resolution, PlaybackRange.GetUpperBoundValue() * Resolution);
+		const int32 FrameValue = ClampedValue * (double)GetFrameRate() / Resolution;
+
+		const UMLDeformerVizSettings* VizSettings = GetModel()->GetVizSettings();
+		if (VizSettings->GetVisualizationMode() == EMLDeformerVizMode::TrainingData)
+		{
+			SetTrainingFrame(FrameValue);
+		}
+		else if (VizSettings->GetVisualizationMode() == EMLDeformerVizMode::TestData)
+		{
+			SetTestFrame(FrameValue);
+		}
+	}
+
+	void FMLDeformerEditorModel::SetDisplayFrames(bool bInDisplayFrames)
+	{
+		bDisplayFrames = bInDisplayFrames;
+	}
+
+	bool FMLDeformerEditorModel::IsDisplayingFrames() const
+	{
+		return bDisplayFrames;
+	}
+
+	void FMLDeformerEditorModel::HandleViewRangeChanged(TRange<double> InRange)
+	{
+		SetViewRange(InRange);
+	}
+
+	void FMLDeformerEditorModel::HandleWorkingRangeChanged(TRange<double> InRange)
+	{
+		WorkingRange = InRange;
+	}
+
+	double FMLDeformerEditorModel::CalcTimelinePosition() const
+	{
+		const UMLDeformerVizSettings* VizSettings = GetModel()->GetVizSettings();
+		if (VizSettings->GetVisualizationMode() == EMLDeformerVizMode::TrainingData)
+		{
+			return CalcTrainingTimelinePosition();
+		}
+		else if (VizSettings->GetVisualizationMode() == EMLDeformerVizMode::TestData)
+		{
+			return CalcTestTimelinePosition();
+		}
+		return 0.0;
+	}
+
+	double FMLDeformerEditorModel::GetFrameRate() const
+	{
+		const UAnimSequence* AnimSequence = GetAnimSequence();
+		if (AnimSequence)
+		{
+			return AnimSequence->GetSamplingFrameRate().AsDecimal();
+		}
+		else
+		{
+			return 30.0;
+		}
+	}
+
+	int32 FMLDeformerEditorModel::GetTickResolution() const
+	{
+		return FMath::RoundToInt32((double)GetTicksPerFrame() * GetFrameRate());
+	}
+
+	void FMLDeformerEditorModel::OnPostEditChangeProperty(FPropertyChangedEvent& PropertyChangedEvent)
+	{
+		const FProperty* Property = PropertyChangedEvent.Property;
+		if (Property == nullptr)
+		{
+			return;
+		}
+
+		OnPropertyChanged(PropertyChangedEvent);
+	}
+
+	FString FMLDeformerEditorModel::GetTrainedNetworkOnnxFile() const
+	{
+		const FString ClassName = Model->GetClass()->GetFName().ToString();
+
+		// This returns something like:
+		// "<intermediatedir>/YourModel/YourModel.onnx"
+		return FString::Printf(TEXT("%s%s/%s.onnx"), *FPaths::ProjectIntermediateDir(), *ClassName, *ClassName);
+	}
+
+}	// namespace UE::MLDeformer
+
+#undef LOCTEXT_NAMESPACE
diff --git a/Engine/Plugins/Animation/MLDeformer/MLDeformerFramework/Source/MLDeformerFrameworkEditor/Private/MLDeformerEditorModule.cpp b/Engine/Plugins/Animation/MLDeformer/MLDeformerFramework/Source/MLDeformerFrameworkEditor/Private/MLDeformerEditorModule.cpp
new file mode 100644
--- /dev/null
+++ b/Engine/Plugins/Animation/MLDeformer/MLDeformerFramework/Source/MLDeformerFrameworkEditor/Private/MLDeformerEditorModule.cpp
@@ -0,0 +1,54 @@
+// Copyright Epic Games, Inc. All Rights Reserved.
+
+#include "MLDeformerEditorModule.h"
+#include "MLDeformerAssetActions.h"
+#include "MLDeformerEditorMode.h"
+#include "MLDeformerModule.h"
+#include "MLDeformerEditorModel.h"
+#include "MLDeformerModel.h"
+#include "MLDeformerCurveReferenceCustomization.h"
+#include "Modules/ModuleManager.h"
+#include "EditorModeRegistry.h"
+#include "PropertyEditorDelegates.h"
+
+#define LOCTEXT_NAMESPACE "MLDeformerEditorModule"
+
+IMPLEMENT_MODULE(UE::MLDeformer::FMLDeformerEditorModule, MLDeformerEditor)
+
+namespace UE::MLDeformer
+{
+	void FMLDeformerEditorModule::StartupModule()
+	{
+		MLDeformerAssetActions = MakeShareable(new FMLDeformerAssetActions);
+		FModuleManager::LoadModuleChecked<FAssetToolsModule>("AssetTools").Get().RegisterAssetTypeActions(MLDeformerAssetActions.ToSharedRef());
+		FEditorModeRegistry::Get().RegisterMode<FMLDeformerEditorMode>(FMLDeformerEditorMode::ModeName, LOCTEXT("MLDeformerEditorMode", "MLDeformer"), FSlateIcon(), false);
+
+		// Register object detail customizations.
+		FPropertyEditorModule& PropertyModule = FModuleManager::GetModuleChecked<FPropertyEditorModule>("PropertyEditor");
+		PropertyModule.RegisterCustomPropertyTypeLayout("CurveReference", FOnGetPropertyTypeCustomizationInstance::CreateStatic(&FMLDeformerCurveReferenceCustomization::MakeInstance) );
+		PropertyModule.NotifyCustomizationModuleChanged();
+	}
+
+	void FMLDeformerEditorModule::ShutdownModule()
+	{
+		FEditorModeRegistry::Get().UnregisterMode(FMLDeformerEditorMode::ModeName);
+		if (MLDeformerAssetActions.IsValid())
+		{
+			if (FModuleManager::Get().IsModuleLoaded("AssetTools"))
+			{
+				FModuleManager::GetModuleChecked<FAssetToolsModule>("AssetTools").Get().UnregisterAssetTypeActions(MLDeformerAssetActions.ToSharedRef());
+			}
+			MLDeformerAssetActions.Reset();
+		}
+
+		// Unregister object detail customizations.
+		if (FModuleManager::Get().IsModuleLoaded(TEXT("PropertyEditor")))
+		{
+			FPropertyEditorModule& PropertyModule = FModuleManager::GetModuleChecked<FPropertyEditorModule>(TEXT("PropertyEditor"));
+			PropertyModule.UnregisterCustomPropertyTypeLayout(TEXT("CurveReference"));
+			PropertyModule.NotifyCustomizationModuleChanged();
+		}
+	}
+}	// namespace UE::MLDeformer
+
+#undef LOCTEXT_NAMESPACE
diff --git a/Engine/Plugins/Animation/MLDeformer/MLDeformerFramework/Source/MLDeformerFrameworkEditor/Private/MLDeformerModelDetails.cpp b/Engine/Plugins/Animation/MLDeformer/MLDeformerFramework/Source/MLDeformerFrameworkEditor/Private/MLDeformerModelDetails.cpp
new file mode 100644
--- /dev/null
+++ b/Engine/Plugins/Animation/MLDeformer/MLDeformerFramework/Source/MLDeformerFrameworkEditor/Private/MLDeformerModelDetails.cpp
@@ -0,0 +1,380 @@
+// Copyright Epic Games, Inc. All Rights Reserved.
+
+#include "MLDeformerModelDetails.h"
+#include "MLDeformerModule.h"
+#include "MLDeformerAsset.h"
+#include "MLDeformerModel.h"
+#include "MLDeformerInputInfo.h"
+#include "MLDeformerEditorModel.h"
+#include "MLDeformerEditorModule.h"
+#include "MLDeformerGeomCacheHelpers.h"
+#include "NeuralNetwork.h"
+#include "GeometryCache.h"
+#include "GeometryCacheTrack.h"
+#include "Animation/AnimSequence.h"
+#include "Engine/SkeletalMesh.h"
+#include "PropertyCustomizationHelpers.h"
+#include "AssetRegistry/AssetData.h"
+#include "Modules/ModuleManager.h"
+#include "DetailLayoutBuilder.h"
+#include "DetailCategoryBuilder.h"
+#include "DetailWidgetRow.h"
+#include "IDetailsView.h"
+#include "IDetailGroup.h"
+#include "Widgets/DeclarativeSyntaxSupport.h"
+#include "Widgets/Layout/SBox.h"
+#include "Widgets/Input/SButton.h"
+#include "DetailLayoutBuilder.h"
+#include "SWarningOrErrorBox.h"
+
+#define LOCTEXT_NAMESPACE "MLDeformerModelDetails"
+
+namespace UE::MLDeformer
+{
+	bool FMLDeformerModelDetails::UpdateMemberPointers(const TArray<TWeakObjectPtr<UObject>>& Objects)
+	{
+		Model = nullptr;
+		EditorModel = nullptr;
+		if (Objects.Num() == 1 && Objects[0] != nullptr)
+		{
+			Model = static_cast<UMLDeformerModel*>(Objects[0].Get());
+
+			// Get the editor model for this runtime model.
+			FMLDeformerEditorModule& EditorModule = FModuleManager::GetModuleChecked<FMLDeformerEditorModule>("MLDeformerFrameworkEditor");
+			if (Model)
+			{
+				EditorModel = EditorModule.GetModelRegistry().GetEditorModel(Model);
+			}
+		}
+
+		return (Model != nullptr && EditorModel != nullptr);
+	}
+
+	void FMLDeformerModelDetails::CreateCategories()
+	{
+		BaseMeshCategoryBuilder = &DetailLayoutBuilder->EditCategory("Base Mesh", FText::GetEmpty(), ECategoryPriority::Important);
+		TargetMeshCategoryBuilder = &DetailLayoutBuilder->EditCategory("Target Mesh", FText::GetEmpty(), ECategoryPriority::Important);
+		InputOutputCategoryBuilder = &DetailLayoutBuilder->EditCategory("Inputs and Output", FText::GetEmpty(), ECategoryPriority::Important);
+		TrainingSettingsCategoryBuilder = &DetailLayoutBuilder->EditCategory("Training Settings", FText::GetEmpty(), ECategoryPriority::Important);
+	}
+
+	void FMLDeformerModelDetails::CustomizeDetails(IDetailLayoutBuilder& DetailBuilder)
+	{
+		DetailLayoutBuilder = &DetailBuilder;
+
+		// Update the pointers and check if they are valid.
+		TArray<TWeakObjectPtr<UObject>> Objects;
+		DetailBuilder.GetObjectsBeingCustomized(Objects);
+		if (!UpdateMemberPointers(Objects))
+		{
+			return;
+		}
+
+		CreateCategories();
+
+		// Base mesh details.
+		BaseMeshCategoryBuilder->AddProperty(UMLDeformerModel::GetSkeletalMeshPropertyName(), UMLDeformerModel::StaticClass());
+
+		AddBaseMeshErrors();
+
+		// Check if the vertex counts of our asset have changed.
+		const FText ChangedErrorText = EditorModel->GetBaseAssetChangedErrorText();
+		FDetailWidgetRow& ChangedErrorRow = BaseMeshCategoryBuilder->AddCustomRow(FText::FromString("BaseMeshChangedError"))
+			.Visibility(!ChangedErrorText.IsEmpty() ? EVisibility::Visible : EVisibility::Collapsed)
+			.WholeRowContent()
+			[
+				SNew(SBox)
+				.Padding(FMargin(0.0f, 4.0f))
+				[
+					SNew(SWarningOrErrorBox)
+					.MessageStyle(EMessageStyle::Error)
+					.Message(ChangedErrorText)
+				]
+			];
+
+		// Check if our skeletal mesh's imported model contains a list of mesh infos. If not, we need to reimport it as it is an older asset.
+		const FText NeedsReimportErrorText = EditorModel->GetSkeletalMeshNeedsReimportErrorText();
+		FDetailWidgetRow& NeedsReimportErrorRow = BaseMeshCategoryBuilder->AddCustomRow(FText::FromString("BaseMeshNeedsReimportError"))
+			.Visibility(!NeedsReimportErrorText.IsEmpty() ? EVisibility::Visible : EVisibility::Collapsed)
+			.WholeRowContent()
+			[
+				SNew(SBox)
+				.Padding(FMargin(0.0f, 4.0f))
+				[
+					SNew(SWarningOrErrorBox)
+					.MessageStyle(EMessageStyle::Error)
+					.Message(NeedsReimportErrorText)
+				]
+			];
+
+		const FText VertexMapMisMatchErrorText = EditorModel->GetVertexMapChangedErrorText();
+		FDetailWidgetRow& VertexMapErrorRow = BaseMeshCategoryBuilder->AddCustomRow(FText::FromString("VertexMapError"))
+			.Visibility(!VertexMapMisMatchErrorText.IsEmpty() ? EVisibility::Visible : EVisibility::Collapsed)
+			.WholeRowContent()
+			[
+				SNew(SBox)
+				.Padding(FMargin(0.0f, 4.0f))
+				[
+					SNew(SWarningOrErrorBox)
+					.MessageStyle(EMessageStyle::Error)
+					.Message(VertexMapMisMatchErrorText)
+				]
+			];
+
+		// Animation sequence.
+		IDetailPropertyRow& AnimRow = BaseMeshCategoryBuilder->AddProperty(UMLDeformerModel::GetAnimSequencePropertyName(), UMLDeformerModel::StaticClass());
+		AnimRow.CustomWidget()
+		.NameContent()
+		[
+			AnimRow.GetPropertyHandle()->CreatePropertyNameWidget()
+		]
+		.ValueContent()
+		[
+			SNew(SObjectPropertyEntryBox)
+			.PropertyHandle(AnimRow.GetPropertyHandle())
+			.AllowedClass(UAnimSequence::StaticClass())
+			.ObjectPath(Model ? Model->GetAnimSequence()->GetPathName() : FString())
+			.ThumbnailPool(DetailBuilder.GetThumbnailPool())
+			.OnShouldFilterAsset(
+				this, 
+				&FMLDeformerModelDetails::FilterAnimSequences, 
+				Model->GetSkeletalMesh() ? Model->GetSkeletalMesh()->GetSkeleton() : nullptr
+			)
+		];
+
+		AddAnimSequenceErrors();
+
+		const FText AnimErrorText = EditorModel->GetIncompatibleSkeletonErrorText(Model->GetSkeletalMesh(), Model->GetAnimSequence());
+		FDetailWidgetRow& AnimErrorRow = BaseMeshCategoryBuilder->AddCustomRow(FText::FromString("AnimSkeletonMisMatchError"))
+			.Visibility(!AnimErrorText.IsEmpty() ? EVisibility::Visible : EVisibility::Collapsed)
+			.WholeRowContent()
+			[
+				SNew(SBox)
+				.Padding(FMargin(0.0f, 4.0f))
+				[
+					SNew(SWarningOrErrorBox)
+					.MessageStyle(EMessageStyle::Warning)
+					.Message(AnimErrorText)
+				]
+			];
+
+
+		AddTargetMesh();
+
+		TargetMeshCategoryBuilder->AddProperty(UMLDeformerModel::GetAlignmentTransformPropertyName(), UMLDeformerModel::StaticClass());
+
+		InputOutputCategoryBuilder->AddProperty(UMLDeformerModel::GetShouldIncludeBonesPropertyName(), UMLDeformerModel::StaticClass())
+			.Visibility(IsBonesFlagVisible() ? EVisibility::Visible : EVisibility::Collapsed);
+
+		InputOutputCategoryBuilder->AddProperty(UMLDeformerModel::GetShouldIncludeCurvesPropertyName(), UMLDeformerModel::StaticClass())
+			.Visibility(IsCurvesFlagVisible() ? EVisibility::Visible : EVisibility::Collapsed);
+
+		AddTrainingInputFlags();
+		AddTrainingInputErrors();
+
+		const FText ErrorText = EditorModel->GetInputsErrorText();
+		FDetailWidgetRow& ErrorRow = InputOutputCategoryBuilder->AddCustomRow(FText::FromString("InputsError"))
+			.Visibility(!ErrorText.IsEmpty() ? EVisibility::Visible : EVisibility::Collapsed)
+			.WholeRowContent()
+			[
+				SNew(SBox)
+				.Padding(FMargin(0.0f, 4.0f))
+				[
+					SNew(SWarningOrErrorBox)
+					.MessageStyle(EMessageStyle::Error)
+					.Message(ErrorText)
+				]
+			];
+
+		InputOutputCategoryBuilder->AddProperty(UMLDeformerModel::GetMaxTrainingFramesPropertyName(), UMLDeformerModel::StaticClass());
+		InputOutputCategoryBuilder->AddProperty(UMLDeformerModel::GetDeltaCutoffLengthPropertyName(), UMLDeformerModel::StaticClass());
+
+		// Bone include list group.
+		if (Model->DoesSupportBones())
+		{
+			const int NumBones = EditorModel->GetEditorInputInfo() ? EditorModel->GetEditorInputInfo()->GetNumBones() : 0;
+			const FText BonesText = FText::Format(FTextFormat(LOCTEXT("BonesGroupName", "Bones ({0})")), NumBones);;
+			FText AllBonesText;
+			if (Model->GetSkeletalMesh() && (Model->GetSkeletalMesh()->GetRefSkeleton().GetNum() == NumBones || NumBones == 0))
+			{
+				AllBonesText = LOCTEXT("BoneGroupValue", "All Bones Included");
+			}
+			IDetailGroup& BoneIncludeGroup = InputOutputCategoryBuilder->AddGroup("BoneIncludeGroup", FText(), false, false);
+			BoneIncludeGroup.HeaderRow()
+			.NameContent()
+			[
+				SNew(SHorizontalBox)
+				+ SHorizontalBox::Slot()
+				.HAlign(HAlign_Right)
+				[
+					SNew(STextBlock)
+					.Text(BonesText)
+					.Font(IDetailLayoutBuilder::GetDetailFont())
+				]
+			]
+			.ValueContent()
+			[
+				SNew(SHorizontalBox)
+				+ SHorizontalBox::Slot()
+				.HAlign(HAlign_Left)
+				[
+					SNew(STextBlock)
+					.Text(AllBonesText)
+					.Font(IDetailLayoutBuilder::GetDetailFontItalic())
+				]
+			];
+			BoneIncludeGroup.AddWidgetRow()
+				.ValueContent()
+				[
+					SNew(SButton)
+					.HAlign(HAlign_Center)
+					.VAlign(VAlign_Center)
+					.Text(LOCTEXT("AnimatedBonesButton", "Animated Bones Only"))
+					.ToolTipText(LOCTEXT("AnimatedBonesButtonTooltip", "Initialize the bone include list to bones that are animated."))
+					.OnClicked(FOnClicked::CreateSP(this, &FMLDeformerModelDetails::OnFilterAnimatedBonesOnly))
+					.IsEnabled_Lambda([this](){ return Model->ShouldIncludeBonesInTraining(); })
+				];
+			BoneIncludeGroup.AddPropertyRow(DetailBuilder.GetProperty(UMLDeformerModel::GetBoneIncludeListPropertyName(), UMLDeformerModel::StaticClass()));
+
+			AddBoneInputErrors();
+		}
+		else
+		{
+			InputOutputCategoryBuilder->AddProperty(UMLDeformerModel::GetBoneIncludeListPropertyName(), UMLDeformerModel::StaticClass())
+				.Visibility(EVisibility::Collapsed);
+		}
+
+		// Curve include list group.
+		if (Model->DoesSupportCurves())
+		{
+			const int NumCurves = EditorModel->GetEditorInputInfo() ? EditorModel->GetEditorInputInfo()->GetNumCurves() : 0;
+			const FText CurvesText = FText::Format(FTextFormat(LOCTEXT("CurvesGroupName", "Curves ({0})")), NumCurves);
+			FText AllCurvesText;
+			const int32 NumCurvesOnSkelMesh = EditorModel->GetNumCurvesOnSkeletalMesh(Model->GetSkeletalMesh());
+			if (NumCurvesOnSkelMesh == 0)
+			{
+				AllCurvesText = LOCTEXT("CurvesGroupValueNoCurves", "No Curves Found");
+			}
+			else if (NumCurves == NumCurvesOnSkelMesh)
+			{
+				AllCurvesText = LOCTEXT("CurvesGroupValue", "All Curves Included");
+			}
+			IDetailGroup& CurvesIncludeGroup = InputOutputCategoryBuilder->AddGroup("CurveIncludeGroup", LOCTEXT("CurveIncludeGroup", "Curves"), false, false);
+			CurvesIncludeGroup.HeaderRow()
+			.NameContent()
+			[
+				SNew(SHorizontalBox)
+				+ SHorizontalBox::Slot()
+				.HAlign(HAlign_Right)
+				[
+					SNew(STextBlock)
+					.Text(CurvesText)
+					.Font(IDetailLayoutBuilder::GetDetailFont())
+				]
+			]
+			.ValueContent()
+			[
+				SNew(SHorizontalBox)
+				+ SHorizontalBox::Slot()
+				.HAlign(HAlign_Left)
+				[
+					SNew(STextBlock)
+					.Text(AllCurvesText)
+					.Font(IDetailLayoutBuilder::GetDetailFontItalic())
+				]
+			];
+
+			CurvesIncludeGroup.AddWidgetRow()
+				.ValueContent()
+				[
+					SNew(SButton)
+					.HAlign(HAlign_Center)
+					.VAlign(VAlign_Center)
+					.Text(LOCTEXT("AnimatedCurvesButton", "Animated Curves Only"))
+					.ToolTipText(LOCTEXT("AnimatedCurvesButtonTooltip", "Initialize the curve include list to curves that are animated."))
+					.OnClicked(FOnClicked::CreateSP(this, &FMLDeformerModelDetails::OnFilterAnimatedCurvesOnly))
+					.IsEnabled_Lambda([this](){ return Model->ShouldIncludeCurvesInTraining() && EditorModel->GetNumCurvesOnSkeletalMesh(Model->GetSkeletalMesh()) > 0; })
+				];
+			CurvesIncludeGroup.AddPropertyRow(DetailBuilder.GetProperty(UMLDeformerModel::GetCurveIncludeListPropertyName(), UMLDeformerModel::StaticClass()));
+
+			AddCurveInputErrors();
+		}
+		else
+		{
+			InputOutputCategoryBuilder->AddProperty(UMLDeformerModel::GetCurveIncludeListPropertyName(), UMLDeformerModel::StaticClass())
+				.Visibility(EVisibility::Collapsed);
+		}
+
+		AddTrainingInputFilters();
+
+		// Show a warning when no neural network has been set.
+		{		
+			UNeuralNetwork* NeuralNetwork = Model->GetNeuralNetwork();
+			FDetailWidgetRow& NeuralNetErrorRow = TrainingSettingsCategoryBuilder->AddCustomRow(FText::FromString("NeuralNetError"))
+				.Visibility((NeuralNetwork == nullptr) ? EVisibility::Visible : EVisibility::Collapsed)
+				.WholeRowContent()
+				[
+					SNew(SBox)
+					.Padding(FMargin(0.0f, 4.0f))
+					[
+						SNew(SWarningOrErrorBox)
+						.MessageStyle(EMessageStyle::Warning)
+						.Message(FText::FromString("Model still needs to be trained."))
+					]
+				];
+
+			// Check if our network is compatible with the skeletal mesh.
+			if (Model->GetSkeletalMesh() && NeuralNetwork)
+			{
+				FDetailWidgetRow& NeuralNetIncompatibleErrorRow = TrainingSettingsCategoryBuilder->AddCustomRow(FText::FromString("NeuralNetIncompatibleError"))
+					.Visibility(!Model->GetInputInfo()->IsCompatible(Model->GetSkeletalMesh()) ? EVisibility::Visible : EVisibility::Collapsed)
+					.WholeRowContent()
+					[
+						SNew(SBox)
+						.Padding(FMargin(0.0f, 4.0f))
+						[
+							SNew(SWarningOrErrorBox)
+							.MessageStyle(EMessageStyle::Error)
+							.Message(FText::FromString("Trained neural network is incompatible with selected SkeletalMesh."))
+						]
+					];
+			}
+		}
+	}
+
+	bool FMLDeformerModelDetails::FilterAnimSequences(const FAssetData& AssetData, USkeleton* Skeleton)
+	{
+		if (Skeleton && Skeleton->IsCompatibleSkeletonByAssetData(AssetData))
+		{
+			return false;
+		}
+
+		return true;
+	}
+
+	FReply FMLDeformerModelDetails::OnFilterAnimatedBonesOnly() const
+	{
+		EditorModel->InitBoneIncludeListToAnimatedBonesOnly();
+		DetailLayoutBuilder->ForceRefreshDetails();
+		return FReply::Handled();
+	}
+
+	FReply FMLDeformerModelDetails::OnFilterAnimatedCurvesOnly() const
+	{
+		EditorModel->InitCurveIncludeListToAnimatedCurvesOnly();
+		DetailLayoutBuilder->ForceRefreshDetails();
+		return FReply::Handled();
+	}
+
+	bool FMLDeformerModelDetails::IsBonesFlagVisible() const
+	{
+		return (Model->DoesSupportBones() && Model->DoesSupportCurves());
+	}
+
+	bool FMLDeformerModelDetails::IsCurvesFlagVisible() const
+	{
+		return (Model->DoesSupportBones() && Model->DoesSupportCurves() && EditorModel->GetNumCurvesOnSkeletalMesh(Model->GetSkeletalMesh()) > 0);
+	}
+}	// namespace UE::MLDeformer
+
+#undef LOCTEXT_NAMESPACE
diff --git a/Engine/Plugins/Animation/MLDeformer/MLDeformerFramework/Source/MLDeformerFrameworkEditor/Private/MLDeformerTrainingModel.cpp b/Engine/Plugins/Animation/MLDeformer/MLDeformerFramework/Source/MLDeformerFrameworkEditor/Private/MLDeformerTrainingModel.cpp
new file mode 100644
--- /dev/null
+++ b/Engine/Plugins/Animation/MLDeformer/MLDeformerFramework/Source/MLDeformerFrameworkEditor/Private/MLDeformerTrainingModel.cpp
@@ -0,0 +1,86 @@
+// Copyright Epic Games, Inc. All Rights Reserved.
+
+#include "MLDeformerTrainingModel.h"
+#include "MLDeformerEditorModel.h"
+#include "MLDeformerModel.h"
+#include "MLDeformerInputInfo.h"
+#include "MLDeformerSampler.h"
+
+UMLDeformerModel* UMLDeformerTrainingModel::GetModel() const
+{
+	return EditorModel->GetModel();
+}
+
+void UMLDeformerTrainingModel::Init(UE::MLDeformer::FMLDeformerEditorModel* InEditorModel)
+{
+	EditorModel = InEditorModel;
+}
+
+void UMLDeformerTrainingModel::SetEditorModel(UE::MLDeformer::FMLDeformerEditorModel* InModel)
+{ 
+	EditorModel = InModel;
+}
+
+UE::MLDeformer::FMLDeformerEditorModel* UMLDeformerTrainingModel::GetEditorModel() const
+{
+	return EditorModel;
+}
+
+int32 UMLDeformerTrainingModel::GetNumberSampleTransforms() const
+{
+	return GetModel()->GetInputInfo()->GetNumBones();
+}
+
+int32 UMLDeformerTrainingModel::GetNumberSampleCurves() const
+{
+	return GetModel()->GetInputInfo()->GetNumCurves();
+}
+
+int32 UMLDeformerTrainingModel::NumSamples() const
+{
+	return EditorModel->GetNumFramesForTraining();
+}
+
+int32 UMLDeformerTrainingModel::GetNumberSampleDeltas() const
+{
+	return GetModel()->GetNumBaseMeshVerts();
+}
+
+bool UMLDeformerTrainingModel::SetCurrentSampleIndex(int32 Index)
+{
+	return SampleFrame(Index);
+}
+
+bool UMLDeformerTrainingModel::GetNeedsResampling() const
+{
+	return EditorModel->GetResamplingInputOutputsNeeded();
+}
+
+void UMLDeformerTrainingModel::SetNeedsResampling(bool bNeedsResampling)
+{
+	EditorModel->SetResamplingInputOutputsNeeded(bNeedsResampling);
+}
+
+bool UMLDeformerTrainingModel::SampleFrame(int32 Index)
+{
+	using namespace UE::MLDeformer;
+
+	// Make sure we have a valid frame number.
+	if (Index < 0 || Index >= EditorModel->GetNumTrainingFrames())
+	{
+		UE_LOG(LogMLDeformer, Warning, TEXT("Sample index must range from %d to %d, but a value of %d was provided."), 0, EditorModel->GetNumTrainingFrames()-1, Index);
+		return false;
+	}
+
+	// Sample the frame.
+	FMLDeformerSampler* Sampler = EditorModel->GetSampler();
+	Sampler->SetVertexDeltaSpace(EVertexDeltaSpace::PreSkinning);
+	Sampler->Sample(Index);
+
+	// Copy sampled values.
+	SampleDeltas = Sampler->GetVertexDeltas();
+	SampleBoneRotations = Sampler->GetBoneRotations();
+	SampleCurveValues = Sampler->GetCurveValues();
+
+	return true;
+}
diff --git a/Engine/Plugins/Animation/MLDeformer/MLDeformerFramework/Source/MLDeformerFrameworkEditor/Public/MLDeformerCurveReferenceCustomization.h b/Engine/Plugins/Animation/MLDeformer/MLDeformerFramework/Source/MLDeformerFrameworkEditor/Public/MLDeformerCurveReferenceCustomization.h
new file mode 100644
--- /dev/null
+++ b/Engine/Plugins/Animation/MLDeformer/MLDeformerFramework/Source/MLDeformerFrameworkEditor/Public/MLDeformerCurveReferenceCustomization.h
@@ -0,0 +1,44 @@
+// Copyright Epic Games, Inc. All Rights Reserved.
+
+#pragma once
+
+#include "CoreMinimal.h"
+#include "IPropertyTypeCustomization.h"
+#include "PropertyHandle.h"
+#include "DetailWidgetRow.h"
+
+class USkeleton;
+
+namespace UE::MLDeformer
+{
+	/**
+	 * The curve reference property detail customization.
+	 */
+	class MLDEFORMERFRAMEWORKEDITOR_API FMLDeformerCurveReferenceCustomization
+		: public IPropertyTypeCustomization
+	{
+	public:
+		static TSharedRef<IPropertyTypeCustomization> MakeInstance();
+
+		// IPropertyTypeCustomization overrides.
+		virtual void CustomizeHeader(TSharedRef<IPropertyHandle> StructPropertyHandle, class FDetailWidgetRow& HeaderRow, IPropertyTypeCustomizationUtils& StructCustomizationUtils) override;
+		virtual void CustomizeChildren(TSharedRef<IPropertyHandle> StructPropertyHandle, class IDetailChildrenBuilder& ChildBuilder, IPropertyTypeCustomizationUtils& StructCustomizationUtils) override {}
+		// ~END IPropertyTypeCustomization overrides.
+
+	protected:
+		void SetSkeleton(TSharedRef<IPropertyHandle> StructPropertyHandle);
+		virtual void SetPropertyHandle(TSharedRef<IPropertyHandle> StructPropertyHandle);
+		TSharedPtr<IPropertyHandle> FindStructMemberProperty(TSharedRef<IPropertyHandle> PropertyHandle, const FName& PropertyName);
+
+		/** Property to change after curve has been picked. */
+		TSharedPtr<IPropertyHandle> CurveNameProperty = nullptr;
+
+		/** The Skeleton we get the curves from. */
+		TObjectPtr<USkeleton> Skeleton = nullptr;
+
+	private:
+		virtual void OnCurveSelectionChanged(const FString& Name);
+		virtual FString OnGetSelectedCurve() const;
+		virtual USkeleton* OnGetSkeleton() const;
+	};
+}	// namespace UE::MLDeformer
diff --git a/Engine/Plugins/Animation/MLDeformer/MLDeformerFramework/Source/MLDeformerFrameworkEditor/Public/MLDeformerEditorActor.h b/Engine/Plugins/Animation/MLDeformer/MLDeformerFramework/Source/MLDeformerFrameworkEditor/Public/MLDeformerEditorActor.h
new file mode 100644
--- /dev/null
+++ b/Engine/Plugins/Animation/MLDeformer/MLDeformerFramework/Source/MLDeformerFrameworkEditor/Public/MLDeformerEditorActor.h
@@ -0,0 +1,126 @@
+// Copyright Epic Games, Inc. All Rights Reserved.
+
+#pragma once
+
+#include "CoreMinimal.h"
+#include "GameFramework/Actor.h"
+
+class UDebugSkelMeshComponent;
+class UTextRenderComponent;
+class UMLDeformerComponent;
+
+namespace UE::MLDeformer
+{
+	/**
+	 * The default ID's for editor actors.
+	 * You can assign any ID to the editor actors, although some code in the editor model classes will assume
+	 * that some of these exist. You can always overload certain methods to change that behavior though.
+	 * @see FMLDeformerEditorActor
+	 */
+	enum : int32
+	{
+		/** The linear skinned (base) actor in training mode. */
+		ActorID_Train_Base,
+
+		/** The ground truth (target) actor in training mode. This is our training target that we try to approximate. */
+		ActorID_Train_GroundTruth,
+
+		/** The linear skinned (base) actor, in test mode. */
+		ActorID_Test_Base,
+
+		/** The ML Deformed actor, in test mode. This one should have an UMLDeformerComponent on it. */
+		ActorID_Test_MLDeformed,
+
+		/** The ground truth test actor, in test mode. This will represent your test anim sequence ground truth. */
+		ActorID_Test_GroundTruth
+	};
+
+	/**
+	 * An actor in the ML Deformer asset editor viewport.
+	 * The UE Actor object itself is created by the FMLDeformerEditorModel class, after which this editor actor class
+	 * can create components on it.
+	 * It also contains some methods to control visibility and playback, among some other helper methods.
+	 */
+	class MLDEFORMERFRAMEWORKEDITOR_API FMLDeformerEditorActor
+	{
+	public:
+		/**
+		 * The construction settings, which can be used as info when creating new components.
+		 */
+		struct MLDEFORMERFRAMEWORKEDITOR_API FConstructSettings
+		{
+			/** The UE actor object that we should add components to. */
+			AActor* Actor = nullptr;
+
+			/** The ID of the actor. Look at the UE::MLDeformer::ActorID_Train_Base etc for example. */
+			int32 TypeID = -1;
+
+			/** The color of the label text. */
+			FLinearColor LabelColor = FLinearColor(1.0f, 0.0f, 0.0f);
+
+			/** The label text, which should be something like "Linear Skinned" or "Ground Truth", basically describing the actor. This will appear inside the render viewport. */
+			FText LabelText;
+
+			/** Is this an actor used during training? Set to false if it is a testing mode actor. */
+			bool bIsTrainingActor = false;
+		};
+
+		FMLDeformerEditorActor(const FConstructSettings& Settings);
+		virtual ~FMLDeformerEditorActor();
+
+		// Main methods you can override.
+		virtual void SetVisibility(bool bIsVisible);
+		virtual bool IsVisible() const;
+		virtual void SetPlayPosition(float TimeInSeconds, bool bAutoPause = true);
+		virtual float GetPlayPosition() const;
+		virtual void SetPlaySpeed(float PlaySpeed);
+		virtual float GetPlaySpeed();
+
+		virtual FBox GetBoundingBox() const;
+		virtual void Pause(bool bPaused);
+		virtual bool IsPlaying() const;
+		virtual bool IsGroundTruthActor() const;
+		virtual bool HasVisualMesh() const;
+
+		AActor* GetActor() const { return Actor; }
+		int32 GetTypeID() const { return TypeID; }
+
+		UDebugSkelMeshComponent* GetSkeletalMeshComponent() const { return SkeletalMeshComponent.Get(); }
+		UTextRenderComponent* GetLabelComponent() const { return LabelComponent.Get(); }
+		UMLDeformerComponent* GetMLDeformerComponent() const { return MLDeformerComponent.Get(); }
+		void SetSkeletalMeshComponent(UDebugSkelMeshComponent* SkelMeshComponent) { SkeletalMeshComponent = SkelMeshComponent; }
+		void SetMLDeformerComponent(UMLDeformerComponent* Component) { MLDeformerComponent = Component; }
+		void SetMeshOffsetFactor(float OffsetFactor) { MeshOffsetFactor = OffsetFactor; }
+		bool IsTrainingActor() const { return bIsTrainingActor; }
+		bool IsTestActor() const { return !bIsTrainingActor; }
+		float GetMeshOffsetFactor() const { return MeshOffsetFactor; }
+
+	protected:
+		UTextRenderComponent* CreateLabelComponent(AActor* InActor, FLinearColor Color, const FText& Text) const;
+
+	protected:
+		/**
+		 * The ID of the editor actor type.
+		 * This can be used to identify what actor we are dealing with, for example the base actor or ML Deformed one etc.
+		 */
+		int32 TypeID = -1;
+
+		/** The label component, which shows above the actor. */
+		TObjectPtr<UTextRenderComponent> LabelComponent = nullptr;
+
+		/** The actual actor pointer. */
+		TObjectPtr<AActor> Actor = nullptr;
+
+		/** The skeletal mesh component (can be nullptr). */
+		TObjectPtr<UDebugSkelMeshComponent> SkeletalMeshComponent = nullptr;
+
+		/** The ML Deformer component (can be nullptr). */
+		TObjectPtr<UMLDeformerComponent> MLDeformerComponent = nullptr;
+
+		/** The position offset factor of the actor. A value of 1.0 would offset the actor with the Mesh Spacing amount, a value of 2.0 would be two times the mesh spacing offset, etc. */
+		float MeshOffsetFactor = 0.0f;
+
+		/** Is this actor used for training? */
+		bool bIsTrainingActor = true;
+	};
+} // namespace UE::MLDeformer
diff --git a/Engine/Plugins/Animation/MLDeformer/MLDeformerFramework/Source/MLDeformerFrameworkEditor/Public/MLDeformerEditorMode.h b/Engine/Plugins/Animation/MLDeformer/MLDeformerFramework/Source/MLDeformerFrameworkEditor/Public/MLDeformerEditorMode.h
new file mode 100644
--- /dev/null
+++ b/Engine/Plugins/Animation/MLDeformer/MLDeformerFramework/Source/MLDeformerFrameworkEditor/Public/MLDeformerEditorMode.h
@@ -0,0 +1,48 @@
+// Copyright Epic Games, Inc. All Rights Reserved.
+
+#pragma once
+
+#include "CoreMinimal.h"
+#include "IPersonaEditMode.h"
+
+class UMLDeformerVizSettings;
+
+namespace UE::MLDeformer
+{
+	class FMLDeformerEditorToolkit;
+	class FMLDeformerPreviewScene;
+
+	/**
+	 * The ML Deformer Persona editor mode.
+	 * This handles ticking and rendering of the render viewport and camera focus.
+	 */
+	class MLDEFORMERFRAMEWORKEDITOR_API FMLDeformerEditorMode
+		: public IPersonaEditMode
+	{
+	public:
+		/** The name of the mode. */
+		static FName ModeName;
+
+		void SetEditorToolkit(FMLDeformerEditorToolkit* InToolkit) { DeformerEditorToolkit = InToolkit; }
+
+		// IPersonaEditMode overrides.
+		virtual bool GetCameraTarget(FSphere& OutTarget) const override;
+		virtual class IPersonaPreviewScene& GetAnimPreviewScene() const override;
+		virtual void GetOnScreenDebugInfo(TArray<FText>& OutDebugInfo) const override {}
+		// ~END IPersonaEditMode overrides.
+
+		// FEdMode overrides.
+		virtual void Tick(FEditorViewportClient* ViewportClient, float DeltaTime) override;
+		virtual void Render(const FSceneView* View, FViewport* Viewport, FPrimitiveDrawInterface* PDI) override;
+		virtual void DrawHUD(FEditorViewportClient* ViewportClient, FViewport* Viewport, const FSceneView* View, FCanvas* Canvas) override;
+		virtual bool IsCompatibleWith(FEditorModeID OtherModeID) const override { return true; }
+		virtual bool AllowWidgetMove() override { return false; }
+		virtual bool ShouldDrawWidget() const override { return false; }
+		virtual bool UsesTransformWidget() const override { return false; }
+		virtual bool UsesTransformWidget(UE::Widget::EWidgetMode CheckMode) const override { return false; }
+		// ~END FEdMode overrides.
+
+	protected:
+		FMLDeformerEditorToolkit* DeformerEditorToolkit = nullptr;
+	};
+}	// namespace UE::MLDeformer
