[MLDeformer] Moved out of Experimental, and marked as beta.

diff --git a/Engine/Plugins/Animation/MLDeformer/MLDeformerFramework/Source/MLDeformerFrameworkEditor/Public/MLDeformerEditorModule.h b/Engine/Plugins/Animation/MLDeformer/MLDeformerFramework/Source/MLDeformerFrameworkEditor/Public/MLDeformerEditorModule.h
new file mode 100644
--- /dev/null
+++ b/Engine/Plugins/Animation/MLDeformer/MLDeformerFramework/Source/MLDeformerFrameworkEditor/Public/MLDeformerEditorModule.h
@@ -0,0 +1,42 @@
+// Copyright Epic Games, Inc. All Rights Reserved.
+
+#pragma once
+
+#include "CoreMinimal.h"
+#include "Modules/ModuleInterface.h"
+#include "MLDeformerModelRegistry.h"
+
+namespace UE::MLDeformer
+{
+	class FMLDeformerAssetActions;
+
+	/**
+	 * The ML Deformer editor module.
+	 * This registers the editor mode and custom property customizations.
+	 * It also contains the model registry, which you will use to register your custom models to the editor.
+	 */
+	class MLDEFORMERFRAMEWORKEDITOR_API FMLDeformerEditorModule
+		: public IModuleInterface
+	{
+	public:
+		// IModuleInterface overrides.
+		void StartupModule() override;
+		void ShutdownModule() override;
+		// ~END IModuleInterface overrides.
+
+		/**
+		 * Get a reference to the ML Deformer model registry.
+		 * You use this to see what kind of model types ther are, register new models, or create editor models for specific types of runtime models.
+		 * @return A reference to the ML Deformer model registry.
+		 */
+		FMLDeformerEditorModelRegistry& GetModelRegistry()				{ return ModelRegistry; }
+		const FMLDeformerEditorModelRegistry& GetModelRegistry() const	{ return ModelRegistry; }
+
+	private:
+		/** The model registry, which keeps track of all model types and instances created of these models, and how to create the editor models for specific runtime model types. */
+		FMLDeformerEditorModelRegistry ModelRegistry;
+
+		/** The asset actions for the ML Deformer asset type. */
+		TSharedPtr<FMLDeformerAssetActions> MLDeformerAssetActions;
+	};
+}	// namespace UE::MLDeformer
diff --git a/Engine/Plugins/Animation/MLDeformer/MLDeformerFramework/Source/MLDeformerFrameworkEditor/Public/MLDeformerEditorStyle.h b/Engine/Plugins/Animation/MLDeformer/MLDeformerFramework/Source/MLDeformerFrameworkEditor/Public/MLDeformerEditorStyle.h
new file mode 100644
--- /dev/null
+++ b/Engine/Plugins/Animation/MLDeformer/MLDeformerFramework/Source/MLDeformerFrameworkEditor/Public/MLDeformerEditorStyle.h
@@ -0,0 +1,21 @@
+// Copyright Epic Games, Inc. All Rights Reserved.
+
+#pragma once
+
+#include "Styling/SlateStyle.h"
+
+namespace UE::MLDeformer
+{
+	/**
+	 * The editor style class that describes specific UI style related settings for the ML Deformer editor.
+	 */
+	class MLDEFORMERFRAMEWORKEDITOR_API FMLDeformerEditorStyle
+		: public FSlateStyleSet
+	{
+	public:
+		FMLDeformerEditorStyle();
+		~FMLDeformerEditorStyle();
+
+		static FMLDeformerEditorStyle& Get();
+	};
+}	// namespace UE::MLDeformer
diff --git a/Engine/Plugins/Animation/MLDeformer/MLDeformerFramework/Source/MLDeformerFrameworkEditor/Public/MLDeformerEditorToolkit.h b/Engine/Plugins/Animation/MLDeformer/MLDeformerFramework/Source/MLDeformerFrameworkEditor/Public/MLDeformerEditorToolkit.h
new file mode 100644
--- /dev/null
+++ b/Engine/Plugins/Animation/MLDeformer/MLDeformerFramework/Source/MLDeformerFrameworkEditor/Public/MLDeformerEditorToolkit.h
@@ -0,0 +1,148 @@
+// Copyright Epic Games, Inc. All Rights Reserved.
+
+#pragma once
+
+#include "CoreMinimal.h"
+#include "MLDeformerTrainingModel.h"
+#include "MLDeformerEditorModel.h"
+#include "MLDeformerVizSettings.h"
+#include "IHasPersonaToolkit.h"
+#include "IPersonaPreviewScene.h"
+#include "IPersonaViewport.h"
+#include "PersonaAssetEditorToolkit.h"
+#include "EditorUndoClient.h"
+#include "Math/Color.h"
+#include "Widgets/Notifications/SNotificationList.h"
+
+class IDetailsView;
+class UMLDeformerAsset;
+class SSimpleTimeSlider;
+
+namespace UE::MLDeformer
+{
+	class SMLDeformerTimeline;
+
+	namespace MLDeformerEditorModes
+	{
+		extern const FName Editor;
+	}
+
+	/**
+	 * The ML Deformer asset editor toolkit.
+	 * This is the editor that opens when you double click an ML Deformer asset.
+	 */
+	class MLDEFORMERFRAMEWORKEDITOR_API FMLDeformerEditorToolkit :
+		public FPersonaAssetEditorToolkit,
+		public IHasPersonaToolkit,
+		public FGCObject,
+		public FEditorUndoClient,
+		public FTickableEditorObject
+	{
+	public:
+		friend class FMLDeformerApplicationMode;
+		friend struct FMLDeformerVizSettingsTabSummoner;
+
+		/** Initialize the asset editor. This will register the application mode, init the preview scene, etc. */
+		void InitAssetEditor(
+			const EToolkitMode::Type Mode,
+			const TSharedPtr<IToolkitHost>& InitToolkitHost,
+			UMLDeformerAsset* InDeformerAsset);
+
+		// FAssetEditorToolkit overrides.
+		virtual void RegisterTabSpawners(const TSharedRef<FTabManager>& InTabManager) override;
+		virtual void UnregisterTabSpawners(const TSharedRef<FTabManager>& InTabManager) override;
+		virtual FName GetToolkitFName() const override;
+		virtual FText GetBaseToolkitName() const override;
+		virtual FText GetToolkitName() const override;
+		virtual FLinearColor GetWorldCentricTabColorScale() const override;
+		virtual FString GetWorldCentricTabPrefix() const override;
+		// ~END FAssetEditorToolkit overrides.
+
+		// FGCObject overrides.
+		virtual FString GetReferencerName() const override { return TEXT("FMLDeformerEditorToolkit"); }
+		virtual void AddReferencedObjects(FReferenceCollector& Collector) override;
+		// ~END FGCObject overrides.
+
+		// FTickableEditorObject overrides.
+		virtual void Tick(float DeltaTime) override {};
+		virtual ETickableTickType GetTickableTickType() const override { return ETickableTickType::Always; }
+		virtual TStatId GetStatId() const override;
+		// ~END FTickableEditorObject overrides.
+
+		// IHasPersonaToolkit overrides.
+		virtual TSharedRef<IPersonaToolkit> GetPersonaToolkit() const override;
+		IPersonaToolkit* GetPersonaToolkitPointer() const { return PersonaToolkit.Get(); }
+		// ~END IHasPersonaToolkit overrides.
+
+		void SetVizSettingsDetailsView(TSharedPtr<IDetailsView> InDetailsView) { VizSettingsDetailsView = InDetailsView; }
+
+		IDetailsView* GetModelDetailsView() const;
+		IDetailsView* GetVizSettingsDetailsView() const;
+
+		void SetTimeSlider(TSharedPtr<SMLDeformerTimeline> InTimeSlider);
+		SMLDeformerTimeline* GetTimeSlider() const;
+
+		UMLDeformerAsset* GetDeformerAsset() const { return DeformerAsset.Get(); }
+		FMLDeformerEditorModel* GetActiveModel() { return ActiveModel.Get(); }
+
+		TWeakPtr<FMLDeformerEditorModel> GetActiveModelPointer(){ return TWeakPtr<FMLDeformerEditorModel>(ActiveModel); }
+		const FMLDeformerEditorModel* GetActiveModel() const { return ActiveModel.Get(); }
+
+		double CalcTimelinePosition() const;
+		void OnTimeSliderScrubPositionChanged(double NewScrubTime, bool bIsScrubbing);
+		void UpdateTimeSliderRange();
+		void SetTimeSliderRange(double StartTime, double EndTime);
+
+	private:
+		/* Toolbar related. */
+		void ExtendToolbar();
+		void FillToolbar(FToolBarBuilder& ToolbarBuilder);
+
+		/** Preview scene setup. */
+		void HandlePreviewSceneCreated(const TSharedRef<IPersonaPreviewScene>& InPersonaPreviewScene);
+		void HandleViewportCreated(const TSharedRef<IPersonaViewport>& InPersonaViewport);
+		void HandleDetailsCreated(const TSharedRef<class IDetailsView>& InDetailsView);
+		void OnFinishedChangingDetails(const FPropertyChangedEvent& PropertyChangedEvent);
+
+		/** Helpers. */
+		void ShowNotification(const FText& Message, SNotificationItem::ECompletionState State, bool PlaySound) const;
+		FText GetOverlayText() const;
+		void OnSwitchedVisualizationMode();
+		bool HandleTrainingResult(ETrainingResult TrainingResult, double TrainingDuration, bool& bOutUsePartiallyTrained);
+		void OnModelChanged(int Index);
+		void OnVizModeChanged(EMLDeformerVizMode Mode);
+		FText GetActiveModelName() const;
+		FText GetCurrentVizModeName() const;
+		FText GetVizModeName(EMLDeformerVizMode Mode) const;
+		void ZoomOnActors();
+		void ShowNoModelsWarningIfNeeded();
+
+		TSharedRef<SWidget> GenerateModelButtonContents(TSharedRef<FUICommandList> InCommandList);
+		TSharedRef<SWidget> GenerateVizModeButtonContents(TSharedRef<FUICommandList> InCommandList);
+
+	private:
+		/** The persona toolkit. */	
+		TSharedPtr<IPersonaToolkit> PersonaToolkit = nullptr;
+
+		/** Model details view. */
+		TSharedPtr<IDetailsView> ModelDetailsView = nullptr;
+
+		/** Model viz settings details view. */
+		TSharedPtr<IDetailsView> VizSettingsDetailsView = nullptr;
+
+		/** The timeline slider widget. */
+		TSharedPtr<SMLDeformerTimeline> TimeSlider = nullptr;
+
+		/** The currently active editor model. */
+		TSharedPtr<FMLDeformerEditorModel> ActiveModel = nullptr;
+
+		// Persona viewport.
+		TSharedPtr<IPersonaViewport> PersonaViewport = nullptr;
+
+		/** The ML Deformer Asset. */
+		TObjectPtr<UMLDeformerAsset> DeformerAsset = nullptr;
+
+		/** Has the asset editor been initialized? */
+		bool bIsInitialized = false;
+	};
+}	// namespace UE::MLDeformer
diff --git a/Engine/Plugins/Animation/MLDeformer/MLDeformerFramework/Source/MLDeformerFrameworkEditor/Public/MLDeformerGeomCacheSampler.h b/Engine/Plugins/Animation/MLDeformer/MLDeformerFramework/Source/MLDeformerFrameworkEditor/Public/MLDeformerGeomCacheSampler.h
new file mode 100644
--- /dev/null
+++ b/Engine/Plugins/Animation/MLDeformer/MLDeformerFramework/Source/MLDeformerFrameworkEditor/Public/MLDeformerGeomCacheSampler.h
@@ -0,0 +1,81 @@
+// Copyright Epic Games, Inc. All Rights Reserved.
+
+#pragma once
+
+#include "CoreMinimal.h"
+#include "MLDeformerSampler.h"
+#include "MLDeformerModel.h"
+#include "MLDeformerGeomCacheHelpers.h"
+#include "UObject/ObjectPtr.h"
+#include "GeometryCacheMeshData.h"
+
+class UGeometryCacheComponent;
+class UGeometryCache;
+
+namespace UE::MLDeformer
+{
+	DECLARE_DELEGATE_RetVal(UGeometryCache*, FMLDeformerGetGeomCacheEvent)
+
+	/**
+	 * The input data sampler, which is used to sample vertex positions from geometry caches.
+	 * It can then also calculate deltas between the sampled skeletal mesh data and geometry cache data.
+	 */
+	class MLDEFORMERFRAMEWORKEDITOR_API FMLDeformerGeomCacheSampler
+		: public FMLDeformerSampler
+	{
+	public:
+		// FVertexDeltaSampler overrides.
+		virtual void Sample(int32 AnimFrameIndex) override;
+		virtual void RegisterTargetComponents() override;
+		virtual float GetTimeAtFrame(int32 InAnimFrameIndex) const override;
+		// ~END FVertexDeltaSampler overrides.
+
+		/**
+		 * Get the array of mesh names that we cannot do any sampling for.
+		 * The reason for this can be that the sampler cannot find a matching mesh in the skeletal mesh for specific geometry cache tracks.
+		 * For example if there is some geometry cache track that is named "Head", but there is no such mesh inside the skeletal mesh, then
+		 * the returned array of names will include a string with the value "Head".
+		 * @return An array of geometry cache track names for which no mesh could be found inside the SkeletalMesh / linear skinned actor.
+		 */
+		const TArray<FString>& GetFailedImportedMeshNames() const	{ return FailedImportedMeshNames; }
+
+		/**
+		 * Executed when we need a pointer to the geometry cache object that we want to sample from.
+		 * @return A reference to the delegate.
+		 */
+		FMLDeformerGetGeomCacheEvent& OnGetGeometryCache()			{ return GetGeometryCacheEvent; }
+
+	protected:
+		/**
+		 * Calculate the vertex deltas between the linear skinned mesh (skeletal mesh) and geometry cache.
+		 * A delta cutoff length can be specified. Deltas with a length longer than this cutoff length will be ignored and set to zero.
+		 * This can be useful in cases where for some reason due to some mesh errors some deltas are very long.
+		 * Setting the DeltaCutoffLength to a very large value would essentially disable this filtering and always include all deltas.
+		 * Typically a good value for the delta cutoff length seems to be 30.
+		 * @param SkinnedPositions The vertex positions of the linear skinned mesh, so of the skeletal mesh.
+		 * @param DeltaCutoffLength The delta cutoff length as described above. Usually a value of 30 is a good value, or set to a very large value to disable it.
+		 * @param OutVertexDeltas The array that will receive the deltas, in a float buffer. The buffer will contain 3*NumVerts number of elements, and the layout is
+		 *                        (x, y, z, x, y, z, x, y, z...) so 3 values per vertex: x, y and z.
+		 */
+		void CalculateVertexDeltas(const TArray<FVector3f>& SkinnedPositions, float DeltaCutoffLength, TArray<float>& OutVertexDeltas);
+
+	protected:
+		/** The geometry cache component used to sample the geometry cache. */
+		TObjectPtr<UGeometryCacheComponent> GeometryCacheComponent = nullptr;
+
+		/** Maps skeletal meshes imported meshes to geometry tracks. */
+		TArray<FMLDeformerGeomCacheMeshMapping> MeshMappings;
+
+		/** The geometry cache mesh data reusable buffers. One for each MeshMapping.*/
+		TArray<FGeometryCacheMeshData> GeomCacheMeshDatas;
+
+		/** Geom cache track names for which no mesh can be found inside the skeletal mesh. */
+		TArray<FString> FailedImportedMeshNames; 
+
+		/** Imported mesh names in the skeletal mesh for which the geometry track had a different vertex count. */
+		TArray<FString> VertexCountMisMatchNames;
+
+		/** The function that grabs the geometry cache. */
+		FMLDeformerGetGeomCacheEvent GetGeometryCacheEvent;
+	};
+}	// namespace UE::MLDeformer
diff --git a/Engine/Plugins/Animation/MLDeformer/MLDeformerFramework/Source/MLDeformerFrameworkEditor/Public/MLDeformerModelDetails.h b/Engine/Plugins/Animation/MLDeformer/MLDeformerFramework/Source/MLDeformerFrameworkEditor/Public/MLDeformerModelDetails.h
new file mode 100644
--- /dev/null
+++ b/Engine/Plugins/Animation/MLDeformer/MLDeformerFramework/Source/MLDeformerFrameworkEditor/Public/MLDeformerModelDetails.h
@@ -0,0 +1,109 @@
+// Copyright Epic Games, Inc. All Rights Reserved.
+
+#pragma once
+
+#include "CoreMinimal.h"
+#include "IDetailCustomization.h"
+#include "AssetRegistry/AssetData.h"
+#include "Input/Reply.h"
+
+class IDetailLayoutBuilder;
+class USkeleton;
+class IDetailCategoryBuilder;
+class UMLDeformerModel;
+class UGeometryCache;
+class USkeletalMesh;
+
+namespace UE::MLDeformer
+{
+	class FMLDeformerEditorModel;
+
+	/**
+	 * The ML Deformer model detail customization base class.
+	 * This adds the shared properties, inserts error messages, creates some groups, etc.
+	 * Model detail customizations should all inherit from this class. If you use a geometry cache based model however, you will 
+	 * likely inherit from other classes such as FMLDeformerGeomCacheModelDetails or FMLDeformerMorphModelDetails. 
+	 * Those classes are also inherited from this one though.
+	 */
+	class MLDEFORMERFRAMEWORKEDITOR_API FMLDeformerModelDetails
+		: public IDetailCustomization
+	{
+	public:
+		// ILayoutDetails overrides.
+		virtual void CustomizeDetails(IDetailLayoutBuilder& DetailBuilder) override;
+		// ~END ILayoutDetails overrides.
+
+		/**
+		 * Update the class member pointers, which includes the pointer to the model, and its editor model.
+		 * @param Objects The array of objects that the detail customization is showing.
+		 */
+		virtual bool UpdateMemberPointers(const TArray<TWeakObjectPtr<UObject>>& Objects);
+
+		/**
+		 * Create the categories that we can add properties to. This will update the category member pointers. */
+		virtual void CreateCategories();
+
+		/**
+		 * Add the target mesh property. For example the FMLDeformerGeomCacheModelDetails class will add the geometry cache property here. */
+		virtual void AddTargetMesh() {}
+
+		/** Add additional errors related to the base mesh. */
+		virtual void AddBaseMeshErrors() {}
+
+		/** Add additional errors related to the training animation sequence. */
+		virtual void AddAnimSequenceErrors() {}
+
+		/** Add additional errors related to bone inputs. */
+		virtual void AddBoneInputErrors() {}
+
+		/** Add additional errors related to curve inputs. */
+		virtual void AddCurveInputErrors() {}
+
+		/** Add training input flags, which basically are things like the check boxes that specify whether bones or curves (or both) should be included. */
+		virtual void AddTrainingInputFlags() {}
+
+		/** Add additional training input filters. Filters are things like a list of bones, or curves. */
+		virtual void AddTrainingInputFilters() {}
+
+		/** Add additional training input errors. */
+		virtual void AddTrainingInputErrors() {}
+
+		/** Is the "include bones" checkbox visible? */
+		virtual bool IsBonesFlagVisible() const;
+
+		/** Is the "include curves" checkbox visible? */
+		virtual bool IsCurvesFlagVisible() const;
+
+	protected:
+		/** The filter that only shows anim sequences that are compatible with the given skeleton. */
+		bool FilterAnimSequences(const FAssetData& AssetData, USkeleton* Skeleton);
+
+		/** Executed when we press the Animated Bones Only button. */
+		FReply OnFilterAnimatedBonesOnly() const;
+
+		/** Executed when we press the Animated Curves Only button. */
+		FReply OnFilterAnimatedCurvesOnly() const;
+
+	protected:
+		/** Associated detail layout builder. */
+		IDetailLayoutBuilder* DetailLayoutBuilder = nullptr;
+
+		/** A pointer to the model we are customizing details for. */
+		TObjectPtr<UMLDeformerModel> Model = nullptr;
+
+		/** A pointer to the editor model of the runtime model. */
+		FMLDeformerEditorModel* EditorModel = nullptr;
+
+		/** The category related to the base mesh. */
+		IDetailCategoryBuilder* BaseMeshCategoryBuilder = nullptr;
+
+		/** The category related to the target mesh. */
+		IDetailCategoryBuilder* TargetMeshCategoryBuilder = nullptr;
+
+		/** The category related to the inputs and outputs. */
+		IDetailCategoryBuilder* InputOutputCategoryBuilder = nullptr;
+
+		/** The training settings category. You most likely add your model properties to this. */
+		IDetailCategoryBuilder* TrainingSettingsCategoryBuilder = nullptr;
+	};
+}	// namespace UE::MLDeformer
diff --git a/Engine/Plugins/Animation/MLDeformer/VertexDeltaModel/Source/VertexDeltaModelEditor/Private/VertexDeltaEditorModel.cpp b/Engine/Plugins/Animation/MLDeformer/VertexDeltaModel/Source/VertexDeltaModelEditor/Private/VertexDeltaEditorModel.cpp
new file mode 100644
--- /dev/null
+++ b/Engine/Plugins/Animation/MLDeformer/VertexDeltaModel/Source/VertexDeltaModelEditor/Private/VertexDeltaEditorModel.cpp
@@ -0,0 +1,37 @@
+// Copyright Epic Games, Inc. All Rights Reserved.
+
+#include "VertexDeltaEditorModel.h"
+#include "VertexDeltaModel.h"
+#include "VertexDeltaModelVizSettings.h"
+#include "VertexDeltaTrainingModel.h"
+#include "MLDeformerGeomCacheActor.h"
+#include "MLDeformerEditorToolkit.h"
+#include "MLDeformerEditorStyle.h"
+#include "MLDeformerInputInfo.h"
+#include "MLDeformerGeomCacheSampler.h"
+#include "MLDeformerGeomCacheHelpers.h"
+#include "GeometryCache.h"
+#include "GeometryCacheComponent.h"
+#include "Animation/DebugSkelMeshComponent.h"
+#include "Rendering/SkeletalMeshModel.h"
+#include "Rendering/SkeletalMeshLODModel.h"
+#include "Engine/SkeletalMesh.h"
+
+#define LOCTEXT_NAMESPACE "VertexDeltaEditorModel"
+
+namespace UE::VertexDeltaModel
+{
+	using namespace UE::MLDeformer;
+
+	FMLDeformerEditorModel* FVertexDeltaEditorModel::MakeInstance()
+	{
+		return new FVertexDeltaEditorModel();
+	}
+
+	ETrainingResult FVertexDeltaEditorModel::Train()
+	{
+		return TrainModel<UVertexDeltaTrainingModel>(this);
+	}
+}	// namespace UE::VertexDeltaModel
+
+#undef LOCTEXT_NAMESPACE
