[MLDeformer]
- Changed the deformer graph for the morph model to the updated deformer graph path.
- Fixed some more caching bugs.
- Log loss during training.
- Fixed a bug in the reset to default button of the deformer graph.

diff --git a/Engine/Plugins/Animation/MLDeformer/MLDeformerFramework/Source/MLDeformerFrameworkEditor/Private/MLDeformerGeomCacheEditorModel.cpp b/Engine/Plugins/Animation/MLDeformer/MLDeformerFramework/Source/MLDeformerFrameworkEditor/Private/MLDeformerGeomCacheEditorModel.cpp
--- a/Engine/Plugins/Animation/MLDeformer/MLDeformerFramework/Source/MLDeformerFrameworkEditor/Private/MLDeformerGeomCacheEditorModel.cpp
+++ b/Engine/Plugins/Animation/MLDeformer/MLDeformerFramework/Source/MLDeformerFrameworkEditor/Private/MLDeformerGeomCacheEditorModel.cpp
@@ -234,8 +234,12 @@ namespace UE::MLDeformer
 		FMLDeformerEditorModel::OnPropertyChanged(PropertyChangedEvent);
 
 		// Properties specific to the geometry cache editor model.
-		if (Property->GetFName() == UMLDeformerGeomCacheModel::GetGeometryCachePropertyName() ||
-			Property->GetFName() == UMLDeformerGeomCacheVizSettings::GetTestGroundTruthPropertyName())
+		if (Property->GetFName() == UMLDeformerGeomCacheModel::GetGeometryCachePropertyName())
+		{
+			SetResamplingInputOutputsNeeded(true);
+			TriggerInputAssetChanged(true);
+		}
+		else if (Property->GetFName() == UMLDeformerGeomCacheVizSettings::GetTestGroundTruthPropertyName())
 		{
 			TriggerInputAssetChanged(true);
 		}
diff --git a/Engine/Plugins/Animation/MLDeformer/MLDeformerFramework/Source/MLDeformerFrameworkEditor/Private/MLDeformerMorphModelEditorModel.cpp b/Engine/Plugins/Animation/MLDeformer/MLDeformerFramework/Source/MLDeformerFrameworkEditor/Private/MLDeformerMorphModelEditorModel.cpp
--- a/Engine/Plugins/Animation/MLDeformer/MLDeformerFramework/Source/MLDeformerFrameworkEditor/Private/MLDeformerMorphModelEditorModel.cpp
+++ b/Engine/Plugins/Animation/MLDeformer/MLDeformerFramework/Source/MLDeformerFrameworkEditor/Private/MLDeformerMorphModelEditorModel.cpp
@@ -61,8 +61,8 @@ namespace UE::MLDeformer
 	}
 
 	FString FMLDeformerMorphModelEditorModel::GetDefaultDeformerGraphAssetPath() const
-	{
-		return FString(TEXT("/Script/OptimusCore.OptimusDeformer'/Optimus/Deformers/DG_LinearBlendSkin_Morph_Cloth_RecomputeNormals.DG_LinearBlendSkin_Morph_Cloth_RecomputeNormals'"));
+	{						 
+		return FString(TEXT("/DeformerGraph/Deformers/DG_LinearBlendSkin_Morph_Cloth_RecomputeNormals.DG_LinearBlendSkin_Morph_Cloth_RecomputeNormals"));
 	}
 
 	FString FMLDeformerMorphModelEditorModel::GetHeatMapDeformerGraphPath() const
