[MLDeformer]
- Removed one macro to declare the deformer graphs, as it only was confusing.
- Every morph based ML deformer asset will have a unique ID for its morph target set. This should allow multiple of the same type of model to be applied without problems on the same blueprint.
- Fixed a crash when changing models.
- More doxygen.
- Some small misc changes.

diff --git a/Engine/Plugins/Animation/MLDeformer/MLDeformerFramework/Source/MLDeformerFramework/Private/MLDeformerGraphDataInterface.cpp b/Engine/Plugins/Animation/MLDeformer/MLDeformerFramework/Source/MLDeformerFramework/Private/MLDeformerGraphDataInterface.cpp
--- a/Engine/Plugins/Animation/MLDeformer/MLDeformerFramework/Source/MLDeformerFramework/Private/MLDeformerGraphDataInterface.cpp
+++ b/Engine/Plugins/Animation/MLDeformer/MLDeformerFramework/Source/MLDeformerFramework/Private/MLDeformerGraphDataInterface.cpp
@@ -44,14 +44,32 @@ BEGIN_SHADER_PARAMETER_STRUCT(FMLDeformerGraphDataInterfaceParameters, )
 	MLDEFORMER_SHADER_PARAMETERS()
 END_SHADER_PARAMETER_STRUCT()
 
-MLDEFORMER_GRAPH_IMPLEMENT_BASICS(
-	UMLDeformerGraphDataInterface,
-	UMLDeformerGraphDataProvider,
-	UE::MLDeformer::FMLDeformerGraphDataProviderProxy,
-	FMLDeformerGraphDataInterfaceParameters,
-	TEXT("#include \"/Plugin/MLDeformerFramework/Private/MLDeformerGraphDataInterface.ush\"\n"),
-	TEXT("ML Deformer"))
+FString UMLDeformerGraphDataInterface::GetDisplayName() const
+{
+	return TEXT("ML Deformer");
+}
 
+void UMLDeformerGraphDataInterface::GetShaderParameters(TCHAR const* UID, FShaderParametersMetadataBuilder& InOutBuilder, FShaderParametersMetadataAllocations& InOutAllocations) const \
+{
+	InOutBuilder.AddNestedStruct<FMLDeformerGraphDataInterfaceParameters>(UID);
+}
+
+void UMLDeformerGraphDataInterface::GetHLSL(FString& OutHLSL) const
+{
+	OutHLSL += TEXT("#include \"/Plugin/MLDeformerFramework/Private/MLDeformerGraphDataInterface.ush\"\n");
+}
+
+UComputeDataProvider* UMLDeformerGraphDataInterface::CreateDataProvider(TObjectPtr<UObject> InBinding, uint64 InInputMask, uint64 InOutputMask) const
+{
+	UMLDeformerGraphDataProvider* Provider = NewObject<UMLDeformerGraphDataProvider>();
+	Provider->DeformerComponent = Cast<UMLDeformerComponent>(InBinding);
+	return Provider;
+}
+
+FComputeDataProviderRenderProxy* UMLDeformerGraphDataProvider::GetRenderProxy()
+{
+	return new UE::MLDeformer::FMLDeformerGraphDataProviderProxy(DeformerComponent);
+}
 
 bool UMLDeformerGraphDataProvider::IsValid() const
 {
diff --git a/Engine/Plugins/Animation/MLDeformer/MLDeformerFramework/Source/MLDeformerFramework/Private/MLDeformerGraphDebugDataInterface.cpp b/Engine/Plugins/Animation/MLDeformer/MLDeformerFramework/Source/MLDeformerFramework/Private/MLDeformerGraphDebugDataInterface.cpp
--- a/Engine/Plugins/Animation/MLDeformer/MLDeformerFramework/Source/MLDeformerFramework/Private/MLDeformerGraphDebugDataInterface.cpp
+++ b/Engine/Plugins/Animation/MLDeformer/MLDeformerFramework/Source/MLDeformerFramework/Private/MLDeformerGraphDebugDataInterface.cpp
@@ -6,13 +6,9 @@
 #include "MLDeformerAsset.h"
 #include "MLDeformerComponent.h"
 #include "MLDeformerVizSettings.h"
-#include "Animation/AnimSequence.h"
 #include "Components/SkeletalMeshComponent.h"
 #include "ComputeFramework/ComputeKernelPermutationSet.h"
 #include "ComputeFramework/ShaderParamTypeDefinition.h"
-#include "GeometryCache.h"
-#include "GeometryCacheMeshData.h"
-#include "GeometryCacheTrack.h"
 #include "OptimusDataDomain.h"
 #include "RenderGraphBuilder.h"
 #include "RenderGraphUtils.h"
@@ -65,13 +61,37 @@ BEGIN_SHADER_PARAMETER_STRUCT(FMLDeformerGraphDebugDataInterfaceParameters, )
 	MLDEFORMER_DEBUG_SHADER_PARAMETERS()
 END_SHADER_PARAMETER_STRUCT()
 
-MLDEFORMER_GRAPH_IMPLEMENT_DEBUG_BASICS_WITH_PROXY(
-	UMLDeformerGraphDebugDataInterface,
-	UMLDeformerGraphDebugDataProvider,
-	UE::MLDeformer::FMLDeformerGraphDebugDataProviderProxy,
-	FMLDeformerGraphDebugDataInterfaceParameters,
-	TEXT("#include \"/Plugin/MLDeformerFramework/Private/MLDeformerGraphHeatMapDataInterface.ush\"\n"),
-	TEXT("ML Deformer Debug"))
+FString UMLDeformerGraphDebugDataInterface::GetDisplayName() const
+{
+	return TEXT("ML Deformer Debug");
+}
+
+void UMLDeformerGraphDebugDataInterface::GetShaderParameters(TCHAR const* UID, FShaderParametersMetadataBuilder& InOutBuilder, FShaderParametersMetadataAllocations& InOutAllocations) const
+{
+	InOutBuilder.AddNestedStruct<FMLDeformerGraphDebugDataInterfaceParameters>(UID);
+}
+
+void UMLDeformerGraphDebugDataInterface::GetHLSL(FString& OutHLSL) const
+{
+	OutHLSL += TEXT("#include \"/Plugin/MLDeformerFramework/Private/MLDeformerGraphHeatMapDataInterface.ush\"\n");
+}
+
+UComputeDataProvider* UMLDeformerGraphDebugDataInterface::CreateDataProvider(TObjectPtr<UObject> InBinding, uint64 InInputMask, uint64 InOutputMask) const
+{
+	UMLDeformerGraphDebugDataProvider* Provider = NewObject<UMLDeformerGraphDebugDataProvider>();
+	Provider->DeformerComponent = Cast<UMLDeformerComponent>(InBinding);
+	if (Provider->DeformerComponent)
+	{
+		Provider->DeformerAsset = Provider->DeformerComponent->GetDeformerAsset();
+	}
+	MLDEFORMER_EDITORDATA_ONLY(
+		if (Provider->DeformerAsset != nullptr)
+		{
+			Provider->Init();
+		}
+	,)
+	return Provider;
+}
 
 bool UMLDeformerGraphDebugDataProvider::IsValid() const
 {
@@ -87,6 +107,24 @@ bool UMLDeformerGraphDebugDataProvider::IsValid() const
 #endif
 }
 
+FComputeDataProviderRenderProxy* UMLDeformerGraphDebugDataProvider::GetRenderProxy()
+{
+#if WITH_EDITORONLY_DATA
+	UE::MLDeformer::FMLDeformerGraphDebugDataProviderProxy* Proxy = new UE::MLDeformer::FMLDeformerGraphDebugDataProviderProxy(DeformerComponent, DeformerAsset, this);
+	if (DeformerAsset)
+	{
+		const float SampleTime = DeformerComponent->GetModelInstance()->GetSkeletalMeshComponent()->GetPosition();
+		UMLDeformerModel* Model = DeformerAsset->GetModel();
+		Model->SampleGroundTruthPositions(SampleTime, Proxy->GetGroundTruthPositions());
+		Proxy->HandleZeroGroundTruthPositions();
+		return Proxy;
+	}
+	return nullptr;
+#else
+	return nullptr;
+#endif
+}
+
 #if WITH_EDITORONLY_DATA
 namespace UE::MLDeformer
 {
@@ -97,13 +135,13 @@ namespace UE::MLDeformer
 
 		if (DeformerAsset)
 		{
-			UMLDeformerModel* Model = DeformerAsset->GetModel();	
-			UMLDeformerVizSettings* VizSettings = Model->GetVizSettings();
+			const UMLDeformerModel* Model = DeformerAsset->GetModel();	
+			const UMLDeformerVizSettings* VizSettings = Model->GetVizSettings();
 			const UMLDeformerModelInstance* ModelInstance = DeformerComponent->GetModelInstance();
 
 			SkeletalMeshObject = ModelInstance->GetSkeletalMeshComponent()->MeshObject;
 			VertexMapBufferSRV = Model->GetVertexMapBuffer().ShaderResourceViewRHI;
-			HeatMapMode = (int32)VizSettings->GetHeatMapMode();
+			HeatMapMode = static_cast<int32>(VizSettings->GetHeatMapMode());
 			HeatMapMax = 1.0f / FMath::Max(VizSettings->GetHeatMapMax(), 0.00001f);
 			GroundTruthLerp = VizSettings->GetGroundTruthLerp();
 		}
@@ -118,7 +156,7 @@ namespace UE::MLDeformer
 			GroundTruthPositions.Add(FVector3f::ZeroVector);
 
 			// Silently disable relevant debug things.
-			if (HeatMapMode == (int32)EMLDeformerHeatMapMode::GroundTruth)
+			if (HeatMapMode == static_cast<int32>(EMLDeformerHeatMapMode::GroundTruth))
 			{
 				HeatMapMode = -1;
 				HeatMapMax = 0.0f;
diff --git a/Engine/Plugins/Animation/MLDeformer/MLDeformerFramework/Source/MLDeformerFramework/Private/MLDeformerMorphModel.cpp b/Engine/Plugins/Animation/MLDeformer/MLDeformerFramework/Source/MLDeformerFramework/Private/MLDeformerMorphModel.cpp
--- a/Engine/Plugins/Animation/MLDeformer/MLDeformerFramework/Source/MLDeformerFramework/Private/MLDeformerMorphModel.cpp
+++ b/Engine/Plugins/Animation/MLDeformer/MLDeformerFramework/Source/MLDeformerFramework/Private/MLDeformerMorphModel.cpp
@@ -9,11 +9,14 @@
 
 #define LOCTEXT_NAMESPACE "MLDeformerMorphModel"
 
+TAtomic<int32> UMLDeformerMorphModel::NextFreeMorphSetID(0);
+
 UMLDeformerMorphModel::UMLDeformerMorphModel(const FObjectInitializer& ObjectInitializer)
 	: Super(ObjectInitializer)
 {
 	MorphTargetSet = MakeShared<FExternalMorphSet>();
 	MorphTargetSet->Name = GetClass()->GetFName();
+	ExternalMorphSetID = NextFreeMorphSetID++;
 }
 
 void UMLDeformerMorphModel::Serialize(FArchive& Archive)
@@ -58,9 +61,6 @@ void UMLDeformerMorphModel::PostMLDeformerComponentInit(UMLDeformerModelInstance
 	USkeletalMeshComponent* SkelMeshComponent = ModelInstance->GetSkeletalMeshComponent();
 	if (SkelMeshComponent && SkelMeshComponent->GetSkeletalMeshAsset())
 	{	
-		// If this triggers, please set the ExternalMorphSetID in your model's constructor to a value that represents your model type.
-		checkf(ExternalMorphSetID != -1, TEXT("Please configure your ExternalMorphSetID member to a unique ID."));
-
 		// Register the morph set. This overwrites the existing one for this model, if it already exists.
 		// Only add to LOD 0 for now.
 		const int32 LOD = 0;
@@ -131,4 +131,10 @@ void UMLDeformerMorphModel::BeginDestroy()
 	Super::BeginDestroy();
 }
 
+FExternalMorphSetWeights* UMLDeformerMorphModel::FindExternalMorphWeights(int32 LOD, USkinnedMeshComponent* SkinnedMeshComponent) const
+{
+	check(SkinnedMeshComponent);
+	return SkinnedMeshComponent->GetExternalMorphWeights(LOD).MorphSets.Find(ExternalMorphSetID);
+}
+
 #undef LOCTEXT_NAMESPACE
diff --git a/Engine/Plugins/Animation/MLDeformer/MLDeformerFramework/Source/MLDeformerFramework/Private/MLDeformerMorphModelInstance.cpp b/Engine/Plugins/Animation/MLDeformer/MLDeformerFramework/Source/MLDeformerFramework/Private/MLDeformerMorphModelInstance.cpp
--- a/Engine/Plugins/Animation/MLDeformer/MLDeformerFramework/Source/MLDeformerFramework/Private/MLDeformerMorphModelInstance.cpp
+++ b/Engine/Plugins/Animation/MLDeformer/MLDeformerFramework/Source/MLDeformerFramework/Private/MLDeformerMorphModelInstance.cpp
@@ -16,14 +16,10 @@ void UMLDeformerMorphModelInstance::RunNeuralNetwork(float ModelWeight)
 		return;
 	}
 
-	const int32 ExternalMorphSetID = MorphModel->GetExternalMorphSetID();
-	// If this check fails please set this member to some value larger than 0 in your model's constructor.
-	checkf(ExternalMorphSetID != -1, TEXT("Please set the ExternalMorphSetID member value to a unique value for your model type."));
-
 	// Grab the weight data for this morph set.
 	// This could potentially fail if we are applying this deformer to the wrong skeletal mesh component.
 	const int LOD = 0;	// For now we only support LOD 0, as we can't setup an ML Deformer per LOD yet.
-	FExternalMorphSetWeights* WeightData = SkeletalMeshComponent->GetExternalMorphWeights(LOD).MorphSets.Find(ExternalMorphSetID);
+	FExternalMorphSetWeights* WeightData = MorphModel->FindExternalMorphWeights(LOD, SkeletalMeshComponent);
 	if (WeightData == nullptr)
 	{
 		return;
@@ -31,19 +27,21 @@ void UMLDeformerMorphModelInstance::RunNeuralNetwork(float ModelWeight)
 
 	// If our model is active, we want to run the neural network and update the morph weights
 	// with the values that the neural net calculated for us.
-	if (ModelWeight > 0.0f)
+	const UNeuralNetwork* NeuralNetwork = Model->GetNeuralNetwork();
+	if (ModelWeight > 0.0f && NeuralNetwork != nullptr)
 	{
 		// Perform the neural network inference, which updates the output tensor.
+		// This takes most CPU time inside this method.
 		UMLDeformerModelInstance::RunNeuralNetwork(ModelWeight);
 
 		// Get the output tensor, read the values and use them as morph target weights inside the skeletal mesh component.
-		const UNeuralNetwork* NeuralNetwork = Model->GetNeuralNetwork();
 		const FNeuralTensor& OutputTensor = NeuralNetwork->GetOutputTensorForContext(NeuralNetworkInferenceHandle);
 		const int32 NumNetworkWeights = OutputTensor.Num();
 		const int32 NumMorphTargets = WeightData->Weights.Num();;
 
 		// If we have the expected amount of weights.
 		// +1 because we always have an extra morph target that represents the means, with fixed weight of 1.
+		// Therefore, the neural network output tensor will contain one less float than the number of morph targets in our morph set.
 		if (NumMorphTargets == NumNetworkWeights + 1)
 		{
 			// Set the first morph target, which represents the means, to a weight of 1.0, as it always needs to be fully active.
@@ -52,16 +50,13 @@ void UMLDeformerMorphModelInstance::RunNeuralNetwork(float ModelWeight)
 			// Update all generated morph target weights with the values calculated by our neural network.
 			for (int32 MorphIndex = 0; MorphIndex < NumNetworkWeights; ++MorphIndex)
 			{
-				WeightData->Weights[MorphIndex + 1] = OutputTensor.At<float>(MorphIndex) * ModelWeight;
+				const float OutputTensorValue = OutputTensor.At<float>(MorphIndex);
+				WeightData->Weights[MorphIndex + 1] = OutputTensorValue * ModelWeight;
 			}
 		}
 	}
 	else
 	{
-		// Reset all weights to zero.
-		for (float& Weight : WeightData->Weights)
-		{
-			Weight = 0.0f;
-		}
+		WeightData->ZeroWeights();
 	}
 }
diff --git a/Engine/Plugins/Animation/MLDeformer/MLDeformerFramework/Source/MLDeformerFramework/Public/MLDeformerGraphDataInterface.h b/Engine/Plugins/Animation/MLDeformer/MLDeformerFramework/Source/MLDeformerFramework/Public/MLDeformerGraphDataInterface.h
--- a/Engine/Plugins/Animation/MLDeformer/MLDeformerFramework/Source/MLDeformerFramework/Public/MLDeformerGraphDataInterface.h
+++ b/Engine/Plugins/Animation/MLDeformer/MLDeformerFramework/Source/MLDeformerFramework/Public/MLDeformerGraphDataInterface.h
@@ -48,30 +48,6 @@ class UMLDeformerModel;
 
 #define MLDEFORMER_GRAPH_DISPATCH_END() }
 
-#define MLDEFORMER_GRAPH_IMPLEMENT_BASICS(InterfaceClassName, DataProviderClassName, DataProviderProxyClassName, ParameterStructType, HLSLText, DisplayName) \
-	FString InterfaceClassName::GetDisplayName() const \
-	{ \
-		return DisplayName; \
-	} \
-	void InterfaceClassName::GetShaderParameters(TCHAR const* UID, FShaderParametersMetadataBuilder& InOutBuilder, FShaderParametersMetadataAllocations& InOutAllocations) const \
-	{ \
-		InOutBuilder.AddNestedStruct<ParameterStructType>(UID); \
-	} \
-	void InterfaceClassName::GetHLSL(FString& OutHLSL) const \
-	{ \
-		OutHLSL += HLSLText; \
-	} \
-	UComputeDataProvider* InterfaceClassName::CreateDataProvider(TObjectPtr<UObject> InBinding, uint64 InInputMask, uint64 InOutputMask) const \
-	{ \
-		DataProviderClassName* Provider = NewObject<DataProviderClassName>(); \
-		Provider->DeformerComponent = Cast<UMLDeformerComponent>(InBinding); \
-		return Provider; \
-	} \
-	FComputeDataProviderRenderProxy* DataProviderClassName::GetRenderProxy() \
-	{ \
-		return new DataProviderProxyClassName(DeformerComponent); \
-	}
-
 #if WITH_EDITORONLY_DATA
 	#define MLDEFORMER_EDITORDATA_ONLY(Statement, ElseStatement) Statement
 #else
diff --git a/Engine/Plugins/Animation/MLDeformer/MLDeformerFramework/Source/MLDeformerFramework/Public/MLDeformerGraphDebugDataInterface.h b/Engine/Plugins/Animation/MLDeformer/MLDeformerFramework/Source/MLDeformerFramework/Public/MLDeformerGraphDebugDataInterface.h
--- a/Engine/Plugins/Animation/MLDeformer/MLDeformerFramework/Source/MLDeformerFramework/Public/MLDeformerGraphDebugDataInterface.h
+++ b/Engine/Plugins/Animation/MLDeformer/MLDeformerFramework/Source/MLDeformerFramework/Public/MLDeformerGraphDebugDataInterface.h
@@ -36,56 +36,6 @@ class USkeletalMeshComponent;
 	Parameters->PositionGroundTruthBuffer = GroundTruthBufferSRV; \
 	Parameters->VertexMapBuffer = VertexMapBufferSRV;
 
-#define MLDEFORMER_GRAPH_IMPLEMENT_DEBUG_BASICS(InterfaceClassName, DataProviderClassName, ParameterStructType, HLSLText, DisplayName) \
-	FString InterfaceClassName::GetDisplayName() const \
-	{ \
-		return DisplayName; \
-	} \
-	void InterfaceClassName::GetShaderParameters(TCHAR const* UID, FShaderParametersMetadataBuilder& InOutBuilder, FShaderParametersMetadataAllocations& InOutAllocations) const \
-	{ \
-		InOutBuilder.AddNestedStruct<ParameterStructType>(UID); \
-	} \
-	void InterfaceClassName::GetHLSL(FString& OutHLSL) const \
-	{ \
-		OutHLSL += HLSLText; \
-	} \
-	UComputeDataProvider* InterfaceClassName::CreateDataProvider(TObjectPtr<UObject> InBinding, uint64 InInputMask, uint64 InOutputMask) const \
-	{ \
-		DataProviderClassName* Provider = NewObject<DataProviderClassName>(); \
-		Provider->DeformerComponent = Cast<UMLDeformerComponent>(InBinding); \
-		if (Provider->DeformerComponent) \
-		{ \
-			Provider->DeformerAsset = Provider->DeformerComponent->GetDeformerAsset(); \
-		} \
-		MLDEFORMER_EDITORDATA_ONLY( \
-			if (Provider->DeformerAsset != nullptr) \
-			{ \
-				Provider->Init(); \
-			} \
-		,) \
-		return Provider; \
-	}
-
-#define MLDEFORMER_GRAPH_IMPLEMENT_DEBUG_BASICS_WITH_PROXY(InterfaceClassName, DataProviderClassName, DataProviderProxyClassName, ParameterStructType, HLSLText, DisplayName) \
-	MLDEFORMER_GRAPH_IMPLEMENT_DEBUG_BASICS(InterfaceClassName, DataProviderClassName, ParameterStructType, HLSLText, DisplayName) \
-	FComputeDataProviderRenderProxy* DataProviderClassName::GetRenderProxy() \
-	{ \
-		MLDEFORMER_EDITORDATA_ONLY( \
-			DataProviderProxyClassName* Proxy = new DataProviderProxyClassName(DeformerComponent, DeformerAsset, this); \
-			if (DeformerAsset) \
-			{ \
-				const float SampleTime = DeformerComponent->GetModelInstance()->GetSkeletalMeshComponent()->GetPosition(); \
-				UMLDeformerModel* Model = DeformerAsset->GetModel(); \
-				Model->SampleGroundTruthPositions(SampleTime, Proxy->GetGroundTruthPositions()); \
-				Proxy->HandleZeroGroundTruthPositions(); \
-				return Proxy; \
-			} \
-			return nullptr; \
-			, \
-			return nullptr; \
-		) \
-	}
-
 /** 
  * Compute Framework Data Interface for MLDeformer debugging data. 
  * This interfaces to editor only data, and so will only give valid results in that context.
diff --git a/Engine/Plugins/Animation/MLDeformer/MLDeformerFramework/Source/MLDeformerFramework/Public/MLDeformerMorphModel.h b/Engine/Plugins/Animation/MLDeformer/MLDeformerFramework/Source/MLDeformerFramework/Public/MLDeformerMorphModel.h
--- a/Engine/Plugins/Animation/MLDeformer/MLDeformerFramework/Source/MLDeformerFramework/Public/MLDeformerMorphModel.h
+++ b/Engine/Plugins/Animation/MLDeformer/MLDeformerFramework/Source/MLDeformerFramework/Public/MLDeformerMorphModel.h
@@ -7,6 +7,8 @@
 #include "MLDeformerMorphModel.generated.h"
 
 struct FExternalMorphSet;
+struct FExternalMorphSetWeights;
+class USkinnedMeshComponent;
 
 UCLASS()
 class MLDEFORMERFRAMEWORK_API UMLDeformerMorphModel
@@ -40,23 +42,60 @@ public:
 
 	/**
 	 * Set the per vertex deltas, as a set of floats. Each vertex delta must have 3 floats.
+	 * These deltas are used to generate compressed morph targets internally. You typically call this method from inside
+	 * the python training script once your morph target deltas have been generated there.
 	 * Concatenate all deltas into one buffer, so like this [morphdeltas_target0, morphdeltas_target1, ..., morphdeltas_targetN].
 	 * The vertex ordering should be: [(x, y, z), (x, y, z), (x, y, z)].
+	 * This is the same as SetMorphTargetDeltas, except that this takes an array of floats instead of vectors.
 	 * @param Deltas The array of floats that contains the deltas. The number of items in the array must be equal to (NumMorphs * NumBaseMeshVerts * 3).
 	 */
 	UFUNCTION(BlueprintCallable, Category = "MLDeformerMorphModel")
 	void SetMorphTargetDeltaFloats(const TArray<float>& Deltas);
 
+	/**
+	 * Set the morph target model deltas as an array of 3D vectors.
+	 * These deltas are used to generate compressed morph targets internally. You typically call this method from inside
+	 * the python training script once your morph target deltas have been generated there.
+	 * Concatenate all deltas into one buffer, so like this [morphdeltas_target0, morphdeltas_target1, ..., morphdeltas_targetN].
+	 * This is the same as SetMorphTargetDeltaFloats, except that it takes vectors instead of floats.
+	 * @param Deltas The array of 3D vectors that contains the vertex deltas. The number of items in the array must be equal to (NumMorphs * NumBaseMeshVerts).
+	 */
 	UFUNCTION(BlueprintCallable, Category = "MLDeformerMorphModel")
 	void SetMorphTargetDeltas(const TArray<FVector3f>& Deltas);
 
+	/**
+	 * Get the morph target delta vectors array.
+	 * The layout of this array is [morphdeltas_target0, morphdeltas_target1, ..., morphdeltas_targetN].
+	 * So the total number of items in the array returned equals (NumMorphTargets * NumBaseMeshVerts).
+	 */
 	const TArray<FVector3f>& GetMorphTargetDeltas() const	{ return MorphTargetDeltas; }
+
+	/**
+	 * Get the external morph target set ID for this model.
+	 * This basically identifies the set of morph targets that belong to this model.
+	 * Different models on the same skeletal mesh gives each model its own unique ID.
+	 * You can use this ID to find the weight values for a specific model instance, inside the USkinnedMeshComponent class.
+	 * @return The unique ID of the morph target set for this model.
+	 * @see USkinnedMeshComponent::GetExternalMorphWeights.
+	 * @see FindExternalMorphWeights.
+	 */
 	int32 GetExternalMorphSetID() const						{ return ExternalMorphSetID; }
-	void SetExternalMorphSetID(int32 ID)					{ check(ID != -1);  ExternalMorphSetID = ID; }
+
+	/**
+	 * Get the weights for the external morph target set that belongs to this model.
+	 * @param LOD The LOD level to get the weights for.
+	 * @param SkinnedMeshComponent The skinned mesh component that we have to search in.
+	 * @return A pointer to the weight data, or nullptr in case we cannot find the weight data.
+	 */
+	FExternalMorphSetWeights* FindExternalMorphWeights(int32 LOD, USkinnedMeshComponent* SkinnedMeshComponent) const;
+
+	/**
+	 * Get the morph target set.
+	 */
 	TSharedPtr<FExternalMorphSet> GetMorphTargetSet() const { return MorphTargetSet; }
 
 	/**
-	 * Get the start index into the array of deltas, for a given morph target.
+	 * Get the start index into the array of deltas (vectors3's), for a given morph target.
 	 * This does not perform a bounds check to see if MorphTargetIndex is in a valid range, so be aware.
 	 * @param MorphTargetIndex The morph target index.
 	 * @return The start index, or INDEX_NONE in case there are no deltas.
@@ -64,6 +103,9 @@ public:
 	int32 GetMorphTargetDeltaStartIndex(int32 MorphTargetIndex) const;
 
 protected:
+	/** The next free morph target set ID. This is used to generate unique ID's for each morph model. */
+	static TAtomic<int32> NextFreeMorphSetID;
+
 	/** The compressed morph target data, ready for the GPU. */
 	TSharedPtr<FExternalMorphSet> MorphTargetSet;
 
diff --git a/Engine/Plugins/Animation/MLDeformer/MLDeformerFramework/Source/MLDeformerFramework/Public/MLDeformerMorphModelInstance.h b/Engine/Plugins/Animation/MLDeformer/MLDeformerFramework/Source/MLDeformerFramework/Public/MLDeformerMorphModelInstance.h
--- a/Engine/Plugins/Animation/MLDeformer/MLDeformerFramework/Source/MLDeformerFramework/Public/MLDeformerMorphModelInstance.h
+++ b/Engine/Plugins/Animation/MLDeformer/MLDeformerFramework/Source/MLDeformerFramework/Public/MLDeformerMorphModelInstance.h
@@ -20,5 +20,7 @@ class MLDEFORMERFRAMEWORK_API UMLDeformerMorphModelInstance
 	GENERATED_BODY()
 
 public:
+	// UMLDeformerModelInstance overrides.
 	virtual void RunNeuralNetwork(float ModelWeight) override;
+	// ~END UMLDeformerModelInstance overrides.
 };
diff --git a/Engine/Plugins/Animation/MLDeformer/MLDeformerFramework/Source/MLDeformerFrameworkEditor/Private/MLDeformerEditorModel.cpp b/Engine/Plugins/Animation/MLDeformer/MLDeformerFramework/Source/MLDeformerFrameworkEditor/Private/MLDeformerEditorModel.cpp
--- a/Engine/Plugins/Animation/MLDeformer/MLDeformerFramework/Source/MLDeformerFrameworkEditor/Private/MLDeformerEditorModel.cpp
+++ b/Engine/Plugins/Animation/MLDeformer/MLDeformerFramework/Source/MLDeformerFrameworkEditor/Private/MLDeformerEditorModel.cpp
@@ -253,8 +253,7 @@ namespace UE::MLDeformer
 
 		PreviewScene->DeselectAll();
 		PreviewScene->SetPreviewAnimationAsset(nullptr);
-		PreviewScene->SetPreviewAnimationBlueprint(nullptr, nullptr);
-		PreviewScene->SetPreviewMesh(nullptr);
+		//PreviewScene->SetPreviewMesh(nullptr);
 		PreviewScene->SetPreviewMeshComponent(nullptr);
 		PreviewScene->SetActor(nullptr);
 
diff --git a/Engine/Plugins/Animation/MLDeformer/NeuralMorphModel/Source/NeuralMorphModel/Private/NeuralMorphModel.cpp b/Engine/Plugins/Animation/MLDeformer/NeuralMorphModel/Source/NeuralMorphModel/Private/NeuralMorphModel.cpp
--- a/Engine/Plugins/Animation/MLDeformer/NeuralMorphModel/Source/NeuralMorphModel/Private/NeuralMorphModel.cpp
+++ b/Engine/Plugins/Animation/MLDeformer/NeuralMorphModel/Source/NeuralMorphModel/Private/NeuralMorphModel.cpp
@@ -28,9 +28,6 @@ UNeuralMorphModel::UNeuralMorphModel(const FObjectInitializer& ObjectInitializer
 #if WITH_EDITORONLY_DATA
 	VizSettings = ObjectInitializer.CreateEditorOnlyDefaultSubobject<UNeuralMorphModelVizSettings>(this, TEXT("VizSettings"));
 #endif
-
-	// The ID for the morph target data generated by this model.
-	ExternalMorphSetID = 0;
 }
 
 #undef LOCTEXT_NAMESPACE
diff --git a/Engine/Source/Editor/Persona/Private/AnimationEditorPreviewScene.cpp b/Engine/Source/Editor/Persona/Private/AnimationEditorPreviewScene.cpp
--- a/Engine/Source/Editor/Persona/Private/AnimationEditorPreviewScene.cpp
+++ b/Engine/Source/Editor/Persona/Private/AnimationEditorPreviewScene.cpp
@@ -805,7 +805,11 @@ void FAnimationEditorPreviewScene::ClearSelectedBone()
 	TGuardValue<bool> RecursionGuard(bSelecting, true);
 
 	SelectedBoneIndex = INDEX_NONE;
-	SkeletalMeshComponent->BonesOfInterest.Empty();
+
+	if (SkeletalMeshComponent)
+	{
+		SkeletalMeshComponent->BonesOfInterest.Empty();
+	}
 
 	OnSelectedBoneChanged.Broadcast(NAME_None, ESelectInfo::Direct);
 
diff --git a/Engine/Source/Runtime/Engine/Classes/Components/SkinnedMeshComponent.h b/Engine/Source/Runtime/Engine/Classes/Components/SkinnedMeshComponent.h
--- a/Engine/Source/Runtime/Engine/Classes/Components/SkinnedMeshComponent.h
+++ b/Engine/Source/Runtime/Engine/Classes/Components/SkinnedMeshComponent.h
@@ -112,12 +112,9 @@ namespace EBoneSpaces
 	};
 }
 
-/** Struct used to indicate one active morph target that should be applied to this USkeletalMesh when rendered. */
-
 /** WeightIndex is an into the MorphTargetWeights array */
 using FMorphTargetWeightMap = TMap<const UMorphTarget* /* MorphTarget */, int32 /* WeightIndex */>;
 
-
 /** An external morph target set. External morph targets are managed by systems outside of the skinned meshes. */
 struct ENGINE_API FExternalMorphSet
 {
@@ -128,6 +125,7 @@ struct ENGINE_API FExternalMorphSet
 	FMorphTargetVertexInfoBuffers MorphBuffers;
 };
 
+/** The map of external morph sets registered on the skinned mesh component. */
 using FExternalMorphSets = TMap<int32, TSharedPtr<FExternalMorphSet>>;
 
 /** The weight data for a specific external morph set. */
@@ -136,8 +134,11 @@ struct ENGINE_API FExternalMorphSetWeights
 	/** Update the number of active morph targets. */
 	void UpdateNumActiveMorphTargets();
 
+	/** Set all weights to 0. Optionally set the NumActiveMorphTargets to zero as well. */
+	void ZeroWeights(bool bZeroNumActiveMorphTargets=true);
+
 	/** The debug name. */
-	FName Name = FName(TEXT("Unknown"));
+	FName Name = FName(TEXT("Unknown ExternalMorphSetWeights"));
 
 	/** The weights, which can also be negative and go beyond 1.0 or -1.0. */
 	TArray<float> Weights;
@@ -233,7 +234,6 @@ struct FVertexOffsetUsage
  *
  * @see USkeletalMeshComponent
 */
-
 UCLASS(hidecategories=Object, config=Engine, editinlinenew, abstract)
 class ENGINE_API USkinnedMeshComponent : public UMeshComponent, public ILODSyncInterface
 {
diff --git a/Engine/Source/Runtime/Engine/Private/Components/SkinnedMeshComponent.cpp b/Engine/Source/Runtime/Engine/Private/Components/SkinnedMeshComponent.cpp
--- a/Engine/Source/Runtime/Engine/Private/Components/SkinnedMeshComponent.cpp
+++ b/Engine/Source/Runtime/Engine/Private/Components/SkinnedMeshComponent.cpp
@@ -342,6 +342,19 @@ void FExternalMorphSetWeights::UpdateNumActiveMorphTargets()
 	}
 }
 
+void FExternalMorphSetWeights::ZeroWeights(bool bZeroNumActiveMorphTargets)
+{
+	for (int32 Index = 0; Index < Weights.Num(); ++Index)
+	{
+		Weights[Index] = 0.0f;
+	}
+
+	if (bZeroNumActiveMorphTargets)
+	{
+		NumActiveMorphTargets = 0;
+	}
+}
+
 //////////////////////////////////////////////////////////////////////////
 
 void FExternalMorphWeightData::UpdateNumActiveMorphTargets()
