[MLDeformer]
- Made morph model (and inherited classses) work without Mesh Deformer set. This didn't work when skin cache was enabled.
- Adjusted UI to not give warnings when using a model that doesn't require a Mesh Deformer graph, and not having a graph selected.
- Give a warning in the log when you have a blueprint that uses an ML model that requires a graph, but you didn't set a graph on your Mesh/SKMcomponent.
- Fixed a crash when applying the default ML Deformer graph on a ML Deformer asset trained with neural morph model.
- Fixed some other crash due to nullptr.
- Fixed a crash after training, when the skeletal mesh had changed.

diff --git a/Engine/Plugins/Animation/MLDeformer/MLDeformerFramework/Source/MLDeformerFramework/Private/MLDeformerComponent.cpp b/Engine/Plugins/Animation/MLDeformer/MLDeformerFramework/Source/MLDeformerFramework/Private/MLDeformerComponent.cpp
--- a/Engine/Plugins/Animation/MLDeformer/MLDeformerFramework/Source/MLDeformerFramework/Private/MLDeformerComponent.cpp
+++ b/Engine/Plugins/Animation/MLDeformer/MLDeformerFramework/Source/MLDeformerFramework/Private/MLDeformerComponent.cpp
@@ -5,6 +5,8 @@
 #include "MLDeformerModelInstance.h"
 #include "MLDeformerModel.h"
 #include "Components/SkeletalMeshComponent.h"
+#include "Animation/MeshDeformer.h"
+#include "Animation/MeshDeformerInstance.h"
 
 UMLDeformerComponent::UMLDeformerComponent(const FObjectInitializer& ObjectInitializer)
 	: Super(ObjectInitializer)
@@ -60,6 +62,25 @@ void UMLDeformerComponent::SetupComponent(UMLDeformerAsset* InDeformerAsset, USk
 	RemoveNeuralNetworkModifyDelegate();
 	Init();
 	AddNeuralNetworkModifyDelegate();
+
+	// Verify that a mesh deformer has been setup when the used ML model requires one.
+	if (!bSuppressMeshDeformerLogWarnings && DeformerAsset && ModelInstance && ModelInstance->GetModel() && SkelMeshComponent)
+	{
+		const FString DefaultGraph = ModelInstance->GetModel()->GetDefaultDeformerGraphAssetPath();
+		if (!DefaultGraph.IsEmpty() && !SkelMeshComponent->HasMeshDeformer())
+		{
+			UE_LOG(LogMLDeformer,
+				Warning, 
+				TEXT("ML Deformer Asset '%s' used with skel mesh component '%s' on actor '%s' uses model type '%s', which requires a deformer graph (default='%s'). " \
+					"No Mesh Deformer has been set on the mesh or skel mesh component, so linear skinning will be used."),
+					*DeformerAsset->GetName(),
+					*SkelMeshComponent->GetName(),
+					*SkelMeshComponent->GetOuter()->GetName(),
+					*ModelInstance->GetModel()->GetDisplayName(),
+					*ModelInstance->GetModel()->GetDefaultDeformerGraphAssetPath()
+			);
+		}
+	}
 }
 
 void UMLDeformerComponent::AddNeuralNetworkModifyDelegate()
@@ -174,6 +195,9 @@ void UMLDeformerComponent::TickComponent(float DeltaTime, enum ELevelTick TickTy
 
 		if (Property->GetFName() == GET_MEMBER_NAME_CHECKED(UMLDeformerComponent, DeformerAsset))
 		{
+			RenderCommandFence.BeginFence();
+			RenderCommandFence.Wait();
+
 			RemoveNeuralNetworkModifyDelegate();
 			Init();
 			AddNeuralNetworkModifyDelegate();
diff --git a/Engine/Plugins/Animation/MLDeformer/MLDeformerFramework/Source/MLDeformerFramework/Private/MLDeformerGraphDataInterface.cpp b/Engine/Plugins/Animation/MLDeformer/MLDeformerFramework/Source/MLDeformerFramework/Private/MLDeformerGraphDataInterface.cpp
--- a/Engine/Plugins/Animation/MLDeformer/MLDeformerFramework/Source/MLDeformerFramework/Private/MLDeformerGraphDataInterface.cpp
+++ b/Engine/Plugins/Animation/MLDeformer/MLDeformerFramework/Source/MLDeformerFramework/Private/MLDeformerGraphDataInterface.cpp
@@ -73,7 +73,7 @@ FComputeDataProviderRenderProxy* UMLDeformerGraphDataProvider::GetRenderProxy()
 
 bool UMLDeformerGraphDataProvider::IsValid() const
 {
-	if (DeformerComponent == nullptr || DeformerComponent->GetDeformerAsset() == nullptr)
+	if (DeformerComponent == nullptr || DeformerComponent->GetDeformerAsset() == nullptr || DeformerComponent->GetModelInstance() == nullptr)
 	{
 		return false;
 	}
@@ -96,7 +96,7 @@ namespace UE::MLDeformer
 			SkeletalMeshObject = ModelInstance->GetSkeletalMeshComponent()->MeshObject;
 			NeuralNetwork = Model->GetNeuralNetwork();
 			NeuralNetworkInferenceHandle = ModelInstance->GetNeuralNetworkInferenceHandle();
-			bCanRunNeuralNet = ModelInstance->IsCompatible();
+			bCanRunNeuralNet = ModelInstance->IsCompatible() && Model->IsNeuralNetworkOnGPU();
 			Weight = DeformerComponent->GetWeight();
 			VertexMapBufferSRV = Model->GetVertexMapBuffer().ShaderResourceViewRHI;
 		}
diff --git a/Engine/Plugins/Animation/MLDeformer/MLDeformerFramework/Source/MLDeformerFramework/Private/MLDeformerGraphDebugDataInterface.cpp b/Engine/Plugins/Animation/MLDeformer/MLDeformerFramework/Source/MLDeformerFramework/Private/MLDeformerGraphDebugDataInterface.cpp
--- a/Engine/Plugins/Animation/MLDeformer/MLDeformerFramework/Source/MLDeformerFramework/Private/MLDeformerGraphDebugDataInterface.cpp
+++ b/Engine/Plugins/Animation/MLDeformer/MLDeformerFramework/Source/MLDeformerFramework/Private/MLDeformerGraphDebugDataInterface.cpp
@@ -96,7 +96,7 @@ UComputeDataProvider* UMLDeformerGraphDebugDataInterface::CreateDataProvider(TOb
 bool UMLDeformerGraphDebugDataProvider::IsValid() const
 {
 #if WITH_EDITORONLY_DATA
-	if (DeformerComponent == nullptr || DeformerComponent->GetDeformerAsset() == nullptr)
+	if (DeformerComponent == nullptr || DeformerComponent->GetDeformerAsset() == nullptr || DeformerComponent->GetModelInstance() == nullptr)
 	{
 		return false;
 	}
diff --git a/Engine/Plugins/Animation/MLDeformer/MLDeformerFramework/Source/MLDeformerFramework/Private/MLDeformerModel.cpp b/Engine/Plugins/Animation/MLDeformer/MLDeformerFramework/Source/MLDeformerFramework/Private/MLDeformerModel.cpp
--- a/Engine/Plugins/Animation/MLDeformer/MLDeformerFramework/Source/MLDeformerFramework/Private/MLDeformerModel.cpp
+++ b/Engine/Plugins/Animation/MLDeformer/MLDeformerFramework/Source/MLDeformerFramework/Private/MLDeformerModel.cpp
@@ -72,6 +72,15 @@ void UMLDeformerModel::Serialize(FArchive& Archive)
 	#endif
 
 	Super::Serialize(Archive);
+
+	#if WITH_EDITOR
+		// Reinitialize the vertex map after loading its contents.
+		// Because our skeletal mesh asset could have changed vertex counts.
+		if (Archive.IsLoading())
+		{
+			InitVertexMap();
+		}
+	#endif
 }
 
 UMLDeformerAsset* UMLDeformerModel::GetDeformerAsset() const
diff --git a/Engine/Plugins/Animation/MLDeformer/MLDeformerFramework/Source/MLDeformerFramework/Private/MLDeformerModelInstance.cpp b/Engine/Plugins/Animation/MLDeformer/MLDeformerFramework/Source/MLDeformerFramework/Private/MLDeformerModelInstance.cpp
--- a/Engine/Plugins/Animation/MLDeformer/MLDeformerFramework/Source/MLDeformerFramework/Private/MLDeformerModelInstance.cpp
+++ b/Engine/Plugins/Animation/MLDeformer/MLDeformerFramework/Source/MLDeformerFramework/Private/MLDeformerModelInstance.cpp
@@ -17,6 +17,11 @@ void UMLDeformerModelInstance::BeginDestroy()
 
 void UMLDeformerModelInstance::Release()
 {
+	// Wait for the render commands to finish, because some neural network might still be executing, using the inference handle
+	// that we are about to delete.
+	RenderCommandFence.BeginFence();
+	RenderCommandFence.Wait();
+
 	// Destroy the neural network instance.
 	UNeuralNetwork* NeuralNetwork = Model.Get() ? Model->GetNeuralNetwork() : nullptr;
 	if (NeuralNetwork && NeuralNetworkInferenceHandle != -1)
diff --git a/Engine/Plugins/Animation/MLDeformer/MLDeformerFramework/Source/MLDeformerFramework/Public/MLDeformerComponent.h b/Engine/Plugins/Animation/MLDeformer/MLDeformerFramework/Source/MLDeformerFramework/Public/MLDeformerComponent.h
--- a/Engine/Plugins/Animation/MLDeformer/MLDeformerFramework/Source/MLDeformerFramework/Public/MLDeformerComponent.h
+++ b/Engine/Plugins/Animation/MLDeformer/MLDeformerFramework/Source/MLDeformerFramework/Public/MLDeformerComponent.h
@@ -5,6 +5,7 @@
 #include "CoreMinimal.h"
 #include "UObject/Object.h"
 #include "Components/ActorComponent.h"
+#include "RenderCommandFence.h"
 #include "MLDeformerComponent.generated.h"
 
 class UMLDeformerAsset;
@@ -91,6 +92,13 @@ public:
 	 */
 	USkeletalMeshComponent* GetSkeletalMeshComponent() const	{ return SkelMeshComponent.Get(); }
 
+	/**
+	 * Suppress logging warnings about mesh deformers not being set.
+	 * A warning is logged when an ML Deformer is used that requires a deformer graph, but the skeletal mesh has no deformer graph setup.
+	 * @param bSuppress Set to true to silent warnings about deformer graphs not being set, while the active ML Model needs one.
+	 */
+	void SetSuppressMeshDeformerLogWarnings(bool bSuppress)		{ bSuppressMeshDeformerLogWarnings = bSuppress; }
+
 	// Get property names.
 	static FName GetDeformerAssetPropertyName()					{ return GET_MEMBER_NAME_CHECKED(UMLDeformerComponent, DeformerAsset); }
 	static FName GetWeightPropertyName()						{ return GET_MEMBER_NAME_CHECKED(UMLDeformerComponent, Weight); }
@@ -115,6 +123,9 @@ protected:
 	void RemoveNeuralNetworkModifyDelegate();
 
 protected:
+	/** Render command fence that let's us wait for all other commands to finish. */
+	FRenderCommandFence RenderCommandFence;
+
 	/** 
 	 * The skeletal mesh component we want to grab the bone transforms etc from. 
 	 * This can be a nullptr. When it is a nullptr then it will internally try to find the first skeletal mesh component on the actor.
@@ -125,6 +136,9 @@ protected:
 	/** DelegateHandle for NeuralNetwork modification. */
 	FDelegateHandle NeuralNetworkModifyDelegateHandle;
 
+	/** Suppress mesh deformer logging warnings? This is used by the ML Deformer editor, as we don't want to show some warnings when using that. */
+	bool bSuppressMeshDeformerLogWarnings = false;
+
 	/** The deformer asset to use. */
 	UPROPERTY(EditAnywhere, DisplayName = "ML Deformer Asset", Category = "ML Deformer")
 	TObjectPtr<UMLDeformerAsset> DeformerAsset = nullptr;
diff --git a/Engine/Plugins/Animation/MLDeformer/MLDeformerFramework/Source/MLDeformerFramework/Public/MLDeformerModel.h b/Engine/Plugins/Animation/MLDeformer/MLDeformerFramework/Source/MLDeformerFramework/Public/MLDeformerModel.h
--- a/Engine/Plugins/Animation/MLDeformer/MLDeformerFramework/Source/MLDeformerFramework/Public/MLDeformerModel.h
+++ b/Engine/Plugins/Animation/MLDeformer/MLDeformerFramework/Source/MLDeformerFramework/Public/MLDeformerModel.h
@@ -132,6 +132,13 @@ public:
 	 */
 	virtual bool IsNeuralNetworkOnGPU() const				{ return false; }	// CPU neural network.
 
+	/**
+	 * Get the default deformer graph asset path that this model uses, or an empty string if it doesn't require any deformer graph.
+	 * An example is some string like: "/DeformerGraph/Deformers/DG_LinearBlendSkin_Morph_Cloth_RecomputeNormals.DG_LinearBlendSkin_Morph_Cloth_RecomputeNormals".
+	 * @return The asset path of the deformer graph that should be used on default. This can return an empty string when no deformer graph is required.
+	 */
+	virtual FString GetDefaultDeformerGraphAssetPath() const	{ return FString(); }
+
 #if WITH_EDITORONLY_DATA
 	/**
 	 * Check whether this model currently has a training mesh setup or not.
diff --git a/Engine/Plugins/Animation/MLDeformer/MLDeformerFramework/Source/MLDeformerFramework/Public/MLDeformerModelInstance.h b/Engine/Plugins/Animation/MLDeformer/MLDeformerFramework/Source/MLDeformerFramework/Public/MLDeformerModelInstance.h
--- a/Engine/Plugins/Animation/MLDeformer/MLDeformerFramework/Source/MLDeformerFramework/Public/MLDeformerModelInstance.h
+++ b/Engine/Plugins/Animation/MLDeformer/MLDeformerFramework/Source/MLDeformerFramework/Public/MLDeformerModelInstance.h
@@ -5,6 +5,7 @@
 #include "CoreMinimal.h"
 #include "MLDeformerModule.h"
 #include "UObject/ObjectPtr.h"
+#include "RenderCommandFence.h"
 #include "MLDeformerModelInstance.generated.h"
 
 class UMLDeformerModel;
@@ -171,6 +172,9 @@ protected:
 	void UpdateBoneTransforms();
 
 protected:
+	/** The fence that let's us wait for all render commands to finish, before we continue. */
+	FRenderCommandFence RenderCommandFence;
+
 	/** The ML Deformer model that this is an instance of. */
 	UPROPERTY(Transient)
 	TObjectPtr<UMLDeformerModel> Model = nullptr;
diff --git a/Engine/Plugins/Animation/MLDeformer/MLDeformerFramework/Source/MLDeformerFramework/Public/MLDeformerMorphModel.h b/Engine/Plugins/Animation/MLDeformer/MLDeformerFramework/Source/MLDeformerFramework/Public/MLDeformerMorphModel.h
--- a/Engine/Plugins/Animation/MLDeformer/MLDeformerFramework/Source/MLDeformerFramework/Public/MLDeformerMorphModel.h
+++ b/Engine/Plugins/Animation/MLDeformer/MLDeformerFramework/Source/MLDeformerFramework/Public/MLDeformerMorphModel.h
@@ -128,7 +128,7 @@ protected:
 	 * A value of 0 will result in the highest quality morph targets, at the cost of higher runtime memory usage.
 	 */
 	UPROPERTY(EditAnywhere, BlueprintReadWrite, Category = "Morph Targets", meta = (ClampMin = "0.0", ClampMax = "1.0", ForceUnits="cm"))
-	float MorphTargetDeltaThreshold = 0.0025f;
+	float MorphTargetDeltaThreshold = 0.01f;
 
 	/** The morph target error tolerance. Higher values result in larger compression, but could result in visual artifacts. */
 	UPROPERTY(EditAnywhere, BlueprintReadWrite, Category = "Morph Targets", meta = (ClampMin = "0.01", ClampMax = "500"))
diff --git a/Engine/Plugins/Animation/MLDeformer/MLDeformerFramework/Source/MLDeformerFramework/Public/MLDeformerMorphModelVizSettings.h b/Engine/Plugins/Animation/MLDeformer/MLDeformerFramework/Source/MLDeformerFramework/Public/MLDeformerMorphModelVizSettings.h
--- a/Engine/Plugins/Animation/MLDeformer/MLDeformerFramework/Source/MLDeformerFramework/Public/MLDeformerMorphModelVizSettings.h
+++ b/Engine/Plugins/Animation/MLDeformer/MLDeformerFramework/Source/MLDeformerFramework/Public/MLDeformerMorphModelVizSettings.h
@@ -41,7 +41,7 @@ protected:
 	 * This only can be used after you trained, in the same editor session directly after training.
 	 */
 	UPROPERTY(EditAnywhere, Category = "Live Settings", meta = (ClampMin = "0.001", ClampMax = "1.0", ForceUnits="cm"))
-	float MorphTargetDeltaThreshold = 0.0025f;
+	float MorphTargetDeltaThreshold = 0.01f;
 
 	/**
 	 * Specify whether we want to debug draw the morph targets.
diff --git a/Engine/Plugins/Animation/MLDeformer/MLDeformerFramework/Source/MLDeformerFrameworkEditor/Private/MLDeformerEditorModel.cpp b/Engine/Plugins/Animation/MLDeformer/MLDeformerFramework/Source/MLDeformerFrameworkEditor/Private/MLDeformerEditorModel.cpp
--- a/Engine/Plugins/Animation/MLDeformer/MLDeformerFramework/Source/MLDeformerFrameworkEditor/Private/MLDeformerEditorModel.cpp
+++ b/Engine/Plugins/Animation/MLDeformer/MLDeformerFramework/Source/MLDeformerFrameworkEditor/Private/MLDeformerEditorModel.cpp
@@ -187,6 +187,7 @@ namespace UE::MLDeformer
 		// Create the ML Deformer component.
 		UMLDeformerAsset* DeformerAsset = Model->GetDeformerAsset();
 		UMLDeformerComponent* MLDeformerComponent = NewObject<UMLDeformerComponent>(Actor);
+		MLDeformerComponent->SetSuppressMeshDeformerLogWarnings(true);
 		MLDeformerComponent->SetDeformerAsset(DeformerAsset);
 		MLDeformerComponent->RegisterComponent();
 		MLDeformerComponent->SetupComponent(DeformerAsset, SkelMeshComponent);
@@ -1390,21 +1391,25 @@ namespace UE::MLDeformer
 		UpdateDeformerGraph();
 	}
 
-	UMeshDeformer* FMLDeformerEditorModel::LoadDefaultDeformerGraph()
+	UMeshDeformer* FMLDeformerEditorModel::LoadDefaultDeformerGraph() const
 	{
-		const FString GraphAssetPath = GetDefaultDeformerGraphAssetPath();
-		UObject* Object = StaticLoadObject(UMeshDeformer::StaticClass(), nullptr, *GraphAssetPath);
-		UMeshDeformer* DeformerGraph = Cast<UMeshDeformer>(Object);
-		if (DeformerGraph == nullptr)
+		const FString GraphAssetPath = Model->GetDefaultDeformerGraphAssetPath();
+		if (!GraphAssetPath.IsEmpty())
 		{
-			UE_LOG(LogMLDeformer, Warning, TEXT("Failed to load default ML Deformer compute graph from: %s"), *GraphAssetPath);
-		}
-		else
-		{
-			UE_LOG(LogMLDeformer, Verbose, TEXT("Loaded default ML Deformer compute graph from: %s"), *GraphAssetPath);
+			UObject* Object = StaticLoadObject(UMeshDeformer::StaticClass(), nullptr, *GraphAssetPath);
+			UMeshDeformer* DeformerGraph = Cast<UMeshDeformer>(Object);
+			if (DeformerGraph == nullptr)
+			{
+				UE_LOG(LogMLDeformer, Warning, TEXT("Failed to load default Mesh Deformer from: %s"), *GraphAssetPath);
+			}
+			else
+			{
+				UE_LOG(LogMLDeformer, Verbose, TEXT("Loaded default Mesh Deformer from: %s"), *GraphAssetPath);
+			}
+			return DeformerGraph;
 		}
 
-		return DeformerGraph;
+		return nullptr;
 	}
 
 	void FMLDeformerEditorModel::SetDefaultDeformerGraphIfNeeded()
@@ -1682,7 +1687,7 @@ namespace UE::MLDeformer
 		OutMorphBuffers.InitMorphResources
 		(
 			GMaxRHIShaderPlatform,
-			TArray<FSkelMeshRenderSection>(),	// Empty array, as we don't need tangents, since we recalculate them in the shader.
+			RenderData->LODRenderData[LOD].RenderSections,
 			MorphTargets,
 			NumRenderVertices,
 			LOD,
@@ -1755,11 +1760,6 @@ namespace UE::MLDeformer
 		return FString(TEXT("/MLDeformerFramework/Deformers/DG_MLDeformerModel_HeatMap.DG_MLDeformerModel_HeatMap"));
 	}
 
-	FString FMLDeformerEditorModel::GetDefaultDeformerGraphAssetPath() const
-	{
-		return FString(TEXT("/MLDeformerFramework/Deformers/DG_MLDeformerModel.DG_MLDeformerModel"));
-	}
-
 	const UAnimSequence* FMLDeformerEditorModel::GetAnimSequence() const
 	{
 		if (GetModel())
diff --git a/Engine/Plugins/Animation/MLDeformer/MLDeformerFramework/Source/MLDeformerFrameworkEditor/Private/MLDeformerMorphModelEditorModel.cpp b/Engine/Plugins/Animation/MLDeformer/MLDeformerFramework/Source/MLDeformerFrameworkEditor/Private/MLDeformerMorphModelEditorModel.cpp
--- a/Engine/Plugins/Animation/MLDeformer/MLDeformerFramework/Source/MLDeformerFrameworkEditor/Private/MLDeformerMorphModelEditorModel.cpp
+++ b/Engine/Plugins/Animation/MLDeformer/MLDeformerFramework/Source/MLDeformerFrameworkEditor/Private/MLDeformerMorphModelEditorModel.cpp
@@ -60,11 +60,6 @@ namespace UE::MLDeformer
 		return Cast<UMLDeformerMorphModelVizSettings>(GetMorphModel()->GetVizSettings());
 	}
 
-	FString FMLDeformerMorphModelEditorModel::GetDefaultDeformerGraphAssetPath() const
-	{						 
-		return FString(TEXT("/DeformerGraph/Deformers/DG_LinearBlendSkin_Morph_Cloth_RecomputeNormals.DG_LinearBlendSkin_Morph_Cloth_RecomputeNormals"));
-	}
-
 	FString FMLDeformerMorphModelEditorModel::GetHeatMapDeformerGraphPath() const
 	{
 		return FString(TEXT("/MLDeformerFramework/Deformers/DG_MLDeformerModel_GPUMorph_HeatMap.DG_MLDeformerModel_GPUMorph_HeatMap"));
diff --git a/Engine/Plugins/Animation/MLDeformer/MLDeformerFramework/Source/MLDeformerFrameworkEditor/Private/MLDeformerVizSettingsDetails.cpp b/Engine/Plugins/Animation/MLDeformer/MLDeformerFramework/Source/MLDeformerFrameworkEditor/Private/MLDeformerVizSettingsDetails.cpp
--- a/Engine/Plugins/Animation/MLDeformer/MLDeformerFramework/Source/MLDeformerFrameworkEditor/Private/MLDeformerVizSettingsDetails.cpp
+++ b/Engine/Plugins/Animation/MLDeformer/MLDeformerFramework/Source/MLDeformerFrameworkEditor/Private/MLDeformerVizSettingsDetails.cpp
@@ -131,8 +131,9 @@ namespace UE::MLDeformer
 		TSharedRef<IPropertyHandle> DeformerGraphProperty = DetailBuilder.GetProperty(UMLDeformerVizSettings::GetDeformerGraphPropertyName(), UMLDeformerVizSettings::StaticClass());
 		if (DeformerGraphProperty->GetValue(Graph) == FPropertyAccess::Result::Success)
 		{
+			const bool bShowError = (Graph == nullptr) && !Model->GetDefaultDeformerGraphAssetPath().IsEmpty();
 			FDetailWidgetRow& GraphErrorRow = TestAssetsCategory->AddCustomRow(FText::FromString("GraphError"))
-				.Visibility((Graph == nullptr) ? EVisibility::Visible : EVisibility::Collapsed)
+				.Visibility(bShowError ? EVisibility::Visible : EVisibility::Collapsed)
 				.WholeRowContent()
 				[
 					SNew(SBox)
@@ -140,7 +141,7 @@ namespace UE::MLDeformer
 					[
 						SNew(SWarningOrErrorBox)
 						.MessageStyle(EMessageStyle::Warning)
-						.Message(FText::FromString("Please select a deformer graph.\nOtherwise only linear skinning is used."))
+						.Message(LOCTEXT("GraphErrorText", "This model requires a deformer graph.\nWithout it linear skinning is used."))
 					]
 				];
 		}
@@ -197,7 +198,8 @@ namespace UE::MLDeformer
 		PropertyHandle->GetValue(CurrentGraph);
 		if (CurrentGraph == nullptr)
 		{
-			return true;
+			const FString DefaultPath = Model->GetDefaultDeformerGraphAssetPath();
+			return !DefaultPath.IsEmpty();
 		}
 
 		if (EditorModel)
@@ -205,7 +207,7 @@ namespace UE::MLDeformer
 			// Check if we already assigned the default asset.
 			const FAssetData CurrentGraphAssetData(CurrentGraph);
 			const FString CurrentPath = CurrentGraphAssetData.ObjectPath.ToString();
-			const FString DefaultPath = EditorModel->GetDefaultDeformerGraphAssetPath();
+			const FString DefaultPath = Model->GetDefaultDeformerGraphAssetPath();
 			return (DefaultPath != CurrentPath);
 		}
 
diff --git a/Engine/Plugins/Animation/MLDeformer/MLDeformerFramework/Source/MLDeformerFrameworkEditor/Public/MLDeformerEditorModel.h b/Engine/Plugins/Animation/MLDeformer/MLDeformerFramework/Source/MLDeformerFrameworkEditor/Public/MLDeformerEditorModel.h
--- a/Engine/Plugins/Animation/MLDeformer/MLDeformerFramework/Source/MLDeformerFrameworkEditor/Public/MLDeformerEditorModel.h
+++ b/Engine/Plugins/Animation/MLDeformer/MLDeformerFramework/Source/MLDeformerFrameworkEditor/Public/MLDeformerEditorModel.h
@@ -485,7 +485,7 @@ namespace UE::MLDeformer
 		 * Load the default Optimus deformer graph to use on the skeletal mesh when using this model.
 		 * @return The mesh deformer graph.
 		 */
-		virtual UMeshDeformer* LoadDefaultDeformerGraph();
+		virtual UMeshDeformer* LoadDefaultDeformerGraph() const;
 
 		/**
 		 * This makes sure that the deformer graph selected in the visualization settings is not a nullptr.
@@ -549,13 +549,6 @@ namespace UE::MLDeformer
 		 */
 		virtual FString GetHeatMapDeformerGraphPath() const;
 
-		/**
-		 * Getthe path to the default deformer graph asset.
-		 * This is the deformer graph used when not in heat map mode.
-		 * @return The path to the default deformer graph asset.
-		 */
-		virtual FString GetDefaultDeformerGraphAssetPath() const;
-
 		/** Get the current view range. */
 		TRange<double> GetViewRange() const;
 
diff --git a/Engine/Plugins/Animation/MLDeformer/MLDeformerFramework/Source/MLDeformerFrameworkEditor/Public/MLDeformerMorphModelEditorModel.h b/Engine/Plugins/Animation/MLDeformer/MLDeformerFramework/Source/MLDeformerFrameworkEditor/Public/MLDeformerMorphModelEditorModel.h
--- a/Engine/Plugins/Animation/MLDeformer/MLDeformerFramework/Source/MLDeformerFrameworkEditor/Public/MLDeformerMorphModelEditorModel.h
+++ b/Engine/Plugins/Animation/MLDeformer/MLDeformerFramework/Source/MLDeformerFrameworkEditor/Public/MLDeformerMorphModelEditorModel.h
@@ -29,7 +29,6 @@ namespace UE::MLDeformer
 		// FMLDeformerEditorModel overrides.
 		virtual void Render(const FSceneView* View, FViewport* Viewport, FPrimitiveDrawInterface* PDI) override;
 		virtual void OnPropertyChanged(FPropertyChangedEvent& PropertyChangedEvent) override;
-		virtual FString GetDefaultDeformerGraphAssetPath() const override;
 		virtual FString GetHeatMapDeformerGraphPath() const override;
 		virtual void OnPreTraining() override;
 		virtual void OnPostTraining(ETrainingResult TrainingResult, bool bUsePartiallyTrainedWhenAborted) override;
diff --git a/Engine/Plugins/Animation/MLDeformer/VertexDeltaModel/Source/VertexDeltaModel/Private/VertexDeltaModel.cpp b/Engine/Plugins/Animation/MLDeformer/VertexDeltaModel/Source/VertexDeltaModel/Private/VertexDeltaModel.cpp
--- a/Engine/Plugins/Animation/MLDeformer/VertexDeltaModel/Source/VertexDeltaModel/Private/VertexDeltaModel.cpp
+++ b/Engine/Plugins/Animation/MLDeformer/VertexDeltaModel/Source/VertexDeltaModel/Private/VertexDeltaModel.cpp
@@ -26,4 +26,9 @@ UVertexDeltaModel::UVertexDeltaModel(const FObjectInitializer& ObjectInitializer
 #endif
 }
 
+FString UVertexDeltaModel::GetDefaultDeformerGraphAssetPath() const 
+{ 
+	return FString(TEXT("/MLDeformerFramework/Deformers/DG_MLDeformerModel.DG_MLDeformerModel"));
+}
+
 #undef LOCTEXT_NAMESPACE
diff --git a/Engine/Plugins/Animation/MLDeformer/VertexDeltaModel/Source/VertexDeltaModel/Public/VertexDeltaModel.h b/Engine/Plugins/Animation/MLDeformer/VertexDeltaModel/Source/VertexDeltaModel/Public/VertexDeltaModel.h
--- a/Engine/Plugins/Animation/MLDeformer/VertexDeltaModel/Source/VertexDeltaModel/Public/VertexDeltaModel.h
+++ b/Engine/Plugins/Animation/MLDeformer/VertexDeltaModel/Source/VertexDeltaModel/Public/VertexDeltaModel.h
@@ -28,6 +28,7 @@ public:
 	// UMLDeformerModel overrides.
 	virtual FString GetDisplayName() const override			{ return "Vertex Delta Model"; }
 	virtual bool IsNeuralNetworkOnGPU() const override		{ return true; }	// GPU neural network.
+	FString GetDefaultDeformerGraphAssetPath() const override;
 	// ~END UMLDeformerModel overrides.
 
 #if WITH_EDITORONLY_DATA
diff --git a/Engine/Source/Runtime/Engine/Private/SkeletalRenderGPUSkin.cpp b/Engine/Source/Runtime/Engine/Private/SkeletalRenderGPUSkin.cpp
--- a/Engine/Source/Runtime/Engine/Private/SkeletalRenderGPUSkin.cpp
+++ b/Engine/Source/Runtime/Engine/Private/SkeletalRenderGPUSkin.cpp
@@ -1032,7 +1032,8 @@ void FSkeletalMeshObjectGPUSkin::UpdateMorphVertexBuffer(FRHICommandListImmediat
 															bool bGPUSkinCacheEnabled, FMorphVertexBuffer& MorphVertexBuffer)
 {
 	QUICK_SCOPE_CYCLE_COUNTER(STAT_FSkeletalMeshObjectGPUSkin_ProcessUpdatedDynamicData_UpdateMorphBuffer);
-	if (UseGPUMorphTargets(FeatureLevel))
+	
+	if (UseGPUMorphTargets(FeatureLevel) && IsValidRef(MorphVertexBuffer.VertexBufferRHI))
 	{
 		const int32 NumMorphSets = CalcNumActiveGPUMorphSets(MorphVertexBuffer, DynamicData->ExternalMorphSets);
 		int32 MorphSetIndex = 0;
@@ -1057,7 +1058,7 @@ void FSkeletalMeshObjectGPUSkin::UpdateMorphVertexBuffer(FRHICommandListImmediat
 			RHICmdList, 
 			DynamicData->MorphTargetWeights, 
 			LODData.MorphTargetVertexInfoBuffers, 
-			DynamicData->SectionIdsUseByActiveMorphTargets, 
+			DynamicData->SectionIdsUseByActiveMorphTargets,
 			GetDebugName(), 
 			Mode, 
 			MorphVertexBuffer, 
@@ -1080,7 +1081,7 @@ void FSkeletalMeshObjectGPUSkin::UpdateMorphVertexBuffer(FRHICommandListImmediat
 				LOD.UpdateMorphVertexBufferGPU(
 					RHICmdList, WeightData->Weights,
 					CompressedBuffers,
-					TArray<int32>(),
+					DynamicData->SectionIdsUseByActiveMorphTargets,
 					GetDebugName(),
 					Mode,
 					MorphVertexBuffer,
@@ -1211,6 +1212,7 @@ void FSkeletalMeshObjectGPUSkin::FSkeletalMeshObjectLOD::UpdateMorphVertexBuffer
 			RHICmdList.ClearUAVUint(MorphVertexBuffer.GetUAV(), FUintVector4(0, 0, 0, 0));
 		}
 
+		if (MorphTargetVertexInfoBuffers.GetNumMorphs() > 0)
 		{
 			{
 				SCOPED_DRAW_EVENTF(RHICmdList, MorphUpdateScatter, TEXT("Scatter"));
@@ -2465,7 +2467,18 @@ void FDynamicSkelMeshObjectDataGPUSkin::InitDynamicSkelMeshObjectDataGPUSkin(
 		MeshComponentSpaceTransforms.Append(InMeshComponent->GetComponentSpaceTransforms());
 	}	
 #endif
-	SectionIdsUseByActiveMorphTargets.Empty();
+	SectionIdsUseByActiveMorphTargets.Reset();
+
+	// If we have external morph targets, just include all sections.
+	if (ExternalMorphWeightData.HasActiveMorphs())
+	{
+		const FSkeletalMeshLODRenderData& LOD = InSkeletalMeshRenderData->LODRenderData[LODIndex];
+		SectionIdsUseByActiveMorphTargets.SetNumUninitialized(LOD.RenderSections.Num(), false);
+		for (int32 Index = 0; Index < LOD.RenderSections.Num(); ++Index)
+		{
+			SectionIdsUseByActiveMorphTargets[Index] = Index;
+		}
+	}
 
 	// find number of morphs that are currently weighted and will affect the mesh
 	ActiveMorphTargets.Reserve(InActiveMorphTargets.Num());
