[MLDeformer]
- Use TSoftObjectPtr for training related properties, and only load them when needed.
- Force some pointers to nullptr on cook.
- Made some members private rather than protected, to force use of Get and Set methods.
- Package Flag recovery for editor only assets to make it safer. When you change training assets such as geom caches, the flags of the previously selected geom cache will be restored, rather than kept at "editor only" while they were not.
- Changed the way we find the right skeletal mesh, by comparing with the Skeletal Mesh inside the deformer asset. This is much safer.
- Unregister weight sets and morph sets when having no deformer set.

diff --git a/Engine/Plugins/Animation/MLDeformer/MLDeformerFramework/Source/MLDeformerFramework/Private/MLDeformerComponent.cpp b/Engine/Plugins/Animation/MLDeformer/MLDeformerFramework/Source/MLDeformerFramework/Private/MLDeformerComponent.cpp
--- a/Engine/Plugins/Animation/MLDeformer/MLDeformerFramework/Source/MLDeformerFramework/Private/MLDeformerComponent.cpp
+++ b/Engine/Plugins/Animation/MLDeformer/MLDeformerFramework/Source/MLDeformerFramework/Private/MLDeformerComponent.cpp
@@ -22,6 +22,10 @@ void UMLDeformerComponent::Init()
 	// If there is no deformer asset linked, release what we currently have.
 	if (DeformerAsset == nullptr)
 	{
+		if (ModelInstance)
+		{
+			ModelInstance->Release();
+		}
 		ModelInstance = nullptr;
 		return;
 	}
@@ -124,14 +128,16 @@ void UMLDeformerComponent::BeginDestroy()
 void UMLDeformerComponent::Activate(bool bReset)
 {
 	// If we haven't pointed to some skeletal mesh component to use, then try to find one on the actor.
-	if (SkelMeshComponent == nullptr)
+	SkelMeshComponent = nullptr;
+	if (!DeformerAsset.IsNull())
 	{
 		// First search for a skeletal mesh component with the expected number of vertices.
 		// This will try to find a skeletal mesh with the same number of vertices.
 		const UMLDeformerModel* Model = DeformerAsset ? DeformerAsset->GetModel() : nullptr;
-		const int32 NumModelVertices = Model ? Model->GetVertexMap().Num() : -1;
-		if (NumModelVertices > 0)
+		if (Model && Model->GetSkeletalMesh())
 		{
+			const USkeletalMesh* ModelSkeletalMesh = Model->GetSkeletalMesh();
+
 			// Get a list of all skeletal mesh components on the actor.
 			TArray<USkeletalMeshComponent*> Components;
 			AActor* Actor = Cast<AActor>(GetOuter());
@@ -140,11 +146,8 @@ void UMLDeformerComponent::Activate(bool bReset)
 			// Find a component that uses a mesh with the same vertex count.
 			for (USkeletalMeshComponent* Component : Components)
 			{
-				const USkeletalMesh* SkeletalMesh = Component->GetSkeletalMeshAsset();
-				const FSkeletalMeshRenderData* RenderData = SkeletalMesh ? SkeletalMesh->GetResourceForRendering() : nullptr;
-				const int32 NumComponentVertices = RenderData && RenderData->LODRenderData.IsValidIndex(0) ? RenderData->LODRenderData[0].GetNumVertices() : -1;
-
-				if (NumComponentVertices == NumModelVertices)
+				const USkeletalMesh* ComponentSkeletalMesh = Component->GetSkeletalMeshAsset();
+				if (ComponentSkeletalMesh == ModelSkeletalMesh)
 				{
 					SkelMeshComponent = Component;
 					break;
@@ -161,12 +164,15 @@ void UMLDeformerComponent::Activate(bool bReset)
 	}
 
 	SetupComponent(DeformerAsset, SkelMeshComponent);
+
+	Super::Activate(bReset);
 }
 
 void UMLDeformerComponent::Deactivate()
 {
 	RemoveNeuralNetworkModifyDelegate();
 	ModelInstance = nullptr;
+	Super::Deactivate();
 }
 
 void UMLDeformerComponent::TickComponent(float DeltaTime, enum ELevelTick TickType, FActorComponentTickFunction* ThisTickFunction)
diff --git a/Engine/Plugins/Animation/MLDeformer/MLDeformerFramework/Source/MLDeformerFramework/Private/MLDeformerGeomCacheModel.cpp b/Engine/Plugins/Animation/MLDeformer/MLDeformerFramework/Source/MLDeformerFramework/Private/MLDeformerGeomCacheModel.cpp
--- a/Engine/Plugins/Animation/MLDeformer/MLDeformerFramework/Source/MLDeformerFramework/Private/MLDeformerGeomCacheModel.cpp
+++ b/Engine/Plugins/Animation/MLDeformer/MLDeformerFramework/Source/MLDeformerFramework/Private/MLDeformerGeomCacheModel.cpp
@@ -9,15 +9,27 @@
 
 #define LOCTEXT_NAMESPACE "MLDeformerGeomCacheModel"
 
+void UMLDeformerGeomCacheModel::Serialize(FArchive& Archive)
+{
+	#if WITH_EDITOR
+		if (Archive.IsSaving() && Archive.IsPersistent() && Archive.IsCooking())
+		{
+			GeometryCache = nullptr;
+		}
+	#endif
+
+	Super::Serialize(Archive);
+}
+
 #if WITH_EDITOR
 void UMLDeformerGeomCacheModel::UpdateNumTargetMeshVertices()
 {
-	NumTargetMeshVerts = UE::MLDeformer::ExtractNumImportedGeomCacheVertices(GeometryCache);
+	SetNumTargetMeshVerts(UE::MLDeformer::ExtractNumImportedGeomCacheVertices(GetGeometryCache()));
 }
 
 UMLDeformerGeomCacheVizSettings* UMLDeformerGeomCacheModel::GetGeomCacheVizSettings() const
 {
-	return Cast<UMLDeformerGeomCacheVizSettings>(VizSettings);
+	return Cast<UMLDeformerGeomCacheVizSettings>(GetVizSettings());
 }
 
 void UMLDeformerGeomCacheModel::SetAssetEditorOnlyFlags()
@@ -25,20 +37,14 @@ void UMLDeformerGeomCacheModel::SetAssetEditorOnlyFlags()
 	// Set the flags for the base class, which filters out the training anim sequence.
 	UMLDeformerModel::SetAssetEditorOnlyFlags();
 
-	// The training geometry cache is something we don't want to package.
-	if (GeometryCache)
-	{
-		GeometryCache->GetPackage()->SetPackageFlags(PKG_EditorOnly);
-	}
+	// The training geometry cache.
+	MarkObjectAsEditorOnly(GetGeometryCache());
 
-	// Filter the viz settings specific assets.
-	UMLDeformerGeomCacheVizSettings* GeomCacheVizSettings = GetGeomCacheVizSettings();
+	// The testing geometry cache (ground truth of the test anim sequence).
+	const UMLDeformerGeomCacheVizSettings* GeomCacheVizSettings = GetGeomCacheVizSettings();
 	if (GeomCacheVizSettings)
 	{
-		if (GeomCacheVizSettings->GetTestGroundTruth())
-		{
-			GeomCacheVizSettings->GetTestGroundTruth()->GetPackage()->SetPackageFlags(PKG_EditorOnly);
-		}
+		MarkObjectAsEditorOnly(GeomCacheVizSettings->GetTestGroundTruth());
 	}
 }
 #endif // WITH_EDITOR
@@ -60,16 +66,16 @@ void UMLDeformerGeomCacheModel::SampleGroundTruthPositions(float SampleTime, TAr
 	{
 		TArray<FString> FailedImportedMeshNames;
 		TArray<FString> VertexMisMatchNames;
-		UE::MLDeformer::GenerateGeomCacheMeshMappings(SkeletalMesh, GeomCache, MeshMappings, FailedImportedMeshNames, VertexMisMatchNames);
+		UE::MLDeformer::GenerateGeomCacheMeshMappings(GetSkeletalMesh(), GeomCache, MeshMappings, FailedImportedMeshNames, VertexMisMatchNames);
 	}
 
 	UE::MLDeformer::SampleGeomCachePositions(
 		0,
 		SampleTime,
 		MeshMappings,
-		SkeletalMesh,
+		GetSkeletalMesh(),
 		GeomCache,
-		AlignmentTransform,
+		GetAlignmentTransform(),
 		OutPositions);
 }
 #endif
diff --git a/Engine/Plugins/Animation/MLDeformer/MLDeformerFramework/Source/MLDeformerFramework/Private/MLDeformerModel.cpp b/Engine/Plugins/Animation/MLDeformer/MLDeformerFramework/Source/MLDeformerFramework/Private/MLDeformerModel.cpp
--- a/Engine/Plugins/Animation/MLDeformer/MLDeformerFramework/Source/MLDeformerFramework/Private/MLDeformerModel.cpp
+++ b/Engine/Plugins/Animation/MLDeformer/MLDeformerFramework/Source/MLDeformerFramework/Private/MLDeformerModel.cpp
@@ -66,6 +66,12 @@ void UMLDeformerModel::Serialize(FArchive& Archive)
 	#if WITH_EDITOR
 		if (Archive.IsSaving() && Archive.IsPersistent())
 		{
+			if (Archive.IsCooking())
+			{
+				AnimSequence = nullptr;
+				VizSettings = nullptr;
+			}
+
 			InitVertexMap();
 			UpdateCachedNumVertices();
 		}
@@ -108,20 +114,21 @@ void UMLDeformerModel::PostLoad()
 		InputInfo->OnPostLoad();
 	}
 
-	if (NeuralNetwork)
+	UNeuralNetwork* Network = GetNeuralNetwork();
+	if (Network)
 	{
 		// If we run the neural network on the GPU.
 		if (IsNeuralNetworkOnGPU())
 		{
-			NeuralNetwork->SetDeviceType(ENeuralDeviceType::GPU, ENeuralDeviceType::CPU, ENeuralDeviceType::GPU);
-			if (NeuralNetwork->GetDeviceType() != ENeuralDeviceType::GPU || NeuralNetwork->GetOutputDeviceType() != ENeuralDeviceType::GPU || NeuralNetwork->GetInputDeviceType() != ENeuralDeviceType::CPU)
+			Network->SetDeviceType(ENeuralDeviceType::GPU, ENeuralDeviceType::CPU, ENeuralDeviceType::GPU);
+			if (Network->GetDeviceType() != ENeuralDeviceType::GPU || Network->GetOutputDeviceType() != ENeuralDeviceType::GPU || Network->GetInputDeviceType() != ENeuralDeviceType::CPU)
 			{
 				UE_LOG(LogMLDeformer, Error, TEXT("Neural net in ML Deformer '%s' cannot run on the GPU, it will not be active."), *GetDeformerAsset()->GetName());
 			}
 		}
 		else // We run our neural network on the CPU.
 		{
-			NeuralNetwork->SetDeviceType(ENeuralDeviceType::CPU, ENeuralDeviceType::CPU, ENeuralDeviceType::CPU);
+			Network->SetDeviceType(ENeuralDeviceType::CPU, ENeuralDeviceType::CPU, ENeuralDeviceType::CPU);
 		}
 	}
 
@@ -139,14 +146,9 @@ void UMLDeformerModel::SetNeuralNetwork(UNeuralNetwork* InNeuralNetwork)
 // Used for the FBoenReference, so it knows what skeleton to pick bones from.
 USkeleton* UMLDeformerModel::GetSkeleton(bool& bInvalidSkeletonIsError, const IPropertyHandle* PropertyHandle)
 {
-	#if WITH_EDITOR
-		bInvalidSkeletonIsError = false;
-		if (SkeletalMesh)
-		{
-			return SkeletalMesh->GetSkeleton();
-		}
-	#endif
-	return nullptr;
+	bInvalidSkeletonIsError = false;
+	USkeletalMesh* SkelMesh = GetSkeletalMesh();
+	return (SkelMesh != nullptr) ? SkelMesh->GetSkeleton() : nullptr;
 }
 
 void UMLDeformerModel::BeginDestroy()
@@ -186,7 +188,7 @@ void UMLDeformerModel::FloatArrayToVector3Array(const TArray<float>& FloatArray,
 #if WITH_EDITOR
 	void UMLDeformerModel::UpdateNumBaseMeshVertices()
 	{
-		NumBaseMeshVerts = UMLDeformerModel::ExtractNumImportedSkinnedVertices(SkeletalMesh);
+		NumBaseMeshVerts = UMLDeformerModel::ExtractNumImportedSkinnedVertices(GetSkeletalMesh());
 	}
 
 	void UMLDeformerModel::PostEditChangeProperty(FPropertyChangedEvent& PropertyChangedEvent)
@@ -211,12 +213,24 @@ void UMLDeformerModel::FloatArrayToVector3Array(const TArray<float>& FloatArray,
 		return SkeletalMesh ? SkeletalMesh->GetNumImportedVertices() : 0;
 	}
 
-	void UMLDeformerModel::SetAssetEditorOnlyFlags()
+	void UMLDeformerModel::MarkObjectAsEditorOnly(UObject* InObject, bool bLogMessage)
 	{
-		if (AnimSequence)
+		if (InObject == nullptr)
 		{
-			AnimSequence->GetPackage()->SetPackageFlags(PKG_EditorOnly);
+			return;
 		}
+
+		InObject->GetPackage()->SetPackageFlags(PKG_EditorOnly);
+
+		if (bLogMessage)
+		{
+			UE_LOG(LogMLDeformer, Display, TEXT("Marking '%s' as editor only asset. Asset will be excluded from packaging."), *InObject->GetName());
+		}
+	}
+
+	void UMLDeformerModel::SetAssetEditorOnlyFlags()
+	{
+		MarkObjectAsEditorOnly(GetAnimSequence());
 	}
 #endif	// #if WITH_EDITOR
 
@@ -224,10 +238,11 @@ void UMLDeformerModel::FloatArrayToVector3Array(const TArray<float>& FloatArray,
 	void UMLDeformerModel::InitVertexMap()
 	{
 		VertexMap.Empty();
-		if (SkeletalMesh)
+		const USkeletalMesh* SkelMesh = GetSkeletalMesh();
+		if (SkelMesh)
 		{
-			FSkeletalMeshModel* SkeletalMeshModel = SkeletalMesh->GetImportedModel();
-			if (SkeletalMeshModel)
+			const FSkeletalMeshModel* SkeletalMeshModel = SkelMesh->GetImportedModel();
+			if (SkeletalMeshModel && SkeletalMeshModel->LODModels.IsValidIndex(0))
 			{
 				VertexMap = SkeletalMeshModel->LODModels[0].MeshToImportVertexMap;
 			}
diff --git a/Engine/Plugins/Animation/MLDeformer/MLDeformerFramework/Source/MLDeformerFramework/Private/MLDeformerModelInstance.cpp b/Engine/Plugins/Animation/MLDeformer/MLDeformerFramework/Source/MLDeformerFramework/Private/MLDeformerModelInstance.cpp
--- a/Engine/Plugins/Animation/MLDeformer/MLDeformerFramework/Source/MLDeformerFramework/Private/MLDeformerModelInstance.cpp
+++ b/Engine/Plugins/Animation/MLDeformer/MLDeformerFramework/Source/MLDeformerFramework/Private/MLDeformerModelInstance.cpp
@@ -152,32 +152,6 @@ FString UMLDeformerModelInstance::CheckCompatibility(USkeletalMeshComponent* InS
 				UE_LOG(LogMLDeformer, Error, TEXT("Deformer '%s': %s"), *(Model->GetDeformerAsset()->GetName()), *InputErrorString);
 			}
 		}
-
-		// Check if the neural network is on the right device, if not, we can't continue.
-		if (Model->IsNeuralNetworkOnGPU())
-		{
-			if (NeuralNetwork->GetDeviceType() != ENeuralDeviceType::GPU || NeuralNetwork->GetOutputDeviceType() != ENeuralDeviceType::GPU || NeuralNetwork->GetInputDeviceType() != ENeuralDeviceType::CPU)
-			{
-				const FString DeviceErrorString = "The neural network is expected to run and output on the GPU, but it isn't.";
-				ErrorText += DeviceErrorString + "\n";
-				if (LogIssues)
-				{
-					UE_LOG(LogMLDeformer, Error, TEXT("Deformer '%s': %s"), *(Model->GetDeformerAsset()->GetName()), *DeviceErrorString);
-				}
-			}
-		}
-		else
-		{
-			if (NeuralNetwork->GetDeviceType() != ENeuralDeviceType::CPU || NeuralNetwork->GetOutputDeviceType() != ENeuralDeviceType::CPU || NeuralNetwork->GetInputDeviceType() != ENeuralDeviceType::CPU)
-			{
-				const FString DeviceErrorString = "The neural network is expected to run fully on the CPU, but it isn't.";
-				ErrorText += DeviceErrorString + "\n";
-				if (LogIssues)
-				{
-					UE_LOG(LogMLDeformer, Error, TEXT("Deformer '%s': %s"), *(Model->GetDeformerAsset()->GetName()), *DeviceErrorString);
-				}
-			}
-		}
 	}
 
 	return ErrorText;
diff --git a/Engine/Plugins/Animation/MLDeformer/MLDeformerFramework/Source/MLDeformerFramework/Private/MLDeformerMorphModelInstance.cpp b/Engine/Plugins/Animation/MLDeformer/MLDeformerFramework/Source/MLDeformerFramework/Private/MLDeformerMorphModelInstance.cpp
--- a/Engine/Plugins/Animation/MLDeformer/MLDeformerFramework/Source/MLDeformerFramework/Private/MLDeformerMorphModelInstance.cpp
+++ b/Engine/Plugins/Animation/MLDeformer/MLDeformerFramework/Source/MLDeformerFramework/Private/MLDeformerMorphModelInstance.cpp
@@ -5,6 +5,26 @@
 #include "NeuralNetwork.h"
 #include "Components/SkeletalMeshComponent.h"
 
+void UMLDeformerMorphModelInstance::Release()
+{
+	Super::Release();
+
+	if (SkeletalMeshComponent == nullptr)
+	{
+		return;
+	}
+
+	const UMLDeformerMorphModel* MorphModel = Cast<UMLDeformerMorphModel>(Model);
+	if (MorphModel == nullptr)
+	{
+		return;
+	}
+
+	const int32 LOD = 0;
+	SkeletalMeshComponent->RemoveExternalMorphSet(LOD, MorphModel->GetExternalMorphSetID());
+	SkeletalMeshComponent->RefreshExternalMorphTargetWeights();
+}
+
 // Run the neural network, which calculates its outputs, which are the weights of our morph targets.
 void UMLDeformerMorphModelInstance::RunNeuralNetwork(float ModelWeight)
 {
diff --git a/Engine/Plugins/Animation/MLDeformer/MLDeformerFramework/Source/MLDeformerFramework/Private/Tests/MLDeformerTestModel.cpp b/Engine/Plugins/Animation/MLDeformer/MLDeformerFramework/Source/MLDeformerFramework/Private/Tests/MLDeformerTestModel.cpp
--- a/Engine/Plugins/Animation/MLDeformer/MLDeformerFramework/Source/MLDeformerFramework/Private/Tests/MLDeformerTestModel.cpp
+++ b/Engine/Plugins/Animation/MLDeformer/MLDeformerFramework/Source/MLDeformerFramework/Private/Tests/MLDeformerTestModel.cpp
@@ -17,6 +17,6 @@ FString UMLDeformerTestModel::GetDisplayName() const
 #if WITH_EDITOR
 	void UMLDeformerTestModel::UpdateNumTargetMeshVertices()
 	{
-		NumTargetMeshVerts = NumBaseMeshVerts;
+		SetNumTargetMeshVerts(GetNumBaseMeshVerts());
 	}
 #endif
diff --git a/Engine/Plugins/Animation/MLDeformer/MLDeformerFramework/Source/MLDeformerFramework/Public/MLDeformerComponent.h b/Engine/Plugins/Animation/MLDeformer/MLDeformerFramework/Source/MLDeformerFramework/Public/MLDeformerComponent.h
--- a/Engine/Plugins/Animation/MLDeformer/MLDeformerFramework/Source/MLDeformerFramework/Public/MLDeformerComponent.h
+++ b/Engine/Plugins/Animation/MLDeformer/MLDeformerFramework/Source/MLDeformerFramework/Public/MLDeformerComponent.h
@@ -16,6 +16,9 @@ class UMLDeformerModelInstance;
  * The ML mesh deformer component.
  * This works in combination with a MLDeformerAsset and SkeletalMeshComponent.
  * The component will perform runtime inference of the deformer model setup inside the asset.
+ * When you have multiple skeletal mesh components on your actor, this component will try to use the skeletal mesh component that uses
+ * the same skeletal mesh as the applied deformer was trained on.
+ * If it cannot find that, it will use the first skeletal mesh component it finds.
  */
 UCLASS(Blueprintable, ClassGroup = Component, BlueprintType, meta = (BlueprintSpawnableComponent))
 class MLDEFORMERFRAMEWORK_API UMLDeformerComponent
@@ -38,14 +41,9 @@ public:
 
 	/** 
 	 * Setup the ML Deformer, by picking the deformer asset and skeletal mesh component. 
-	 * Call this when you want to assign a specific ML Deformer to a specific skeletal mesh component.
-	 * This is used in case there are multiple skeletal mesh components on your actor.
-	 * On default the first skeletal mesh component it finds will be used. This function allows you to change that behavior.
-	 * Keep in mind that the ML Deformer asset has to be trained using the same Skeletal Mesh as used by the skeletal mesh component you pick.
 	 * @param InDeformerAsset The ML Deformer asset to apply to the specified skeletal mesh component.	 
 	 * @param InSkelMeshComponent The skeletal mesh component to apply the specified deformer to.
 	 */
-	UFUNCTION(BlueprintCallable, Category = "MLDeformer")
 	void SetupComponent(UMLDeformerAsset* InDeformerAsset, USkeletalMeshComponent* InSkelMeshComponent);
 
 	/** 
diff --git a/Engine/Plugins/Animation/MLDeformer/MLDeformerFramework/Source/MLDeformerFramework/Public/MLDeformerGeomCacheModel.h b/Engine/Plugins/Animation/MLDeformer/MLDeformerFramework/Source/MLDeformerFramework/Public/MLDeformerGeomCacheModel.h
--- a/Engine/Plugins/Animation/MLDeformer/MLDeformerFramework/Source/MLDeformerFramework/Public/MLDeformerGeomCacheModel.h
+++ b/Engine/Plugins/Animation/MLDeformer/MLDeformerFramework/Source/MLDeformerFramework/Public/MLDeformerGeomCacheModel.h
@@ -7,12 +7,13 @@
 #include "MLDeformerVizSettings.h"
 #include "MLDeformerInputInfo.h"
 #include "MLDeformerGeomCacheHelpers.h"
+#include "GeometryCache.h"
 #include "UObject/ObjectPtr.h"
+#include "UObject/SoftObjectPtr.h"
 #include "MLDeformerGeomCacheModel.generated.h"
 
 class UMLDeformerAsset;
 class UMLDeformerGeomCacheVizSettings;
-class UGeometryCache;
 
 /**
  * A ML Deformer model that has a geometry cache as target mesh.
@@ -25,10 +26,14 @@ class MLDEFORMERFRAMEWORK_API UMLDeformerGeomCacheModel
 	GENERATED_BODY()
 
 public:
+	// UObject overrides.
+	virtual void Serialize(FArchive& Archive) override;
+	// ~END UObject overrides.
+
 	// UMLDeformerModel overrides.
-	virtual FString GetDisplayName() const override { return "GeomCache Based Model"; }
+	virtual FString GetDisplayName() const override			{ return "GeomCache Based Model"; }
 #if WITH_EDITORONLY_DATA
-	virtual bool HasTrainingGroundTruth() const override { return (GeometryCache.Get() != nullptr); }
+	virtual bool HasTrainingGroundTruth() const override	{ return (GetGeometryCache() != nullptr); }
 	virtual void SampleGroundTruthPositions(float SampleTime, TArray<FVector3f>& OutPositions) override;
 #endif
 #if WITH_EDITOR
@@ -49,14 +54,14 @@ public:
 	 * This is our training target.
 	 * @return A pointer to the geometry cache, in read-only mode.
 	 */
-	const UGeometryCache* GetGeometryCache() const { return GeometryCache; }
+	const UGeometryCache* GetGeometryCache() const			{ return GeometryCache.LoadSynchronous(); }
 
 	/** 
 	 * Get the geometry cache that represents the target deformation.
 	 * This is our training target.
 	 * @return A pointer to the geometry cache.
 	 */
-	UGeometryCache* GetGeometryCache() { return GeometryCache; }
+	UGeometryCache* GetGeometryCache()						{ return GeometryCache.LoadSynchronous(); }
 
 	/**
 	 * Get the mapping between geometry cache tracks and meshes inside the skeletal mesh.
@@ -64,18 +69,18 @@ public:
 	 * Once we have that, we can calculate deltas between the two.
 	 * @return The geometry cache mesh to skeletal mesh mappings.
 	 */
-	TArray<UE::MLDeformer::FMLDeformerGeomCacheMeshMapping>& GetGeomCacheMeshMappings() { return MeshMappings; }
+	TArray<UE::MLDeformer::FMLDeformerGeomCacheMeshMapping>& GetGeomCacheMeshMappings()				{ return MeshMappings; }
 	const TArray<UE::MLDeformer::FMLDeformerGeomCacheMeshMapping>& GetGeomCacheMeshMappings() const { return MeshMappings; }
 
 	// Get property names.
-	static FName GetGeometryCachePropertyName() { return GET_MEMBER_NAME_CHECKED(UMLDeformerGeomCacheModel, GeometryCache); }
+	static FName GetGeometryCachePropertyName()				{ return GET_MEMBER_NAME_CHECKED(UMLDeformerGeomCacheModel, GeometryCache); }
 
-protected:
+private:
 	/** The mappings between the geometry cache tracks and skeletal mesh imported meshes. */
 	TArray<UE::MLDeformer::FMLDeformerGeomCacheMeshMapping> MeshMappings;
 
 	/** The geometry cache that represents the target deformations. */
 	UPROPERTY(EditAnywhere, Category = "Target Mesh")
-	TObjectPtr<UGeometryCache> GeometryCache = nullptr;
+	TSoftObjectPtr<UGeometryCache> GeometryCache = nullptr;
 #endif // WITH_EDITORONLY_DATA
 };
diff --git a/Engine/Plugins/Animation/MLDeformer/MLDeformerFramework/Source/MLDeformerFramework/Public/MLDeformerGeomCacheVizSettings.h b/Engine/Plugins/Animation/MLDeformer/MLDeformerFramework/Source/MLDeformerFramework/Public/MLDeformerGeomCacheVizSettings.h
--- a/Engine/Plugins/Animation/MLDeformer/MLDeformerFramework/Source/MLDeformerFramework/Public/MLDeformerGeomCacheVizSettings.h
+++ b/Engine/Plugins/Animation/MLDeformer/MLDeformerFramework/Source/MLDeformerFramework/Public/MLDeformerGeomCacheVizSettings.h
@@ -4,11 +4,11 @@
 
 #include "CoreMinimal.h"
 #include "UObject/ObjectPtr.h"
+#include "UObject/SoftObjectPtr.h"
+#include "GeometryCache.h"
 #include "MLDeformerVizSettings.h"
 #include "MLDeformerGeomCacheVizSettings.generated.h"
 
-class UGeometryCache;
-
 /**
  * The vizualization settings for a model that has a geometry cache.
  * This can be used in combination with a UMLDeformerGeomCacheModel.
@@ -22,14 +22,14 @@ class MLDEFORMERFRAMEWORK_API UMLDeformerGeomCacheVizSettings
 public:
 #if WITH_EDITORONLY_DATA
 	// UMLDeformerModel overrides.
-	virtual bool HasTestGroundTruth() const override	{ return !GroundTruth.IsNull(); }
+	virtual bool HasTestGroundTruth() const override	{ return (GetTestGroundTruth() != nullptr); }
 	// ~END UMLDeformerModel overrides.
 
 	/**
 	 * Get the test ground truth geometry cache, which represents the ground truth version of the test animation sequence.
 	 * @return A pointer to the geometry cache object.
 	 */
-	UGeometryCache* GetTestGroundTruth() const			{ return GroundTruth; }
+	UGeometryCache* GetTestGroundTruth() const			{ return GroundTruth.LoadSynchronous(); }
 
 	// Get property names.
 	static FName GetTestGroundTruthPropertyName()		{ return GET_MEMBER_NAME_CHECKED(UMLDeformerGeomCacheVizSettings, GroundTruth); }
@@ -37,6 +37,6 @@ public:
 protected:
 	/** The geometry cache that represents the ground truth of the test anim sequence. */
 	UPROPERTY(EditAnywhere, Category = "Test Assets")
-	TObjectPtr<UGeometryCache> GroundTruth = nullptr;
+	TSoftObjectPtr<UGeometryCache> GroundTruth = nullptr;
 #endif // WITH_EDITORONLY_DATA
 };
diff --git a/Engine/Plugins/Animation/MLDeformer/MLDeformerFramework/Source/MLDeformerFramework/Public/MLDeformerModel.h b/Engine/Plugins/Animation/MLDeformer/MLDeformerFramework/Source/MLDeformerFramework/Public/MLDeformerModel.h
--- a/Engine/Plugins/Animation/MLDeformer/MLDeformerFramework/Source/MLDeformerFramework/Public/MLDeformerModel.h
+++ b/Engine/Plugins/Animation/MLDeformer/MLDeformerFramework/Source/MLDeformerFramework/Public/MLDeformerModel.h
@@ -5,10 +5,13 @@
 #include "CoreMinimal.h"
 #include "UObject/ObjectPtr.h"
 #include "UObject/Object.h"
+#include "UObject/SoftObjectPtr.h"
 #include "Interfaces/Interface_BoneReferenceSkeletonProvider.h"
 #include "BoneContainer.h"
 #include "RenderCommandFence.h"
 #include "RenderResource.h"
+#include "NeuralNetwork.h"
+#include "Animation/AnimSequence.h"
 #include "MLDeformerCurveReference.h"
 #include "MLDeformerModel.generated.h"
 
@@ -17,9 +20,6 @@ class UMLDeformerVizSettings;
 class UMLDeformerModelInstance;
 class UMLDeformerComponent;
 class UMLDeformerInputInfo;
-class UAnimSequence;
-class UGeometryCache;
-class UNeuralNetwork;
 
 namespace UE::MLDeformer
 {
@@ -139,6 +139,21 @@ public:
 	 */
 	virtual FString GetDefaultDeformerGraphAssetPath() const	{ return FString(); }
 
+	/**
+	 * Get the skeletal mesh that is used during training.
+	 * You typically want to apply the ML Deformer on this specific skeletal mesh in your game as well.
+	 * @return A pointer to the skeletal mesh.
+	 */
+	const USkeletalMesh* GetSkeletalMesh() const			{ return SkeletalMesh.Get(); }
+	USkeletalMesh* GetSkeletalMesh()						{ return SkeletalMesh.Get(); }
+
+	/**
+	 * Set the skeletal mesh that this deformer uses.
+	 * @param SkelMesh The skeletal mesh.
+	 */
+	void SetSkeletalMesh(USkeletalMesh* SkelMesh)			{ SkeletalMesh = SkelMesh; }
+
+
 #if WITH_EDITORONLY_DATA
 	/**
 	 * Check whether this model currently has a training mesh setup or not.
@@ -156,6 +171,14 @@ public:
 #endif
 
 #if WITH_EDITOR
+	/**
+	 * Set an object to be editor only.
+	 * This will add the PKG_EditorOnly package flag. This will cause it to not be included during packaging of the project.
+	 * @param InObject The object to mark as editor only.
+	 * @param bLogMessage Set to true when you wish to log a message about this marking.
+	 */
+	void MarkObjectAsEditorOnly(UObject* InObject, bool bLogMessage=true);
+
 	/**
 	 * Update the cached number of vertices of both base and target meshes.
 	 */ 
@@ -217,21 +240,21 @@ public:
 	 * Inputs are things like bone transforms and curve values.
 	 * @return A pointer to the input info object.
 	 */
-	UMLDeformerInputInfo* GetInputInfo() const				{ return InputInfo.Get(); }
+	UMLDeformerInputInfo* GetInputInfo() const					{ return InputInfo.Get(); }
 
 	/**
 	 * Get the number of vertices in the base mesh (linear skinned skeletal mesh).
 	 * This is the number of vertices in the DCC, so not the render vertex count.
 	 * @return The number of imported vertices inside linear skinned skeletal mesh.
 	 */
-	int32 GetNumBaseMeshVerts() const						{ return NumBaseMeshVerts; }
+	int32 GetNumBaseMeshVerts() const							{ return NumBaseMeshVerts; }
 
 	/**
 	 * Get the number of vertices of the target mesh.
 	 * This is the number of vertices in the DCC, so not the render vertex count.
 	 * @return The number of imported vertices inside the target mesh.
 	 */
-	int32 GetNumTargetMeshVerts() const						{ return NumTargetMeshVerts; }
+	int32 GetNumTargetMeshVerts() const							{ return NumTargetMeshVerts; }
 
 	/**
 	 * The mapping that maps from render vertices into dcc vertices.
@@ -239,7 +262,14 @@ public:
 	 * For a cube with 32 render vertices, the item values would be between 0..7 as in the dcc the cube has 8 vertices.
 	 * @return A reference to the array that contains the DCC vertex number for each render vertex.
 	 */
-	const TArray<int32>& GetVertexMap() const				{ return VertexMap; }
+	const TArray<int32>& GetVertexMap() const					{ return VertexMap; }
+
+	/**
+	 * Manually set the vertex map. This normally gets initialized automatically.
+	 * @param Map The original vertex number indices for each render vertex.
+	 * @see GetVertexMap.
+	 */
+	void SetVertexMap(const TArray<int32>& Map)					{ VertexMap = Map; }
 
 	/**
 	 * Get the GPU buffer of the VertexMap.
@@ -253,7 +283,7 @@ public:
 	 * This network is used during inference.
 	 * @return A pointer to the neural network, or nullptr when the network has not yet been trained.
 	 */
-	UNeuralNetwork* GetNeuralNetwork() const				{ return NeuralNetwork.Get(); }
+	UNeuralNetwork* GetNeuralNetwork() const					{ return NeuralNetwork.Get(); }
 
 	/**
 	 * Set the neural network object that we use during inference.
@@ -285,14 +315,28 @@ public:
 	 * Check whether we should include bone transforms as input to the model during training or not.
 	 * @return Returns true when bone transfomations should be a part of the network inputs, during the training process.
 	 */
-	bool ShouldIncludeBonesInTraining() const				{ return bIncludeBones; }
+	bool ShouldIncludeBonesInTraining() const					{ return bIncludeBones; }
+
+	/**
+	 * Set whether we want to include bones during training or not.
+	 * This will make bone transforms part of the neural network inputs.
+	 * @param bInclude Set to true if you wish bone transforms to be included during training and at inference time.
+	 */
+	void SetShouldIncludeBonesInTraining(bool bInclude)			{ bIncludeBones = bInclude; }
 
 	/**
 	 * Check whether we should include curve values as input to the model during training or not.
 	 * Curve values are single floats.
 	 * @return Returns true when curve values should be a part of the network inputs, during the training process.
 	 */
-	bool ShouldIncludeCurvesInTraining() const				{ return bIncludeCurves; }
+	bool ShouldIncludeCurvesInTraining() const					{ return bIncludeCurves; }
+
+	/**
+	 * Set whether we want to include curves during training.
+	 * This will make curves part of the neural network inputs.
+	 * @param bInclude Set to true to include curves during training and inference time.
+	 */
+	void SetShouldIncludeCurvesInTraining(bool bInclude)		{ bIncludeCurves = bInclude; }
 
 	/**
 	 * The delegate that gets fired when a property value changes.
@@ -307,29 +351,15 @@ public:
 	 * @return A pointer to the visualization settings. You can cast this to the type specific for your model, in case you inherited from 
 	 *         the UMLDeformerVizSettings base class. This never return a nullptr.
 	 */
-	UMLDeformerVizSettings* GetVizSettings() const			{ return VizSettings; }
-
-	/**
-	 * Get the skeletal mesh that is used during training.
-	 * You typically want to apply the ML Deformer on this specific skeletal mesh in your game as well.
-	 * @return A pointer to the skeletal mesh.
-	 */
-	const USkeletalMesh* GetSkeletalMesh() const			{ return SkeletalMesh;  }
-	USkeletalMesh* GetSkeletalMesh()						{ return SkeletalMesh; }
-
-	/**
-	 * Set the skeletal mesh that this deformer uses.
-	 * @param SkelMesh The skeletal mesh.
-	 */
-	void SetSkeletalMesh(USkeletalMesh* SkelMesh)			{ SkeletalMesh = SkelMesh; }
+	UMLDeformerVizSettings* GetVizSettings() const				{ return VizSettings; }
 
 	/**
 	 * Get the animation sequence that is used during training.
 	 * Each frame of this anim sequence will contain a training pose.
 	 * @return A pointer to the animation sequence used for training.
 	 */
-	const UAnimSequence* GetAnimSequence() const			{ return AnimSequence;  }
-	UAnimSequence* GetAnimSequence()						{ return AnimSequence; }
+	const UAnimSequence* GetAnimSequence() const				{ return AnimSequence.LoadSynchronous();  }
+	UAnimSequence* GetAnimSequence()							{ return AnimSequence.LoadSynchronous(); }
 
 	/**
 	 * Get the maximum number of training frames to use during training.
@@ -337,7 +367,14 @@ public:
 	 * want to train on only 2000 frames instead. You can do this by setting the maximum training frames to 2000.
 	 * @return The max number of frames to use during training.
 	 */
-	int32 GetTrainingFrameLimit() const						{ return MaxTrainingFrames; }
+	int32 GetTrainingFrameLimit() const							{ return MaxTrainingFrames; }
+
+	/**
+	 * Set the maximum number of frames to train on.
+	 * For example if your training data has 10000 frames, but you wish to only train on 2000 frames, you can set this to 2000.
+	 * @param MaxNumFrames The maximum number of frames to train on.
+	 */
+	void SetTrainingFrameLimit(int32 MaxNumFrames)				{ MaxTrainingFrames = MaxNumFrames; }
 
 	/**
 	 * Get the target mesh alignment tranformation.
@@ -347,7 +384,14 @@ public:
 	 * The alignment transform is then used to correct this and align both base and target mesh.
 	 * @return The alignment transformation. When set to Identity it will not do anything, which is its default.
 	 */
-	const FTransform& GetAlignmentTransform() const			{ return AlignmentTransform; }
+	const FTransform& GetAlignmentTransform() const				{ return AlignmentTransform; }
+
+	/**
+	 * Set the alignment transform, which is the transform applied to the target mesh vertices, before calculating the deltas.
+	 * @param Transform The transformation to apply.
+	 * @see GetAlignmentTransform.
+	 */
+	void SetAlignmentTransform(const FTransform& Transform)		{ AlignmentTransform = Transform; }
 
 	/**
 	 * Get the list of bones that we configured to be included during training.
@@ -362,6 +406,14 @@ public:
 	TArray<FBoneReference>& GetBoneIncludeList()				{ return BoneIncludeList; }
 	const TArray<FBoneReference>& GetBoneIncludeList() const	{ return BoneIncludeList; }
 
+	/**
+	 * Set the list of bones that should be included during training.
+	 * This list is ignored if ShouldIncludeBonesInTraining() returns false.
+	 * If the list is empty, all bones will be included.
+	 * @param List The list of bones to include.
+	 */
+	void SetBoneIncludeList(const TArray<FBoneReference>& List)	{ BoneIncludeList = List; }
+
 	/**
 	 * Get the list of curves that we configured to be included during training.
 	 * A curve reference is basically just a name of a curve.
@@ -375,13 +427,36 @@ public:
 	TArray<FMLDeformerCurveReference>& GetCurveIncludeList()				{ return CurveIncludeList; }
 	const TArray<FMLDeformerCurveReference>& GetCurveIncludeList() const	{ return CurveIncludeList; }
 
+	/**
+	 * Set the list of curves that should be included during training.
+	 * This list is ignored if ShouldIncludeCurvesInTraining() returns false.
+	 * If the list is empty, all curves will be included.
+	 * @param List The list of curves to include.
+	 */
+	void SetCurveIncludeList(const TArray<FMLDeformerCurveReference>& List)	{ CurveIncludeList = List; }
+
 	/**
 	 * Get the delta cutoff length. Deltas that have a length larger than this length will be set to zero.
 	 * This can be useful when there are some vertices that due to incorrect data have a very long length.
 	 * Skipping those deltas will prevent issues.
 	 * @return The length after which deltas will be ignored. So anything delta length larger than this value will be ignored.
 	 */
-	float GetDeltaCutoffLength() const					{ return DeltaCutoffLength; }
+	float GetDeltaCutoffLength() const										{ return DeltaCutoffLength; }
+
+	/**
+	 * Set the delta cutoff length. Deltas that are larger than this length will be set to zero.
+	 * This can be useful when there are some vertices that due to incorrect data have a very long length.
+	 * Skipping those deltas will prevent issues.
+	 * @param Length The new delta cutoff length.
+	 */
+	void SetDeltaCutoffLength(float Length)									{ DeltaCutoffLength = Length; }
+
+	/**
+	 * Set the visualization settings object.
+	 * You need to call this in the constructor of your model.
+	 * @param VizSettings The visualization settings object for this model.
+	 */
+	void SetVizSettings(UMLDeformerVizSettings* VizSettingsObject)			{ VizSettings = VizSettingsObject; }
 
 	// Get property names.
 	static FName GetSkeletalMeshPropertyName()			{ return GET_MEMBER_NAME_CHECKED(UMLDeformerModel, SkeletalMesh); }
@@ -413,7 +488,21 @@ protected:
 	 */
 	void FloatArrayToVector3Array(const TArray<float>& FloatArray, TArray<FVector3f>& OutVectorArray);
 
-protected:
+	/**
+	 * Set the number of vertices in the base mesh.
+	 * This is the number of imported (dcc) vertices, so not render vertices.
+	 * @param NumVerts The number of vertices.
+	 */
+	void SetNumBaseMeshVerts(int32 NumVerts)			{ NumBaseMeshVerts = NumVerts; }
+
+	/**
+	 * Set the number of vertices in the target mesh.
+	 * This is the number of imported (dcc) vertices, so not render vertices.
+	 * @param NumVerts The number of vertices.
+	 */
+	void SetNumTargetMeshVerts(int32 NumVerts)			{ NumTargetMeshVerts = NumVerts; }
+
+private:
 	/** The deformer asset that this model is part of. */
 	TObjectPtr<UMLDeformerAsset> DeformerAsset = nullptr;
 
@@ -451,6 +540,10 @@ protected:
 	UPROPERTY()
 	TObjectPtr<UNeuralNetwork> NeuralNetwork = nullptr;
 
+	/** The skeletal mesh that represents the linear skinned mesh. */
+	UPROPERTY(EditAnywhere, Category = "Base Mesh")
+	TObjectPtr<USkeletalMesh> SkeletalMesh = nullptr;
+
 #if WITH_EDITORONLY_DATA
 	UPROPERTY()
 	TObjectPtr<UMLDeformerVizSettings> VizSettings = nullptr;
@@ -463,19 +556,15 @@ protected:
 	UPROPERTY(EditAnywhere, Category = "Inputs and Output")
 	bool bIncludeCurves = false;
 
-	/** The skeletal mesh that represents the linear skinned mesh. */
-	UPROPERTY(EditAnywhere, BlueprintReadWrite, Category = "Base Mesh")
-	TObjectPtr<USkeletalMesh> SkeletalMesh = nullptr;
-
 	/**
 	 * The animation sequence to apply to the base mesh. This has to match the animation of the target mesh's geometry cache. 
 	 * Internally we force the Interpolation property for this motion to be "Step".
 	 */
-	UPROPERTY(EditAnywhere, BlueprintReadWrite, Category = "Base Mesh")
-	TObjectPtr<UAnimSequence> AnimSequence = nullptr;
+	UPROPERTY(EditAnywhere, Category = "Base Mesh")
+	TSoftObjectPtr<UAnimSequence> AnimSequence = nullptr;
 
 	/** The transform that aligns the Geometry Cache to the SkeletalMesh. This will mostly apply some scale and a rotation, but no translation. */
-	UPROPERTY(EditAnywhere, BlueprintReadWrite, Category = "Target Mesh")
+	UPROPERTY(EditAnywhere, Category = "Target Mesh")
 	FTransform AlignmentTransform = FTransform::Identity;
 
 	/** The bones to include during training. When none are provided, all bones of the Skeleton will be included. */
@@ -487,14 +576,14 @@ protected:
 	TArray<FMLDeformerCurveReference> CurveIncludeList;
 
 	/** The maximum numer of training frames (samples) to train on. Use this to train on a sub-section of your full training data. */
-	UPROPERTY(EditAnywhere, BlueprintReadWrite, Category = "Inputs and Output", meta = (ClampMin = "1"))
+	UPROPERTY(EditAnywhere, Category = "Inputs and Output", meta = (ClampMin = "1"))
 	int32 MaxTrainingFrames = 1000000;
 
 	/**
 	 * Sometimes there can be some vertices that cause some issues that cause deltas to be very long. We can ignore these deltas by setting a cutoff value. 
 	 * Deltas that are longer than the cutoff value (in units), will be ignored and set to zero length. 
 	 */
-	UPROPERTY(EditAnywhere, BlueprintReadWrite, Category = "Inputs and Output", meta = (ClampMin = "0.01", ForceUnits="cm"))
+	UPROPERTY(EditAnywhere, Category = "Inputs and Output", meta = (ClampMin = "0.01", ForceUnits="cm"))
 	float DeltaCutoffLength = 30.0f;
 #endif
 };
diff --git a/Engine/Plugins/Animation/MLDeformer/MLDeformerFramework/Source/MLDeformerFramework/Public/MLDeformerMorphModel.h b/Engine/Plugins/Animation/MLDeformer/MLDeformerFramework/Source/MLDeformerFramework/Public/MLDeformerMorphModel.h
--- a/Engine/Plugins/Animation/MLDeformer/MLDeformerFramework/Source/MLDeformerFramework/Public/MLDeformerMorphModel.h
+++ b/Engine/Plugins/Animation/MLDeformer/MLDeformerFramework/Source/MLDeformerFramework/Public/MLDeformerMorphModel.h
@@ -35,6 +35,9 @@ public:
 	float GetMorphTargetDeltaThreshold() const				{ return MorphTargetDeltaThreshold; }
 	float GetMorphTargetErrorTolerance() const				{ return MorphTargetErrorTolerance; }
 
+	void SetMorphTargetDeltaThreshold(float Threshold)		{ MorphTargetDeltaThreshold = Threshold; }
+	void SetMorphTargetErrorTolerance(float Tolerance)		{ MorphTargetErrorTolerance = Tolerance; }
+
 	// Get property names.
 	static FName GetMorphTargetDeltaThresholdPropertyName() { return GET_MEMBER_NAME_CHECKED(UMLDeformerMorphModel, MorphTargetDeltaThreshold); }
 	static FName GetMorphTargetErrorTolerancePropertyName() { return GET_MEMBER_NAME_CHECKED(UMLDeformerMorphModel, MorphTargetErrorTolerance); }
@@ -102,7 +105,7 @@ public:
 	 */
 	int32 GetMorphTargetDeltaStartIndex(int32 MorphTargetIndex) const;
 
-protected:
+private:
 	/** The next free morph target set ID. This is used to generate unique ID's for each morph model. */
 	static TAtomic<int32> NextFreeMorphSetID;
 
@@ -117,7 +120,7 @@ protected:
 
 	/** 
 	 * The external morph set data type ID, specific to this model.
-	 * If you inherit your model from this base class, you should set this to some unique value, that represents your model.
+	 * This will get automatically set. Each model instance will get its own unique ID, using the NextFreeMorphSetID atomic that will be increased.
 	 */
 	int32 ExternalMorphSetID = -1;
 
@@ -127,11 +130,11 @@ protected:
 	 * This essentially removes small deltas from morph targets, which will lower the memory usage at runtime, however when set too high it can also introduce visual artifacts.
 	 * A value of 0 will result in the highest quality morph targets, at the cost of higher runtime memory usage.
 	 */
-	UPROPERTY(EditAnywhere, BlueprintReadWrite, Category = "Morph Targets", meta = (ClampMin = "0.0", ClampMax = "1.0", ForceUnits="cm"))
+	UPROPERTY(EditAnywhere, Category = "Morph Targets", meta = (ClampMin = "0.0", ClampMax = "1.0", ForceUnits="cm"))
 	float MorphTargetDeltaThreshold = 0.01f;
 
 	/** The morph target error tolerance. Higher values result in larger compression, but could result in visual artifacts. */
-	UPROPERTY(EditAnywhere, BlueprintReadWrite, Category = "Morph Targets", meta = (ClampMin = "0.01", ClampMax = "500"))
+	UPROPERTY(EditAnywhere, Category = "Morph Targets", meta = (ClampMin = "0.01", ClampMax = "500"))
 	float MorphTargetErrorTolerance = 50.0f;
 #endif // WITH_EDITORONLY_DATA
 };
diff --git a/Engine/Plugins/Animation/MLDeformer/MLDeformerFramework/Source/MLDeformerFramework/Public/MLDeformerMorphModelInstance.h b/Engine/Plugins/Animation/MLDeformer/MLDeformerFramework/Source/MLDeformerFramework/Public/MLDeformerMorphModelInstance.h
--- a/Engine/Plugins/Animation/MLDeformer/MLDeformerFramework/Source/MLDeformerFramework/Public/MLDeformerMorphModelInstance.h
+++ b/Engine/Plugins/Animation/MLDeformer/MLDeformerFramework/Source/MLDeformerFramework/Public/MLDeformerMorphModelInstance.h
@@ -21,6 +21,7 @@ class MLDEFORMERFRAMEWORK_API UMLDeformerMorphModelInstance
 
 public:
 	// UMLDeformerModelInstance overrides.
+	virtual void Release() override;
 	virtual void RunNeuralNetwork(float ModelWeight) override;
 	// ~END UMLDeformerModelInstance overrides.
 };
diff --git a/Engine/Plugins/Animation/MLDeformer/MLDeformerFramework/Source/MLDeformerFramework/Public/MLDeformerVizSettings.h b/Engine/Plugins/Animation/MLDeformer/MLDeformerFramework/Source/MLDeformerFramework/Public/MLDeformerVizSettings.h
--- a/Engine/Plugins/Animation/MLDeformer/MLDeformerFramework/Source/MLDeformerFramework/Public/MLDeformerVizSettings.h
+++ b/Engine/Plugins/Animation/MLDeformer/MLDeformerFramework/Source/MLDeformerFramework/Public/MLDeformerVizSettings.h
@@ -5,6 +5,8 @@
 #include "CoreMinimal.h"
 #include "UObject/Object.h"
 #include "UObject/ObjectPtr.h"
+#include "UObject/SoftObjectPtr.h"
+#include "Animation/MeshDeformer.h"
 #include "MLDeformerVizSettings.generated.h"
 
 class UMeshDeformer;
@@ -67,8 +69,8 @@ public:
 	int32 GetTrainingFrameNumber() const					{ return TrainingFrameNumber; }
 	int32 GetTestingFrameNumber() const						{ return TestingFrameNumber; }
 	float GetAnimPlaySpeed() const							{ return AnimPlaySpeed; }
-	const UAnimSequence* GetTestAnimSequence() const		{ return TestAnimSequence; }
-	UAnimSequence* GetTestAnimSequence()					{ return TestAnimSequence; }
+	const UAnimSequence* GetTestAnimSequence() const		{ return TestAnimSequence.LoadSynchronous(); }
+	UAnimSequence* GetTestAnimSequence()					{ return TestAnimSequence.LoadSynchronous(); }
 	bool GetDrawLinearSkinnedActor() const					{ return bDrawLinearSkinnedActor; }
 	bool GetDrawMLDeformedActor() const						{ return bDrawMLDeformedActor; }
 	bool GetDrawGroundTruthActor() const					{ return bDrawGroundTruthActor; }
@@ -76,7 +78,7 @@ public:
 	EMLDeformerHeatMapMode GetHeatMapMode() const			{ return HeatMapMode; }
 	float GetHeatMapMax() const								{ return HeatMapMax; }
 	float GetGroundTruthLerp() const						{ return HeatMapMode == EMLDeformerHeatMapMode::GroundTruth ? GroundTruthLerp : 0.0f; }
-	UMeshDeformer* GetDeformerGraph() const					{ return DeformerGraph; }
+	UMeshDeformer* GetDeformerGraph() const					{ return DeformerGraph.LoadSynchronous(); }
 	float GetWeight() const									{ return Weight; }
 	bool GetXRayDeltas() const								{ return bXRayDeltas; }
 	bool GetDrawVertexDeltas() const						{ return bDrawDeltas; }
@@ -112,11 +114,11 @@ protected:
 
 	/** The animation sequence to play on the skeletal mesh. */
 	UPROPERTY(EditAnywhere, Category = "Test Assets")
-	TObjectPtr<UAnimSequence> TestAnimSequence = nullptr;
+	TSoftObjectPtr<UAnimSequence> TestAnimSequence = nullptr;
 
 	/** The deformer graph to use on the asset editor's deformed test actor. */
 	UPROPERTY(EditAnywhere, Category = "Test Assets")
-	TObjectPtr<UMeshDeformer> DeformerGraph = nullptr;
+	TSoftObjectPtr<UMeshDeformer> DeformerGraph = nullptr;
 
 	/** The play speed factor of the test anim sequence. */
 	UPROPERTY(EditAnywhere, Category = "Live Settings", meta = (ClampMin = "0.0", ClampMax = "2.0", ForceUnits="Multiplier"))
diff --git a/Engine/Plugins/Animation/MLDeformer/MLDeformerFramework/Source/MLDeformerFrameworkEditor/Private/MLDeformerEditorModel.cpp b/Engine/Plugins/Animation/MLDeformer/MLDeformerFramework/Source/MLDeformerFrameworkEditor/Private/MLDeformerEditorModel.cpp
--- a/Engine/Plugins/Animation/MLDeformer/MLDeformerFramework/Source/MLDeformerFrameworkEditor/Private/MLDeformerEditorModel.cpp
+++ b/Engine/Plugins/Animation/MLDeformer/MLDeformerFramework/Source/MLDeformerFrameworkEditor/Private/MLDeformerEditorModel.cpp
@@ -69,6 +69,8 @@ namespace UE::MLDeformer
 		WorkingRange = TRange<double>(0.0, 100.0);
 		PlaybackRange = TRange<double>(0.0, 100.0);
 
+		AnimSequenceFlags.SetAsset(Model->GetAnimSequence());
+
 		PostEditPropertyDelegateHandle = Model->OnPostEditChangeProperty().AddRaw(this, &FMLDeformerEditorModel::OnPostEditChangeProperty);
 	}
 
@@ -254,8 +256,6 @@ namespace UE::MLDeformer
 
 		PreviewScene->DeselectAll();
 		PreviewScene->SetPreviewAnimationAsset(nullptr);
-		PreviewScene->SetPreviewAnimationBlueprint(nullptr, nullptr);
-		PreviewScene->SetPreviewMesh(nullptr);
 		PreviewScene->SetPreviewMeshComponent(nullptr);
 		PreviewScene->SetActor(nullptr);
 
@@ -481,8 +481,6 @@ namespace UE::MLDeformer
 			SkeletalMeshComponent->SetPlayRate(TestAnimSpeed);
 			SkeletalMeshComponent->Play(true);
 		}
-
-		SetResamplingInputOutputsNeeded(true);
 	}
 
 	void FMLDeformerEditorModel::OnPostInputAssetChanged()
@@ -608,6 +606,7 @@ namespace UE::MLDeformer
 		}
 		else if (Property->GetFName() == UMLDeformerModel::GetAnimSequencePropertyName())
 		{
+			AnimSequenceFlags.SetAsset(Model->GetAnimSequence());
 			TriggerInputAssetChanged();
 			SetResamplingInputOutputsNeeded(true);
 		}
diff --git a/Engine/Plugins/Animation/MLDeformer/MLDeformerFramework/Source/MLDeformerFrameworkEditor/Private/MLDeformerEditorToolkit.cpp b/Engine/Plugins/Animation/MLDeformer/MLDeformerFramework/Source/MLDeformerFrameworkEditor/Private/MLDeformerEditorToolkit.cpp
--- a/Engine/Plugins/Animation/MLDeformer/MLDeformerFramework/Source/MLDeformerFrameworkEditor/Private/MLDeformerEditorToolkit.cpp
+++ b/Engine/Plugins/Animation/MLDeformer/MLDeformerFramework/Source/MLDeformerFrameworkEditor/Private/MLDeformerEditorToolkit.cpp
@@ -636,6 +636,7 @@ namespace UE::MLDeformer
 			ActiveModel->OnInputAssetsChanged();
 			ActiveModel->CreateHeatMapAssets();
 			ActiveModel->SetHeatMapMaterialEnabled(ActiveModel->GetModel()->GetVizSettings()->GetShowHeatMap());
+			ActiveModel->SetResamplingInputOutputsNeeded(true);
 		}
 	}
 
diff --git a/Engine/Plugins/Animation/MLDeformer/MLDeformerFramework/Source/MLDeformerFrameworkEditor/Private/MLDeformerGeomCacheEditorModel.cpp b/Engine/Plugins/Animation/MLDeformer/MLDeformerFramework/Source/MLDeformerFrameworkEditor/Private/MLDeformerGeomCacheEditorModel.cpp
--- a/Engine/Plugins/Animation/MLDeformer/MLDeformerFramework/Source/MLDeformerFrameworkEditor/Private/MLDeformerGeomCacheEditorModel.cpp
+++ b/Engine/Plugins/Animation/MLDeformer/MLDeformerFramework/Source/MLDeformerFrameworkEditor/Private/MLDeformerGeomCacheEditorModel.cpp
@@ -22,6 +22,14 @@
 
 namespace UE::MLDeformer
 {
+	void FMLDeformerGeomCacheEditorModel::Init(const InitSettings& Settings)
+	{
+		FMLDeformerEditorModel::Init(Settings);
+
+		TrainingGroundTruthFlags.SetAsset(GetGeomCacheModel()->GetGeometryCache());
+		TestGroundTruthFlags.SetAsset(GetGeomCacheVizSettings()->GetTestGroundTruth());
+	}
+	
 	FMLDeformerSampler* FMLDeformerGeomCacheEditorModel::CreateSampler() const
 	{
 		FMLDeformerGeomCacheSampler* NewSampler = new FMLDeformerGeomCacheSampler();
@@ -236,11 +244,13 @@ namespace UE::MLDeformer
 		// Properties specific to the geometry cache editor model.
 		if (Property->GetFName() == UMLDeformerGeomCacheModel::GetGeometryCachePropertyName())
 		{
+			TrainingGroundTruthFlags.SetAsset(GetGeomCacheModel()->GetGeometryCache());
 			SetResamplingInputOutputsNeeded(true);
 			TriggerInputAssetChanged(true);
 		}
 		else if (Property->GetFName() == UMLDeformerGeomCacheVizSettings::GetTestGroundTruthPropertyName())
 		{
+			TestGroundTruthFlags.SetAsset(GetGeomCacheVizSettings()->GetTestGroundTruth());
 			TriggerInputAssetChanged(true);
 		}
 	}
diff --git a/Engine/Plugins/Animation/MLDeformer/MLDeformerFramework/Source/MLDeformerFrameworkEditor/Public/MLDeformerEditorModel.h b/Engine/Plugins/Animation/MLDeformer/MLDeformerFramework/Source/MLDeformerFrameworkEditor/Public/MLDeformerEditorModel.h
--- a/Engine/Plugins/Animation/MLDeformer/MLDeformerFramework/Source/MLDeformerFrameworkEditor/Public/MLDeformerEditorModel.h
+++ b/Engine/Plugins/Animation/MLDeformer/MLDeformerFramework/Source/MLDeformerFrameworkEditor/Public/MLDeformerEditorModel.h
@@ -9,6 +9,7 @@
 #include "MLDeformerEditorActor.h"
 #include "Misc/FrameTime.h"
 #include "MLDeformerVizSettings.h"
+#include "MLDeformerModule.h"
 
 class UMLDeformerModel;
 class UMLDeformerInputInfo;
@@ -54,6 +55,44 @@ namespace UE::MLDeformer
 	class FMLDeformerSampler;
 	class FMLDeformerEditorModel;
 
+	/**
+	 * A helper template to keep track of the packaging flags for a given object. 
+	 * We use this to restore some flags of the current asset when user changes the asset to another one.
+	 */
+	template <class T>
+	class FMLDeformerEditorOnlyAssetFlags
+	{
+	public:
+		/** Update the current asset and flags state. */
+		void SetAsset(T* InAsset)
+		{
+			// Restore flags of the current asset, before we change to the new.
+			if (Asset)
+			{
+				if (bEditorOnly)
+				{
+					Asset->GetPackage()->SetPackageFlags(PKG_EditorOnly);
+				}
+				else
+				{
+					Asset->GetPackage()->ClearPackageFlags(PKG_EditorOnly);
+					UE_LOG(LogMLDeformer, Display, TEXT("Unmarked '%s' as editor only asset. Asset can be included in packaging again."), *Asset->GetName());
+				}
+			}
+
+			// Update to the new asset.
+			Asset = InAsset;
+			bEditorOnly = !Asset.IsNull() ? static_cast<bool>(Asset->GetPackage()->GetPackageFlags() & PKG_EditorOnly) : false;
+		}
+
+	private:
+		/** A pointer to the current asset. */
+		TObjectPtr<T> Asset = nullptr;
+
+		/** Was this asset previously marked as editor only already? */
+		bool bEditorOnly = false;
+	};
+
 	/**
 	 * The base class for the editor side of an UMLDeformerModel.
 	 * The editor model class handles most of the editor interactions, such as property changes, triggering the training process,
@@ -832,6 +871,9 @@ namespace UE::MLDeformer
 		/** A pointer to the sampler, which can sample target meshes to calculate deltas. */
 		FMLDeformerSampler* Sampler = nullptr;
 
+		/** The flags state of the anim sequence. */
+		FMLDeformerEditorOnlyAssetFlags<UAnimSequence> AnimSequenceFlags;
+
 		/**
 		 * The input info as currently setup in the editor.
 		 * This is different from the runtime model's input info, as that is the one that was used to train with.
diff --git a/Engine/Plugins/Animation/MLDeformer/MLDeformerFramework/Source/MLDeformerFrameworkEditor/Public/MLDeformerGeomCacheEditorModel.h b/Engine/Plugins/Animation/MLDeformer/MLDeformerFramework/Source/MLDeformerFrameworkEditor/Public/MLDeformerGeomCacheEditorModel.h
--- a/Engine/Plugins/Animation/MLDeformer/MLDeformerFramework/Source/MLDeformerFrameworkEditor/Public/MLDeformerGeomCacheEditorModel.h
+++ b/Engine/Plugins/Animation/MLDeformer/MLDeformerFramework/Source/MLDeformerFrameworkEditor/Public/MLDeformerGeomCacheEditorModel.h
@@ -27,6 +27,7 @@ namespace UE::MLDeformer
 		// ~END FGCObject overrides.
 
 		// FMLDeformerEditorModel overrides.
+		virtual void Init(const InitSettings& Settings) override;
 		virtual FMLDeformerEditorActor* CreateEditorActor(const FMLDeformerEditorActor::FConstructSettings& Settings) const override;
 		virtual FMLDeformerSampler* CreateSampler() const override;
 		virtual void CreateTrainingGroundTruthActor(UWorld* World) override;
@@ -62,5 +63,12 @@ namespace UE::MLDeformer
 		 * @param bIsTrainingActor Set this to true when the actor is an actor to be used in training mode, or set to false when it is to be used in testing mode.
 		 */
 		void CreateGeomCacheActor(UWorld* World, int32 ActorID, const FName& Name, UGeometryCache* GeomCache, FLinearColor LabelColor, FLinearColor WireframeColor, const FText& LabelText, bool bIsTrainingActor);
+
+	protected:
+		/** Packaging flag handling for the training ground truth. */
+		FMLDeformerEditorOnlyAssetFlags<UGeometryCache> TrainingGroundTruthFlags;
+
+		/** Packaging flag handling for the test ground truth. */
+		FMLDeformerEditorOnlyAssetFlags<UGeometryCache> TestGroundTruthFlags;
 	};
 }	// namespace UE::MLDeformer
diff --git a/Engine/Plugins/Animation/MLDeformer/NearestNeighborModel/Source/NearestNeighborModel/Private/NearestNeighborModel.cpp b/Engine/Plugins/Animation/MLDeformer/NearestNeighborModel/Source/NearestNeighborModel/Private/NearestNeighborModel.cpp
--- a/Engine/Plugins/Animation/MLDeformer/NearestNeighborModel/Source/NearestNeighborModel/Private/NearestNeighborModel.cpp
+++ b/Engine/Plugins/Animation/MLDeformer/NearestNeighborModel/Source/NearestNeighborModel/Private/NearestNeighborModel.cpp
@@ -29,17 +29,15 @@ UNearestNeighborModel::UNearestNeighborModel(const FObjectInitializer& ObjectIni
 	: Super(ObjectInitializer)
 {
 #if WITH_EDITORONLY_DATA
-	VizSettings = ObjectInitializer.CreateEditorOnlyDefaultSubobject<UNearestNeighborModelVizSettings>(this, TEXT("VizSettings"));
+	SetVizSettings(ObjectInitializer.CreateEditorOnlyDefaultSubobject<UNearestNeighborModelVizSettings>(this, TEXT("VizSettings")));
 #endif
-
-	ExternalMorphSetID = 100;
 }
 
 UMLDeformerInputInfo* UNearestNeighborModel::CreateInputInfo()
 {
 	UNearestNeighborModelInputInfo* NearestNeighborModelInputInfo = NewObject<UNearestNeighborModelInputInfo>(this);
 #if WITH_EDITORONLY_DATA
-	NearestNeighborModelInputInfo->InitRefBoneRotations(SkeletalMesh);
+	NearestNeighborModelInputInfo->InitRefBoneRotations(GetSkeletalMesh());
 #endif
 	return NearestNeighborModelInputInfo;
 }
@@ -108,7 +106,7 @@ void UNearestNeighborModel::UpdatePCACoeffNums()
 
 void UNearestNeighborModel::UpdateNetworkInputDim()
 {
-	InputDim = 3 * BoneIncludeList.Num();
+	InputDim = 3 * GetBoneIncludeList().Num();
 }
 
 void UNearestNeighborModel::UpdateNetworkOutputDim()
@@ -221,7 +219,7 @@ void UNearestNeighborModel::InitInputInfo()
 	UNearestNeighborModelInputInfo* NearestNeighborModelInputInfo = static_cast<UNearestNeighborModelInputInfo*>(GetInputInfo());
 	check(NearestNeighborModelInputInfo != nullptr);
 #if WITH_EDITORONLY_DATA
-	NearestNeighborModelInputInfo->InitRefBoneRotations(SkeletalMesh);
+	NearestNeighborModelInputInfo->InitRefBoneRotations(GetSkeletalMesh());
 #endif
 }
 
diff --git a/Engine/Plugins/Animation/MLDeformer/NeuralMorphModel/Source/NeuralMorphModel/Private/NeuralMorphModel.cpp b/Engine/Plugins/Animation/MLDeformer/NeuralMorphModel/Source/NeuralMorphModel/Private/NeuralMorphModel.cpp
--- a/Engine/Plugins/Animation/MLDeformer/NeuralMorphModel/Source/NeuralMorphModel/Private/NeuralMorphModel.cpp
+++ b/Engine/Plugins/Animation/MLDeformer/NeuralMorphModel/Source/NeuralMorphModel/Private/NeuralMorphModel.cpp
@@ -26,7 +26,7 @@ UNeuralMorphModel::UNeuralMorphModel(const FObjectInitializer& ObjectInitializer
 	// Never directly create one of the frameworks base classes such as the FMLDeformerMorphModelVizSettings as
 	// that can cause issues with detail customizations.
 #if WITH_EDITORONLY_DATA
-	VizSettings = ObjectInitializer.CreateEditorOnlyDefaultSubobject<UNeuralMorphModelVizSettings>(this, TEXT("VizSettings"));
+	SetVizSettings(ObjectInitializer.CreateEditorOnlyDefaultSubobject<UNeuralMorphModelVizSettings>(this, TEXT("VizSettings")));
 #endif
 }
 
diff --git a/Engine/Plugins/Animation/MLDeformer/NeuralMorphModel/Source/NeuralMorphModel/Public/NeuralMorphModel.h b/Engine/Plugins/Animation/MLDeformer/NeuralMorphModel/Source/NeuralMorphModel/Public/NeuralMorphModel.h
--- a/Engine/Plugins/Animation/MLDeformer/NeuralMorphModel/Source/NeuralMorphModel/Public/NeuralMorphModel.h
+++ b/Engine/Plugins/Animation/MLDeformer/NeuralMorphModel/Source/NeuralMorphModel/Public/NeuralMorphModel.h
@@ -3,6 +3,7 @@
 
 #include "CoreMinimal.h"
 #include "MLDeformerMorphModel.h"
+#include "GeometryCache.h"
 #include "NeuralMorphModel.generated.h"
 
 class USkeletalMesh;
diff --git a/Engine/Plugins/Animation/MLDeformer/VertexDeltaModel/Source/VertexDeltaModel/Private/VertexDeltaModel.cpp b/Engine/Plugins/Animation/MLDeformer/VertexDeltaModel/Source/VertexDeltaModel/Private/VertexDeltaModel.cpp
--- a/Engine/Plugins/Animation/MLDeformer/VertexDeltaModel/Source/VertexDeltaModel/Private/VertexDeltaModel.cpp
+++ b/Engine/Plugins/Animation/MLDeformer/VertexDeltaModel/Source/VertexDeltaModel/Private/VertexDeltaModel.cpp
@@ -22,7 +22,7 @@ UVertexDeltaModel::UVertexDeltaModel(const FObjectInitializer& ObjectInitializer
 	: Super(ObjectInitializer)
 {
 #if WITH_EDITORONLY_DATA
-	VizSettings = ObjectInitializer.CreateEditorOnlyDefaultSubobject<UVertexDeltaModelVizSettings>(this, TEXT("VizSettings"));
+	SetVizSettings(ObjectInitializer.CreateEditorOnlyDefaultSubobject<UVertexDeltaModelVizSettings>(this, TEXT("VizSettings")));
 #endif
 }
 
