[MLDeformer] Exposed some functions to blueprint to be able to set and get deformer assets to use.

#trivial

diff --git a/Engine/Plugins/Animation/MLDeformer/MLDeformerFramework/Source/MLDeformerFramework/Public/MLDeformerComponent.h b/Engine/Plugins/Animation/MLDeformer/MLDeformerFramework/Source/MLDeformerFramework/Public/MLDeformerComponent.h
--- a/Engine/Plugins/Animation/MLDeformer/MLDeformerFramework/Source/MLDeformerFramework/Public/MLDeformerComponent.h
+++ b/Engine/Plugins/Animation/MLDeformer/MLDeformerFramework/Source/MLDeformerFramework/Public/MLDeformerComponent.h
@@ -68,12 +68,14 @@ public:
 	 * Get the ML Deformer asset that is used by this component.
 	 * @return A pointer to the ML Deformer asset.
 	 */
+	UFUNCTION(BlueprintCallable, Category = "MLDeformer")
 	UMLDeformerAsset* GetDeformerAsset() const					{ return DeformerAsset; }
 
 	/**
 	 * Set the deformer asset that is used by this component.
 	 * @param InDeformerAsset A pointer to the deformer asset.
 	 */
+	UFUNCTION(BlueprintCallable, Category = "MLDeformer")
 	void SetDeformerAsset(UMLDeformerAsset* InDeformerAsset)	{ DeformerAsset = InDeformerAsset; }
 
 	/**
@@ -88,6 +90,7 @@ public:
 	 * The skeletal mesh that is setup inside the skeletal mesh component will be the mesh that will be deformed by this ML Deformer component.
 	 * @return A pointer to the skeletal mesh component who's mesh will be deformed by this ML Deformer component.
 	 */
+	UFUNCTION(BlueprintCallable, Category = "MLDeformer")
 	USkeletalMeshComponent* GetSkeletalMeshComponent() const	{ return SkelMeshComponent.Get(); }
 
 	/**
