[MLDeformer] Fix some non unity compile issue.

diff --git a/Engine/Plugins/Animation/MLDeformer/MLDeformerFramework/Source/MLDeformerFramework/Public/MLDeformerVizSettings.h b/Engine/Plugins/Animation/MLDeformer/MLDeformerFramework/Source/MLDeformerFramework/Public/MLDeformerVizSettings.h
--- a/Engine/Plugins/Animation/MLDeformer/MLDeformerFramework/Source/MLDeformerFramework/Public/MLDeformerVizSettings.h
+++ b/Engine/Plugins/Animation/MLDeformer/MLDeformerFramework/Source/MLDeformerFramework/Public/MLDeformerVizSettings.h
@@ -7,10 +7,9 @@
 #include "UObject/ObjectPtr.h"
 #include "UObject/SoftObjectPtr.h"
 #include "Animation/MeshDeformer.h"
+#include "Animation/AnimSequence.h"
 #include "MLDeformerVizSettings.generated.h"
 
-class UMeshDeformer;
-class UAnimSequence;
 
 /** The visualization mode, which selects whether you want to view the training data, or test your already trained model. */
 UENUM()
