[MLDeformer] Some improvements in the ML Deformer component, to make it less reliant on the component activation.

diff --git a/Engine/Plugins/Animation/MLDeformer/MLDeformerFramework/Source/MLDeformerFramework/Private/MLDeformerComponent.cpp b/Engine/Plugins/Animation/MLDeformer/MLDeformerFramework/Source/MLDeformerFramework/Private/MLDeformerComponent.cpp
--- a/Engine/Plugins/Animation/MLDeformer/MLDeformerFramework/Source/MLDeformerFramework/Private/MLDeformerComponent.cpp
+++ b/Engine/Plugins/Animation/MLDeformer/MLDeformerFramework/Source/MLDeformerFramework/Private/MLDeformerComponent.cpp
@@ -54,6 +54,16 @@ void UMLDeformerComponent::SetupComponent(UMLDeformerAsset* InDeformerAsset, USk
 {
 	TRACE_CPUPROFILER_EVENT_SCOPE(UMLDeformerComponent::SetupComponent)
 
+	RenderCommandFence.BeginFence();
+	RenderCommandFence.Wait();
+
+	// Remove the existing skeletal mesh component as tick prerequisite.
+	if (SkelMeshComponent)
+	{
+		RemoveTickPrerequisiteComponent(SkelMeshComponent);
+	}
+
+	// Add the new one.
 	if (InSkelMeshComponent)
 	{		
 		AddTickPrerequisiteComponent(InSkelMeshComponent);
@@ -125,15 +135,13 @@ void UMLDeformerComponent::BeginDestroy()
 	Super::BeginDestroy();
 }
 
-void UMLDeformerComponent::Activate(bool bReset)
+USkeletalMeshComponent* UMLDeformerComponent::FindSkeletalMeshComponent(UMLDeformerAsset* Asset)
 {
-	// If we haven't pointed to some skeletal mesh component to use, then try to find one on the actor.
-	SkelMeshComponent = nullptr;
-	if (!DeformerAsset.IsNull())
+	USkeletalMeshComponent* ResultingComponent = nullptr;
+	if (Asset != nullptr)
 	{
-		// First search for a skeletal mesh component with the expected number of vertices.
-		// This will try to find a skeletal mesh with the same number of vertices.
-		const UMLDeformerModel* Model = DeformerAsset ? DeformerAsset->GetModel() : nullptr;
+		// First search for a skeletal mesh component that uses the same skeletal mesh as the ML Deformer asset was trained on.
+		const UMLDeformerModel* Model = Asset ? Asset->GetModel() : nullptr;
 		if (Model && Model->GetSkeletalMesh())
 		{
 			const USkeletalMesh* ModelSkeletalMesh = Model->GetSkeletalMesh();
@@ -149,22 +157,33 @@ void UMLDeformerComponent::Activate(bool bReset)
 				const USkeletalMesh* ComponentSkeletalMesh = Component->GetSkeletalMeshAsset();
 				if (ComponentSkeletalMesh == ModelSkeletalMesh)
 				{
-					SkelMeshComponent = Component;
+					ResultingComponent = Component;
 					break;
 				}
 			}
 		}
 	}
 
-	if (SkelMeshComponent == nullptr)
+	if (ResultingComponent == nullptr)
 	{
 		// Fall back to the first skeletal mesh component.
 		const AActor* Actor = Cast<AActor>(GetOuter());
-		SkelMeshComponent = Actor->FindComponentByClass<USkeletalMeshComponent>();
+		ResultingComponent = Actor->FindComponentByClass<USkeletalMeshComponent>();
 	}
 
+	return ResultingComponent;
+}
+
+void UMLDeformerComponent::UpdateSkeletalMeshComponent()
+{
+	SkelMeshComponent = FindSkeletalMeshComponent(DeformerAsset.Get());
 	SetupComponent(DeformerAsset, SkelMeshComponent);
+}
 
+void UMLDeformerComponent::Activate(bool bReset)
+{
+	UpdateSkeletalMeshComponent();
+	SetupComponent(DeformerAsset, SkelMeshComponent);
 	Super::Activate(bReset);
 }
 
@@ -201,12 +220,7 @@ void UMLDeformerComponent::TickComponent(float DeltaTime, enum ELevelTick TickTy
 
 		if (Property->GetFName() == GET_MEMBER_NAME_CHECKED(UMLDeformerComponent, DeformerAsset))
 		{
-			RenderCommandFence.BeginFence();
-			RenderCommandFence.Wait();
-
-			RemoveNeuralNetworkModifyDelegate();
-			Init();
-			AddNeuralNetworkModifyDelegate();
+			SetDeformerAsset(DeformerAsset);
 		}
 	}
 #endif
diff --git a/Engine/Plugins/Animation/MLDeformer/MLDeformerFramework/Source/MLDeformerFramework/Public/MLDeformerComponent.h b/Engine/Plugins/Animation/MLDeformer/MLDeformerFramework/Source/MLDeformerFramework/Public/MLDeformerComponent.h
--- a/Engine/Plugins/Animation/MLDeformer/MLDeformerFramework/Source/MLDeformerFramework/Public/MLDeformerComponent.h
+++ b/Engine/Plugins/Animation/MLDeformer/MLDeformerFramework/Source/MLDeformerFramework/Public/MLDeformerComponent.h
@@ -73,10 +73,22 @@ public:
 
 	/**
 	 * Set the deformer asset that is used by this component.
+	 * This will trigger the internal ML Deformer instance to be recreated.
 	 * @param InDeformerAsset A pointer to the deformer asset.
 	 */
 	UFUNCTION(BlueprintCallable, Category = "MLDeformer")
-	void SetDeformerAsset(UMLDeformerAsset* InDeformerAsset)	{ DeformerAsset = InDeformerAsset; }
+	void SetDeformerAsset(UMLDeformerAsset* InDeformerAsset)	{ DeformerAsset = InDeformerAsset; UpdateSkeletalMeshComponent(); }
+
+	/**
+	 * Find the skeletal mesh component to apply the deformer on.
+	 * This will return the skeletal mesh component (on this actor) which uses the same skeletal mesh as the passed in ML Deformer asset was trained on.
+	 * If there is no such skeletal mesh component then it will return the first skeletal mesh component it finds.
+	 * In case there is no skeletal mesh component on the actor, a nullptr is returned.
+	 * @param Asset The ML Deformer asset to search a component for.
+	 * @return The skeletal mesh component that would be the best.
+	 */
+	UFUNCTION(BlueprintCallable, Category = "MLDeformer")
+	USkeletalMeshComponent* FindSkeletalMeshComponent(UMLDeformerAsset* Asset);
 
 	/**
 	 * Get the ML Deformer model instance that this component currently uses.
@@ -93,6 +105,10 @@ public:
 	UFUNCTION(BlueprintCallable, Category = "MLDeformer")
 	USkeletalMeshComponent* GetSkeletalMeshComponent() const	{ return SkelMeshComponent.Get(); }
 
+	/** Find the skeletal mesh component that this deformer should work on, and set it as our target component. */
+	UFUNCTION(BlueprintCallable, Category = "MLDeformer")
+	void UpdateSkeletalMeshComponent();
+
 	/**
 	 * Suppress logging warnings about mesh deformers not being set.
 	 * A warning is logged when an ML Deformer is used that requires a deformer graph, but the skeletal mesh has no deformer graph setup.
