[MLDeformer]
- Some improvements in handling of deactivated and zero weight ML Deformer and moved some methods to protected instead of public. Tested performance in Insights.
- Fixed a bug where the deformer could stop working. This required a small refactor.
- Changed the handling of editor only packages. Marking assets editor only and flagging them dirty, only when changing the properties.
- Added cook warnings when training assets are going to be included in the packaging. Also showing warnings when opening the asset.
- MLD assets now load very fast in PIE/Level, as they don't load training assets anymore.

#rnx

diff --git a/Engine/Plugins/Animation/MLDeformer/MLDeformerFramework/Source/MLDeformerFramework/Private/MLDeformerComponent.cpp b/Engine/Plugins/Animation/MLDeformer/MLDeformerFramework/Source/MLDeformerFramework/Private/MLDeformerComponent.cpp
--- a/Engine/Plugins/Animation/MLDeformer/MLDeformerFramework/Source/MLDeformerFramework/Private/MLDeformerComponent.cpp
+++ b/Engine/Plugins/Animation/MLDeformer/MLDeformerFramework/Source/MLDeformerFramework/Private/MLDeformerComponent.cpp
@@ -25,8 +25,8 @@ void UMLDeformerComponent::Init()
 		if (ModelInstance)
 		{
 			ModelInstance->Release();
+			ModelInstance = nullptr;
 		}
-		ModelInstance = nullptr;
 		return;
 	}
 
@@ -37,11 +37,12 @@ void UMLDeformerComponent::Init()
 		if (ModelInstance)
 		{
 			ModelInstance->Release();
+			ModelInstance = nullptr;
 		}
 		ModelInstance = Model->CreateModelInstance(this);
 		ModelInstance->SetModel(Model);
 		ModelInstance->Init(SkelMeshComponent);
-		Model->PostMLDeformerComponentInit(ModelInstance);
+		ModelInstance->PostMLDeformerComponentInit();
 	}
 	else
 	{
@@ -149,16 +150,19 @@ USkeletalMeshComponent* UMLDeformerComponent::FindSkeletalMeshComponent(UMLDefor
 			// Get a list of all skeletal mesh components on the actor.
 			TArray<USkeletalMeshComponent*> Components;
 			AActor* Actor = Cast<AActor>(GetOuter());
-			Actor->GetComponents<USkeletalMeshComponent>(Components);
-		
-			// Find a component that uses a mesh with the same vertex count.
-			for (USkeletalMeshComponent* Component : Components)
+			if (Actor)
 			{
-				const USkeletalMesh* ComponentSkeletalMesh = Component->GetSkeletalMeshAsset();
-				if (ComponentSkeletalMesh == ModelSkeletalMesh)
+				Actor->GetComponents<USkeletalMeshComponent>(Components);
+		
+				// Find a component that uses a mesh with the same vertex count.
+				for (USkeletalMeshComponent* Component : Components)
 				{
-					ResultingComponent = Component;
-					break;
+					const USkeletalMesh* ComponentSkeletalMesh = Component->GetSkeletalMeshAsset();
+					if (ComponentSkeletalMesh == ModelSkeletalMesh)
+					{
+						ResultingComponent = Component;
+						break;
+					}
 				}
 			}
 		}
@@ -168,7 +172,10 @@ USkeletalMeshComponent* UMLDeformerComponent::FindSkeletalMeshComponent(UMLDefor
 	{
 		// Fall back to the first skeletal mesh component.
 		const AActor* Actor = Cast<AActor>(GetOuter());
-		ResultingComponent = Actor->FindComponentByClass<USkeletalMeshComponent>();
+		if (Actor)
+		{
+			ResultingComponent = Actor->FindComponentByClass<USkeletalMeshComponent>();
+		}
 	}
 
 	return ResultingComponent;
@@ -190,7 +197,11 @@ void UMLDeformerComponent::Activate(bool bReset)
 void UMLDeformerComponent::Deactivate()
 {
 	RemoveNeuralNetworkModifyDelegate();
-	ModelInstance = nullptr;
+	if (ModelInstance)
+	{
+		ModelInstance->Release();
+		ModelInstance = nullptr;
+	}
 	Super::Deactivate();
 }
 
diff --git a/Engine/Plugins/Animation/MLDeformer/MLDeformerFramework/Source/MLDeformerFramework/Private/MLDeformerGeomCacheModel.cpp b/Engine/Plugins/Animation/MLDeformer/MLDeformerFramework/Source/MLDeformerFramework/Private/MLDeformerGeomCacheModel.cpp
--- a/Engine/Plugins/Animation/MLDeformer/MLDeformerFramework/Source/MLDeformerFramework/Private/MLDeformerGeomCacheModel.cpp
+++ b/Engine/Plugins/Animation/MLDeformer/MLDeformerFramework/Source/MLDeformerFramework/Private/MLDeformerGeomCacheModel.cpp
@@ -3,6 +3,7 @@
 #include "MLDeformerGeomCacheVizSettings.h"
 #include "MLDeformerGeomCacheHelpers.h"
 #include "MLDeformerComponent.h"
+#include "MLDeformerAsset.h"
 #include "UObject/Object.h"
 #include "UObject/UObjectGlobals.h"
 #include "GeometryCache.h"
@@ -22,6 +23,37 @@ void UMLDeformerGeomCacheModel::Serialize(FArchive& Archive)
 }
 
 #if WITH_EDITOR
+void UMLDeformerGeomCacheModel::LogPackagingWarnings()
+{
+	const UGeometryCache* GeomCache = GetGeometryCache();
+	if (GeomCache && !GeomCache->GetPackage()->HasAnyPackageFlags(PKG_EditorOnly))
+	{
+		UE_LOG(
+			LogMLDeformer,
+			Warning, 
+			TEXT("Training geometry cache '%s' is not marked as editor only and will be included during packaging. Reopen the '%s' MLD asset and save this geom cache asset to fix this."),
+			*GeomCache->GetName(),
+			*GetDeformerAsset()->GetName());
+	}
+
+	UMLDeformerGeomCacheVizSettings* Viz = GetGeomCacheVizSettings();
+	if (Viz)
+	{
+		const UGeometryCache* TestGeomCache = Viz->GetTestGroundTruth();
+		if (TestGeomCache && !TestGeomCache->GetPackage()->HasAnyPackageFlags(PKG_EditorOnly))
+		{
+			UE_LOG(
+				LogMLDeformer,
+				Warning, 
+				TEXT("Test geometry cache '%s' is not marked as editor only and will be included during packaging. Reopen the '%s' MLD asset and save this geom cache asset to fix this."),
+				*TestGeomCache->GetName(),
+				*GetDeformerAsset()->GetName());
+		}
+	}
+
+	Super::LogPackagingWarnings();
+}
+
 void UMLDeformerGeomCacheModel::UpdateNumTargetMeshVertices()
 {
 	SetNumTargetMeshVerts(UE::MLDeformer::ExtractNumImportedGeomCacheVertices(GetGeometryCache()));
@@ -31,22 +63,6 @@ UMLDeformerGeomCacheVizSettings* UMLDeformerGeomCacheModel::GetGeomCacheVizSetti
 {
 	return Cast<UMLDeformerGeomCacheVizSettings>(GetVizSettings());
 }
-
-void UMLDeformerGeomCacheModel::SetAssetEditorOnlyFlags()
-{
-	// Set the flags for the base class, which filters out the training anim sequence.
-	UMLDeformerModel::SetAssetEditorOnlyFlags();
-
-	// The training geometry cache.
-	MarkObjectAsEditorOnly(GetGeometryCache());
-
-	// The testing geometry cache (ground truth of the test anim sequence).
-	const UMLDeformerGeomCacheVizSettings* GeomCacheVizSettings = GetGeomCacheVizSettings();
-	if (GeomCacheVizSettings)
-	{
-		MarkObjectAsEditorOnly(GeomCacheVizSettings->GetTestGroundTruth());
-	}
-}
 #endif // WITH_EDITOR
 
 #if WITH_EDITORONLY_DATA
diff --git a/Engine/Plugins/Animation/MLDeformer/MLDeformerFramework/Source/MLDeformerFramework/Private/MLDeformerModel.cpp b/Engine/Plugins/Animation/MLDeformer/MLDeformerFramework/Source/MLDeformerFramework/Private/MLDeformerModel.cpp
--- a/Engine/Plugins/Animation/MLDeformer/MLDeformerFramework/Source/MLDeformerFramework/Private/MLDeformerModel.cpp
+++ b/Engine/Plugins/Animation/MLDeformer/MLDeformerFramework/Source/MLDeformerFramework/Private/MLDeformerModel.cpp
@@ -68,6 +68,8 @@ void UMLDeformerModel::Serialize(FArchive& Archive)
 		{
 			if (Archive.IsCooking())
 			{
+				LogPackagingWarnings();
+
 				AnimSequence = nullptr;
 				VizSettings = nullptr;
 			}
@@ -78,15 +80,6 @@ void UMLDeformerModel::Serialize(FArchive& Archive)
 	#endif
 
 	Super::Serialize(Archive);
-
-	#if WITH_EDITOR
-		// Reinitialize the vertex map after loading its contents.
-		// Because our skeletal mesh asset could have changed vertex counts.
-		if (Archive.IsLoading())
-		{
-			InitVertexMap();
-		}
-	#endif
 }
 
 UMLDeformerAsset* UMLDeformerModel::GetDeformerAsset() const
@@ -100,6 +93,10 @@ void UMLDeformerModel::PostLoad()
 
 	Super::PostLoad();
 
+	#if WITH_EDITOR
+		InitVertexMap();
+	#endif
+
 	InitGPUData();
 
 	#if WITH_EDITOR
@@ -131,10 +128,6 @@ void UMLDeformerModel::PostLoad()
 			Network->SetDeviceType(ENeuralDeviceType::CPU, ENeuralDeviceType::CPU, ENeuralDeviceType::CPU);
 		}
 	}
-
-#if WITH_EDITOR
-	SetAssetEditorOnlyFlags();
-#endif
 }
 
 void UMLDeformerModel::SetNeuralNetwork(UNeuralNetwork* InNeuralNetwork)
@@ -213,25 +206,19 @@ void UMLDeformerModel::FloatArrayToVector3Array(const TArray<float>& FloatArray,
 		return SkeletalMesh ? SkeletalMesh->GetNumImportedVertices() : 0;
 	}
 
-	void UMLDeformerModel::MarkObjectAsEditorOnly(UObject* InObject, bool bLogMessage)
+	void UMLDeformerModel::LogPackagingWarnings()
 	{
-		if (InObject == nullptr)
-		{
-			return;
-		}
-
-		InObject->GetPackage()->SetPackageFlags(PKG_EditorOnly);
-
-		if (bLogMessage)
+		const UAnimSequence* AnimSeq = GetAnimSequence();
+		if (AnimSeq && !AnimSeq->GetPackage()->HasAnyPackageFlags(PKG_EditorOnly))
 		{
-			UE_LOG(LogMLDeformer, Display, TEXT("Marking '%s' as editor only asset. Asset will be excluded from packaging."), *InObject->GetName());
+			UE_LOG(
+				LogMLDeformer,
+				Warning, 
+				TEXT("Training animation sequence '%s' is not marked as editor only and will be included during packaging. Reopen the '%s' MLD asset and save this anim sequence asset to fix this."),
+				*AnimSeq->GetName(),
+				*GetDeformerAsset()->GetName());
 		}
 	}
-
-	void UMLDeformerModel::SetAssetEditorOnlyFlags()
-	{
-		MarkObjectAsEditorOnly(GetAnimSequence());
-	}
 #endif	// #if WITH_EDITOR
 
 #if WITH_EDITORONLY_DATA
diff --git a/Engine/Plugins/Animation/MLDeformer/MLDeformerFramework/Source/MLDeformerFramework/Private/MLDeformerModelInstance.cpp b/Engine/Plugins/Animation/MLDeformer/MLDeformerFramework/Source/MLDeformerFramework/Private/MLDeformerModelInstance.cpp
--- a/Engine/Plugins/Animation/MLDeformer/MLDeformerFramework/Source/MLDeformerFramework/Private/MLDeformerModelInstance.cpp
+++ b/Engine/Plugins/Animation/MLDeformer/MLDeformerFramework/Source/MLDeformerFramework/Private/MLDeformerModelInstance.cpp
@@ -298,11 +298,16 @@ bool UMLDeformerModelInstance::IsValidForDataProvider() const
 	return (Model->GetVertexMapBuffer().ShaderResourceViewRHI != nullptr) && (GetNeuralNetworkInferenceHandle() != -1);
 }
 
-void UMLDeformerModelInstance::RunNeuralNetwork(float ModelWeight)
+void UMLDeformerModelInstance::Execute(float ModelWeight)
 {
-	TRACE_CPUPROFILER_EVENT_SCOPE(UMLDeformerModelInstance::RunNeuralNetwork)
+	TRACE_CPUPROFILER_EVENT_SCOPE(UMLDeformerModelInstance::Execute)
 
 	UNeuralNetwork* NeuralNetwork = Model->GetNeuralNetwork();
+	if (NeuralNetwork == nullptr)
+	{
+		return;
+	}
+
 	if (Model->IsNeuralNetworkOnGPU())
 	{
 		// NOTE: Inputs still come from the CPU.
@@ -326,9 +331,9 @@ void UMLDeformerModelInstance::RunNeuralNetwork(float ModelWeight)
 	}
 }
 
-bool UMLDeformerModelInstance::SetupNeuralNetworkForFrame()
+bool UMLDeformerModelInstance::SetupInputs()
 {
-	TRACE_CPUPROFILER_EVENT_SCOPE(UMLDeformerModelInstance::SetupNeuralNetworkForFrame)
+	TRACE_CPUPROFILER_EVENT_SCOPE(UMLDeformerModelInstance::SetupInputs)
 
 	// Some safety checks.
 	if (Model == nullptr ||
@@ -378,15 +383,18 @@ void UMLDeformerModelInstance::Tick(float DeltaTime, float ModelWeight)
 	// This could for example happen when you add an ML Deformer component, but your SkeletalMesh isn't setup yet, but later becomes valid.
 	if (Model && !HasPostInitialized())
 	{
-		Model->PostMLDeformerComponentInit(this);
+		PostMLDeformerComponentInit();
 	}
 
-	// Setup the neural network, by feeding it's inputs and doing some compatibility checks.
-	if (!SetupNeuralNetworkForFrame())
+	if (ModelWeight > 0.0001f && SetupInputs())
 	{
-		return;
+		// Execute the model instance.
+		// For models using neural networks this will perform the inference, 
+		// calculate the network outputs and possibly use them, depending on how the model works.
+		Execute(ModelWeight);
+	}
+	else
+	{
+		HandleZeroModelWeight();
 	}
-
-	// Perform the inference, calculate the network outputs and possibly use them, depending on how the model works.
-	RunNeuralNetwork(ModelWeight);
 }
diff --git a/Engine/Plugins/Animation/MLDeformer/MLDeformerFramework/Source/MLDeformerFramework/Private/MLDeformerMorphModel.cpp b/Engine/Plugins/Animation/MLDeformer/MLDeformerFramework/Source/MLDeformerFramework/Private/MLDeformerMorphModel.cpp
--- a/Engine/Plugins/Animation/MLDeformer/MLDeformerFramework/Source/MLDeformerFramework/Private/MLDeformerMorphModel.cpp
+++ b/Engine/Plugins/Animation/MLDeformer/MLDeformerFramework/Source/MLDeformerFramework/Private/MLDeformerMorphModel.cpp
@@ -9,14 +9,11 @@
 
 #define LOCTEXT_NAMESPACE "MLDeformerMorphModel"
 
-TAtomic<int32> UMLDeformerMorphModel::NextFreeMorphSetID(0);
-
 UMLDeformerMorphModel::UMLDeformerMorphModel(const FObjectInitializer& ObjectInitializer)
 	: Super(ObjectInitializer)
 {
 	MorphTargetSet = MakeShared<FExternalMorphSet>();
 	MorphTargetSet->Name = GetClass()->GetFName();
-	ExternalMorphSetID = NextFreeMorphSetID++;
 }
 
 void UMLDeformerMorphModel::Serialize(FArchive& Archive)
@@ -46,58 +43,6 @@ UMLDeformerModelInstance* UMLDeformerMorphModel::CreateModelInstance(UMLDeformer
 	return NewObject<UMLDeformerMorphModelInstance>(Component);
 }
 
-void UMLDeformerMorphModel::PostMLDeformerComponentInit(UMLDeformerModelInstance* ModelInstance)
-{
-	if (ModelInstance->HasPostInitialized())
-	{
-		return;
-	}
-
-	TRACE_CPUPROFILER_EVENT_SCOPE(UMLDeformerMorphModel::PostMLDeformerComponentInit)
-
-	Super::PostMLDeformerComponentInit(ModelInstance);
-
-	// Register the external morph targets buffer to the skinned mesh component.
-	USkeletalMeshComponent* SkelMeshComponent = ModelInstance->GetSkeletalMeshComponent();
-	if (SkelMeshComponent && SkelMeshComponent->GetSkeletalMeshAsset())
-	{	
-		// Register the morph set. This overwrites the existing one for this model, if it already exists.
-		// Only add to LOD 0 for now.
-		const int32 LOD = 0;
-		SkelMeshComponent->AddExternalMorphSet(LOD, ExternalMorphSetID, MorphTargetSet);
-
-		// When we're in editor mode, keep the CPU data around, so we can re-initialize when needed.
-#if WITH_EDITOR
-		MorphTargetSet->MorphBuffers.SetEmptyMorphCPUDataOnInitRHI(false);
-#else
-		MorphTargetSet->MorphBuffers.SetEmptyMorphCPUDataOnInitRHI(true);
-#endif
-
-		// Release the render resources, but only in an editor build.
-		// The non-editor build shouldn't do this, as then it can't initialize again. The non-editor build assumes
-		// that the data doesn't change and we don't need to re-init.
-		// In the editor build we have to re-initialize the render resources as the morph targets can change after (re)training, so
-		// that is why we release them here, and intialize them again after.
-		FMorphTargetVertexInfoBuffers& MorphBuffers = MorphTargetSet->MorphBuffers;
-#if WITH_EDITOR
-		BeginReleaseResource(&MorphBuffers);
-#endif
-
-		// Reinitialize the GPU compressed buffers.
-		if (MorphBuffers.IsMorphCPUDataValid() && MorphBuffers.GetNumMorphs() > 0)
-		{
-			// In a non-editor build this will clear the CPU data.
-			// That also means it can't re-init the resources later on again.
-			BeginInitResource(&MorphBuffers);
-		}
-
-		// Update the weight information in the Skeletal Mesh.
-		SkelMeshComponent->RefreshExternalMorphTargetWeights();
-
-		ModelInstance->SetHasPostInitialized(true);
-	}
-}
-
 void UMLDeformerMorphModel::SetMorphTargetDeltaFloats(const TArray<float>& Deltas)
 {
 	FloatArrayToVector3Array(Deltas, MorphTargetDeltas);
@@ -131,10 +76,4 @@ void UMLDeformerMorphModel::BeginDestroy()
 	Super::BeginDestroy();
 }
 
-FExternalMorphSetWeights* UMLDeformerMorphModel::FindExternalMorphWeights(int32 LOD, USkinnedMeshComponent* SkinnedMeshComponent) const
-{
-	check(SkinnedMeshComponent);
-	return SkinnedMeshComponent->GetExternalMorphWeights(LOD).MorphSets.Find(ExternalMorphSetID);
-}
-
 #undef LOCTEXT_NAMESPACE
diff --git a/Engine/Plugins/Animation/MLDeformer/MLDeformerFramework/Source/MLDeformerFramework/Private/MLDeformerMorphModelInstance.cpp b/Engine/Plugins/Animation/MLDeformer/MLDeformerFramework/Source/MLDeformerFramework/Private/MLDeformerMorphModelInstance.cpp
--- a/Engine/Plugins/Animation/MLDeformer/MLDeformerFramework/Source/MLDeformerFramework/Private/MLDeformerMorphModelInstance.cpp
+++ b/Engine/Plugins/Animation/MLDeformer/MLDeformerFramework/Source/MLDeformerFramework/Private/MLDeformerMorphModelInstance.cpp
@@ -5,41 +5,130 @@
 #include "NeuralNetwork.h"
 #include "Components/SkeletalMeshComponent.h"
 
+TAtomic<int32> UMLDeformerMorphModelInstance::NextFreeMorphSetID(0);
+
+void UMLDeformerMorphModelInstance::Init(USkeletalMeshComponent* SkelMeshComponent)
+{
+	Super::Init(SkelMeshComponent);
+
+	// Generate a unique ID for our morph target set.
+	ExternalMorphSetID = NextFreeMorphSetID++;
+}
+
 void UMLDeformerMorphModelInstance::Release()
 {
+	// Try to unregister the morph target and morph target set.
+	if (!SkeletalMeshComponent.IsNull())
+	{
+		const UMLDeformerMorphModel* MorphModel = Cast<UMLDeformerMorphModel>(Model);
+		if (MorphModel)
+		{
+			const int32 LOD = 0;
+			SkeletalMeshComponent->RemoveExternalMorphSet(LOD, ExternalMorphSetID);
+			SkeletalMeshComponent->RefreshExternalMorphTargetWeights();
+		}
+	}
+
 	Super::Release();
+}
 
-	if (SkeletalMeshComponent == nullptr)
+void UMLDeformerMorphModelInstance::PostMLDeformerComponentInit()
+{
+	if (HasPostInitialized())
 	{
 		return;
 	}
 
+	TRACE_CPUPROFILER_EVENT_SCOPE(UMLDeformerMorphModelInstance::PostMLDeformerComponentInit)
+
+	Super::PostMLDeformerComponentInit();
+
+	// Register the external morph targets buffer to the skinned mesh component.
+	USkeletalMeshComponent* SkelMeshComponent = GetSkeletalMeshComponent();
+	if (SkelMeshComponent && SkelMeshComponent->GetSkeletalMeshAsset())
+	{	
+		// Get the morph model and its morph target set.
+		UMLDeformerMorphModel* MorphModel = Cast<UMLDeformerMorphModel>(Model);
+		check(MorphModel);
+		TSharedPtr<FExternalMorphSet> MorphTargetSet = MorphModel->GetMorphTargetSet();
+
+		// Register the morph set. This overwrites the existing one for this model, if it already exists.
+		// Only add to LOD 0 for now.
+		const int32 LOD = 0;
+		SkelMeshComponent->AddExternalMorphSet(LOD, ExternalMorphSetID, MorphTargetSet);
+
+		// When we're in editor mode, keep the CPU data around, so we can re-initialize when needed.
+#if WITH_EDITOR
+		MorphTargetSet->MorphBuffers.SetEmptyMorphCPUDataOnInitRHI(false);
+#else
+		MorphTargetSet->MorphBuffers.SetEmptyMorphCPUDataOnInitRHI(true);
+#endif
+
+		// Release the render resources, but only in an editor build.
+		// The non-editor build shouldn't do this, as then it can't initialize again. The non-editor build assumes
+		// that the data doesn't change and we don't need to re-init.
+		// In the editor build we have to re-initialize the render resources as the morph targets can change after (re)training, so
+		// that is why we release them here, and intialize them again after.
+		FMorphTargetVertexInfoBuffers& MorphBuffers = MorphTargetSet->MorphBuffers;
+#if WITH_EDITOR
+		BeginReleaseResource(&MorphBuffers);
+#endif
+
+		// Reinitialize the GPU compressed buffers.
+		if (MorphBuffers.IsMorphCPUDataValid() && MorphBuffers.GetNumMorphs() > 0)
+		{
+			// In a non-editor build this will clear the CPU data.
+			// That also means it can't re-init the resources later on again.
+			BeginInitResource(&MorphBuffers);
+		}
+
+		// Update the weight information in the Skeletal Mesh.
+		SkelMeshComponent->RefreshExternalMorphTargetWeights();
+
+		SetHasPostInitialized(true);
+	}
+}
+
+FExternalMorphSetWeights* UMLDeformerMorphModelInstance::FindWeightData(int32 LOD) const
+{
 	const UMLDeformerMorphModel* MorphModel = Cast<UMLDeformerMorphModel>(Model);
-	if (MorphModel == nullptr)
+	if (MorphModel == nullptr || SkeletalMeshComponent.IsNull())
 	{
-		return;
+		return nullptr;
 	}
 
-	const int32 LOD = 0;
-	SkeletalMeshComponent->RemoveExternalMorphSet(LOD, MorphModel->GetExternalMorphSetID());
-	SkeletalMeshComponent->RefreshExternalMorphTargetWeights();
+	// If we haven't got an external morph set registered yet, let's just return a nullptr.
+	if (!SkeletalMeshComponent->HasExternalMorphSet(LOD, ExternalMorphSetID))
+	{
+		return nullptr;
+	}
+
+	// Grab the weight data for this morph set.
+	// This could potentially fail if we are applying this deformer to the wrong skeletal mesh component.
+	return SkeletalMeshComponent->GetExternalMorphWeights(LOD).MorphSets.Find(ExternalMorphSetID);
 }
 
-// Run the neural network, which calculates its outputs, which are the weights of our morph targets.
-void UMLDeformerMorphModelInstance::RunNeuralNetwork(float ModelWeight)
+void UMLDeformerMorphModelInstance::HandleZeroModelWeight()
 {
-	TRACE_CPUPROFILER_EVENT_SCOPE(UNeuralMorphModelInstance::RunNeuralNetwork)
+	TRACE_CPUPROFILER_EVENT_SCOPE(UNeuralMorphModelInstance::HandleZeroModelWeight)
 
-	const UMLDeformerMorphModel* MorphModel = Cast<UMLDeformerMorphModel>(Model);
-	if (MorphModel == nullptr)
+	const int LOD = 0;	// For now we only support LOD 0, as we can't setup an ML Deformer per LOD yet.
+	FExternalMorphSetWeights* WeightData = FindWeightData(LOD);
+	if (WeightData)
 	{
-		return;
+		WeightData->ZeroWeights();
 	}
+}
+
+// Run the neural network, which calculates its outputs, which are the weights of our morph targets.
+void UMLDeformerMorphModelInstance::Execute(float ModelWeight)
+{
+	TRACE_CPUPROFILER_EVENT_SCOPE(UNeuralMorphModelInstance::Execute)
 
 	// Grab the weight data for this morph set.
 	// This could potentially fail if we are applying this deformer to the wrong skeletal mesh component.
 	const int LOD = 0;	// For now we only support LOD 0, as we can't setup an ML Deformer per LOD yet.
-	FExternalMorphSetWeights* WeightData = MorphModel->FindExternalMorphWeights(LOD, SkeletalMeshComponent);
+	FExternalMorphSetWeights* WeightData = FindWeightData(LOD);
 	if (WeightData == nullptr)
 	{
 		return;
@@ -48,11 +137,11 @@ void UMLDeformerMorphModelInstance::RunNeuralNetwork(float ModelWeight)
 	// If our model is active, we want to run the neural network and update the morph weights
 	// with the values that the neural net calculated for us.
 	const UNeuralNetwork* NeuralNetwork = Model->GetNeuralNetwork();
-	if (ModelWeight > 0.0f && NeuralNetwork != nullptr)
+	if (NeuralNetwork)
 	{
 		// Perform the neural network inference, which updates the output tensor.
 		// This takes most CPU time inside this method.
-		UMLDeformerModelInstance::RunNeuralNetwork(ModelWeight);
+		UMLDeformerModelInstance::Execute(ModelWeight);
 
 		// Get the output tensor, read the values and use them as morph target weights inside the skeletal mesh component.
 		const FNeuralTensor& OutputTensor = NeuralNetwork->GetOutputTensorForContext(NeuralNetworkInferenceHandle);
diff --git a/Engine/Plugins/Animation/MLDeformer/MLDeformerFramework/Source/MLDeformerFramework/Public/MLDeformerGeomCacheModel.h b/Engine/Plugins/Animation/MLDeformer/MLDeformerFramework/Source/MLDeformerFramework/Public/MLDeformerGeomCacheModel.h
--- a/Engine/Plugins/Animation/MLDeformer/MLDeformerFramework/Source/MLDeformerFramework/Public/MLDeformerGeomCacheModel.h
+++ b/Engine/Plugins/Animation/MLDeformer/MLDeformerFramework/Source/MLDeformerFramework/Public/MLDeformerGeomCacheModel.h
@@ -38,7 +38,7 @@ public:
 #endif
 #if WITH_EDITOR
 	virtual void UpdateNumTargetMeshVertices() override;
-	virtual void SetAssetEditorOnlyFlags() override;
+	virtual void LogPackagingWarnings() override;
 	// ~END UMLDeformerModel overrides.
 
 	/**
diff --git a/Engine/Plugins/Animation/MLDeformer/MLDeformerFramework/Source/MLDeformerFramework/Public/MLDeformerModel.h b/Engine/Plugins/Animation/MLDeformer/MLDeformerFramework/Source/MLDeformerFramework/Public/MLDeformerModel.h
--- a/Engine/Plugins/Animation/MLDeformer/MLDeformerFramework/Source/MLDeformerFramework/Public/MLDeformerModel.h
+++ b/Engine/Plugins/Animation/MLDeformer/MLDeformerFramework/Source/MLDeformerFramework/Public/MLDeformerModel.h
@@ -102,14 +102,6 @@ public:
 	 */
 	virtual FString GetDisplayName() const PURE_VIRTUAL(UMLDeformerModel::GetDisplayName, return FString(););
 
-	/**
-	 * Initialize this model, after the ML deformer component has been initialized.
-	 * This is a post-step after component initialization. It can be used to register things such as 
-	 * external morph target sets that are generated by this model.
-	 * @param ModelInstance The ML Deformer model instance that we are initializing with.
-	 */
-	virtual void PostMLDeformerComponentInit(UMLDeformerModelInstance* ModelInstance) {}
-
 	/**
 	 * Defines whether this model supports bone transforms as input or not.
 	 * On default this is set to return true as most models have bone rotations as inputs to the neural network.
@@ -172,12 +164,11 @@ public:
 
 #if WITH_EDITOR
 	/**
-	 * Set an object to be editor only.
-	 * This will add the PKG_EditorOnly package flag. This will cause it to not be included during packaging of the project.
-	 * @param InObject The object to mark as editor only.
-	 * @param bLogMessage Set to true when you wish to log a message about this marking.
+	 * Log package warnings.
+	 * This can identify things like when the training assets aren't marked as editor only.
+	 * If this flag isn't set, those assets will be included during packaging of your project.
 	 */
-	void MarkObjectAsEditorOnly(UObject* InObject, bool bLogMessage=true);
+	virtual void LogPackagingWarnings();
 
 	/**
 	 * Update the cached number of vertices of both base and target meshes.
@@ -198,18 +189,6 @@ public:
 	 */ 
 	virtual void UpdateNumTargetMeshVertices() PURE_VIRTUAL(UMLDeformerModel::GetNumTargetMeshVertices; mTargetMeshVerts = 0;);
 
-	/**
-	 * Update all editor only assets with the editor only flag.
-	 * This prevents certain assets that are only used during training to be included in a packaged build.
-	 * Examples of these are a training geometry cache object, and the ground truth geometry cache of a test animation in the visualization settings.
-	 * If you do not set the flag for this, those will unfortunately be included, and they can be massive in size and aren't needed at runtime.
-	 * You can set these flags like this:
-	 * @code
-	 * TrainingObject->GetPackage()->SetPackageFlags(PKG_EditorOnly).
-	 * @endcode
-	 */ 
-	virtual void SetAssetEditorOnlyFlags();
-
 	/**
 	 * Extract the number of imported (DCC) vertices from a skeletal mesh.
 	 * @param SkeletalMesh The skeletal mesh to get the number of imported vertices for.
diff --git a/Engine/Plugins/Animation/MLDeformer/MLDeformerFramework/Source/MLDeformerFramework/Public/MLDeformerModelInstance.h b/Engine/Plugins/Animation/MLDeformer/MLDeformerFramework/Source/MLDeformerFramework/Public/MLDeformerModelInstance.h
--- a/Engine/Plugins/Animation/MLDeformer/MLDeformerFramework/Source/MLDeformerFramework/Public/MLDeformerModelInstance.h
+++ b/Engine/Plugins/Animation/MLDeformer/MLDeformerFramework/Source/MLDeformerFramework/Public/MLDeformerModelInstance.h
@@ -41,6 +41,13 @@ public:
 	 */
 	virtual void Init(USkeletalMeshComponent* SkelMeshComponent);
 
+	/**
+	 * Initialize this model, after the ML deformer component has been initialized.
+	 * This is a post-step after component initialization. It can be used to register things such as 
+	 * external morph target sets that are generated by this model.
+	 */
+	virtual void PostMLDeformerComponentInit() {}
+
 	/**
 	 * Release all built lookup tables etc.
 	 * Also call this base class method if you override this method.
@@ -55,17 +62,6 @@ public:
 	 */
 	virtual void Tick(float DeltaTime, float ModelWeight);
 
-	/**
-	 * Update the neural network input values directly inside its input tensor.
-	 * This will copy over the bone transforms (and possibly morph and curve weights) into this flat array.
-	 * @param InputData The buffer of floats to write the input values to.
-	 * @param NumInputFloats The number of floats of the InputData buffer. Do not write more than this number of floats to the InputData buffer.
-	 * @return Returns the new buffer element index. For example if there were 2 bones and 3 curves, and each bone takes 6 inputs and each curve takes one 
-	 * input, then after calling this base class method, it will return 15 (2*6 + 3*1). You could call this base class method and then write data after that, starting
-	 * from the offset that is returned.
-	 */
-	virtual int64 SetNeuralNetworkInputValues(float* InputData, int64 NumInputFloats);
-
 	/**
 	 * Check whether the deformer is compatible with a given skeletal mesh component.
 	 * This internally also edits the value returned by GetCompatibilityErrorText().
@@ -84,30 +80,6 @@ public:
 	 */
 	virtual bool IsValidForDataProvider() const;
 
-	/**
-	 * Setup the neural network for this frame.
-	 * This has to perform compatibility checks, valid pointer checks, and call the SetNeuralNetworkInputValues.
-	 * After this method is executed, the neural network should be ready to execute.
-	 * The UMLDeformerModelInstance::Tick method internally calls both the SetupNeuralNetworkForFrame and RunNeuralNetwork methods.
-	 * @return Returns true when the setup is done correctly and the neural network is ready to be executed. Otherwise false is returned, which 
-	 * can happen when the NeuralNetwork pointer is invalid, when the model is not set, when the network is not compatible, etc.
-	 */
-	virtual bool SetupNeuralNetworkForFrame();
-
-	/**
-	 * Run the neural network.
-	 * This already assumes that compatibility checks are done, and that the network inputs are set etc.
-	 * Internally this will execute the UNeuralNetwork::Run() method, either on GPU or CPU.
-	 * @param ModelWeight The weight of the model, must be between 0 and 1.
-	 */
-	virtual void RunNeuralNetwork(float ModelWeight);
-
-	/**
-	 * Get the current component space bone transforms that we grabbed from the skeletal mesh component.
-	 * @return The cached bone transforms, in component space.
-	 */
-	const TArray<FTransform>& GetBoneTransforms() const;
-
 	/**
 	 * Is the deformer asset used compatible with the skeletal mesh component used during the Init call?
 	 * @return True if compatible, false if not.
@@ -148,6 +120,49 @@ public:
 	bool HasPostInitialized() const;
 
 protected:
+	/**
+	 * Update the neural network input values directly inside its input tensor.
+	 * This will copy over the bone transforms (and possibly morph and curve weights) into this flat array.
+	 * @param InputData The buffer of floats to write the input values to.
+	 * @param NumInputFloats The number of floats of the InputData buffer. Do not write more than this number of floats to the InputData buffer.
+	 * @return Returns the new buffer element index. For example if there were 2 bones and 3 curves, and each bone takes 6 inputs and each curve takes one 
+	 * input, then after calling this base class method, it will return 15 (2*6 + 3*1). You could call this base class method and then write data after that, starting
+	 * from the offset that is returned.
+	 */
+	virtual int64 SetNeuralNetworkInputValues(float* InputData, int64 NumInputFloats);
+
+	/**
+	 * Setup the neural network for this frame.
+	 * This has to perform compatibility checks, valid pointer checks, and call the SetNeuralNetworkInputValues.
+	 * After this method is executed, the neural network should be ready to execute.
+	 * The UMLDeformerModelInstance::Tick method internally calls both the SetupNeuralNetworkForFrame and RunNeuralNetwork methods.
+	 * @return Returns true when the setup is done correctly and the neural network is ready to be executed. Otherwise false is returned, which 
+	 * can happen when the NeuralNetwork pointer is invalid, when the model is not set, when the network is not compatible, etc.
+	 */
+	virtual bool SetupInputs();
+
+	/**
+	 * Execute the model instance, which can run the neural network inference in case the model uses a neural network.
+	 * This already assumes that compatibility checks are done, and that the network inputs are set etc.
+	 * Internally this will typically execute the UNeuralNetwork::Run() method, either on GPU or CPU.
+	 * @param ModelWeight The weight of the model, must be between 0 and 1.
+	 */
+	virtual void Execute(float ModelWeight);
+
+	/**
+	 * Handle when the model weight is zero.
+	 * This is called when ticking with a zero model weight.
+	 * The SetupNeuralNetworkForFrame and RunNeuralNetwork calls are skipped in case we have a model weight near zero, as optimization.
+	 * Instead, this method is called in that case. It can be used to for example set morph target weights to zero.
+	 */
+	virtual void HandleZeroModelWeight() {}
+
+	/**
+	 * Get the current component space bone transforms that we grabbed from the skeletal mesh component.
+	 * @return The cached bone transforms, in component space.
+	 */
+	const TArray<FTransform>& GetBoneTransforms() const;
+
 	/**
 	 * Set the bone transformations inside a given output buffer, starting from a given StartIndex.
 	 * @param OutputBuffer The buffer we need to write the transformations to.
diff --git a/Engine/Plugins/Animation/MLDeformer/MLDeformerFramework/Source/MLDeformerFramework/Public/MLDeformerMorphModel.h b/Engine/Plugins/Animation/MLDeformer/MLDeformerFramework/Source/MLDeformerFramework/Public/MLDeformerMorphModel.h
--- a/Engine/Plugins/Animation/MLDeformer/MLDeformerFramework/Source/MLDeformerFramework/Public/MLDeformerMorphModel.h
+++ b/Engine/Plugins/Animation/MLDeformer/MLDeformerFramework/Source/MLDeformerFramework/Public/MLDeformerMorphModel.h
@@ -23,7 +23,6 @@ public:
 	virtual FString GetDisplayName() const override			{ return "Morph Base Model"; }
 	virtual bool IsNeuralNetworkOnGPU() const override		{ return false; }	// CPU based neural network.
 	virtual void Serialize(FArchive& Archive) override;
-	virtual void PostMLDeformerComponentInit(UMLDeformerModelInstance* ModelInstance) override;
 	virtual UMLDeformerModelInstance* CreateModelInstance(UMLDeformerComponent* Component) override;
 	// ~END UMLDeformerModel overrides.
 
@@ -73,25 +72,6 @@ public:
 	 */
 	const TArray<FVector3f>& GetMorphTargetDeltas() const	{ return MorphTargetDeltas; }
 
-	/**
-	 * Get the external morph target set ID for this model.
-	 * This basically identifies the set of morph targets that belong to this model.
-	 * Different models on the same skeletal mesh gives each model its own unique ID.
-	 * You can use this ID to find the weight values for a specific model instance, inside the USkinnedMeshComponent class.
-	 * @return The unique ID of the morph target set for this model.
-	 * @see USkinnedMeshComponent::GetExternalMorphWeights.
-	 * @see FindExternalMorphWeights.
-	 */
-	int32 GetExternalMorphSetID() const						{ return ExternalMorphSetID; }
-
-	/**
-	 * Get the weights for the external morph target set that belongs to this model.
-	 * @param LOD The LOD level to get the weights for.
-	 * @param SkinnedMeshComponent The skinned mesh component that we have to search in.
-	 * @return A pointer to the weight data, or nullptr in case we cannot find the weight data.
-	 */
-	FExternalMorphSetWeights* FindExternalMorphWeights(int32 LOD, USkinnedMeshComponent* SkinnedMeshComponent) const;
-
 	/**
 	 * Get the morph target set.
 	 */
@@ -106,9 +86,6 @@ public:
 	int32 GetMorphTargetDeltaStartIndex(int32 MorphTargetIndex) const;
 
 private:
-	/** The next free morph target set ID. This is used to generate unique ID's for each morph model. */
-	static TAtomic<int32> NextFreeMorphSetID;
-
 	/** The compressed morph target data, ready for the GPU. */
 	TSharedPtr<FExternalMorphSet> MorphTargetSet;
 
@@ -118,12 +95,6 @@ private:
 	 */
 	TArray<FVector3f> MorphTargetDeltas;
 
-	/** 
-	 * The external morph set data type ID, specific to this model.
-	 * This will get automatically set. Each model instance will get its own unique ID, using the NextFreeMorphSetID atomic that will be increased.
-	 */
-	int32 ExternalMorphSetID = -1;
-
 #if WITH_EDITORONLY_DATA
 	/**
 	 * Morph target delta values that are smaller than or equal to this threshold will be zeroed out.
diff --git a/Engine/Plugins/Animation/MLDeformer/MLDeformerFramework/Source/MLDeformerFramework/Public/MLDeformerMorphModelInstance.h b/Engine/Plugins/Animation/MLDeformer/MLDeformerFramework/Source/MLDeformerFramework/Public/MLDeformerMorphModelInstance.h
--- a/Engine/Plugins/Animation/MLDeformer/MLDeformerFramework/Source/MLDeformerFramework/Public/MLDeformerMorphModelInstance.h
+++ b/Engine/Plugins/Animation/MLDeformer/MLDeformerFramework/Source/MLDeformerFramework/Public/MLDeformerMorphModelInstance.h
@@ -5,6 +5,9 @@
 #include "MLDeformerModelInstance.h"
 #include "MLDeformerMorphModelInstance.generated.h"
 
+struct FExternalMorphSetWeights;
+class USkeletalMeshComponent;
+
 /**
  * The model instance for the UMLDeformerMorphModel.
  * This instance will assume the neural network outputs a set of weights, one for each morph target.
@@ -21,7 +24,27 @@ class MLDEFORMERFRAMEWORK_API UMLDeformerMorphModelInstance
 
 public:
 	// UMLDeformerModelInstance overrides.
+	virtual void Init(USkeletalMeshComponent* SkelMeshComponent) override;
+	virtual void PostMLDeformerComponentInit() override;
 	virtual void Release() override;
-	virtual void RunNeuralNetwork(float ModelWeight) override;
+	virtual void Execute(float ModelWeight) override;
+	virtual void HandleZeroModelWeight() override;
 	// ~END UMLDeformerModelInstance overrides.
+
+	int32 GetExternalMorphSetID() const { return ExternalMorphSetID; }
+
+protected:
+	/**
+	 * Find the external morph target weight data for this model instance.
+	 * @param LOD The LOD level to get the weight data for.
+	 * @return A pointer to the weight data, or a nullptr in case it cannot be found.
+	 */
+	FExternalMorphSetWeights* FindWeightData(int32 LOD) const;
+
+protected:
+	/** The next free morph target set ID. This is used to generate unique ID's for each morph model. */
+	static TAtomic<int32> NextFreeMorphSetID;
+
+	/** The ID of the external morph target set for this instance. This gets initialized during Init. */
+	int32 ExternalMorphSetID = -1;
 };
diff --git a/Engine/Plugins/Animation/MLDeformer/MLDeformerFramework/Source/MLDeformerFrameworkEditor/Private/MLDeformerEditorModel.cpp b/Engine/Plugins/Animation/MLDeformer/MLDeformerFramework/Source/MLDeformerFrameworkEditor/Private/MLDeformerEditorModel.cpp
--- a/Engine/Plugins/Animation/MLDeformer/MLDeformerFramework/Source/MLDeformerFrameworkEditor/Private/MLDeformerEditorModel.cpp
+++ b/Engine/Plugins/Animation/MLDeformer/MLDeformerFramework/Source/MLDeformerFrameworkEditor/Private/MLDeformerEditorModel.cpp
@@ -69,6 +69,8 @@ namespace UE::MLDeformer
 		WorkingRange = TRange<double>(0.0, 100.0);
 		PlaybackRange = TRange<double>(0.0, 100.0);
 
+		Model->LogPackagingWarnings();
+
 		AnimSequenceFlags.SetAsset(Model->GetAnimSequence());
 
 		PostEditPropertyDelegateHandle = Model->OnPostEditChangeProperty().AddRaw(this, &FMLDeformerEditorModel::OnPostEditChangeProperty);
@@ -1577,7 +1579,6 @@ namespace UE::MLDeformer
 	{
 		OnInputAssetsChanged();
 		OnPostInputAssetChanged();
-		Model->SetAssetEditorOnlyFlags();
 		GetEditor()->GetModelDetailsView()->ForceRefresh();
 		if (bRefreshVizSettings)
 		{
diff --git a/Engine/Plugins/Animation/MLDeformer/MLDeformerFramework/Source/MLDeformerFrameworkEditor/Public/MLDeformerEditorModel.h b/Engine/Plugins/Animation/MLDeformer/MLDeformerFramework/Source/MLDeformerFrameworkEditor/Public/MLDeformerEditorModel.h
--- a/Engine/Plugins/Animation/MLDeformer/MLDeformerFramework/Source/MLDeformerFrameworkEditor/Public/MLDeformerEditorModel.h
+++ b/Engine/Plugins/Animation/MLDeformer/MLDeformerFramework/Source/MLDeformerFrameworkEditor/Public/MLDeformerEditorModel.h
@@ -6,6 +6,7 @@
 #include "UObject/GCObject.h"
 #include "UObject/ObjectMacros.h"
 #include "UObject/ObjectPtr.h"
+#include "UObject/Package.h"
 #include "MLDeformerEditorActor.h"
 #include "Misc/FrameTime.h"
 #include "MLDeformerVizSettings.h"
@@ -69,20 +70,38 @@ namespace UE::MLDeformer
 			// Restore flags of the current asset, before we change to the new.
 			if (Asset)
 			{
+				UPackage* Package = Asset->GetPackage();
 				if (bEditorOnly)
 				{
-					Asset->GetPackage()->SetPackageFlags(PKG_EditorOnly);
+					Package->SetPackageFlags(PKG_EditorOnly);
+					Package->SetDirtyFlag(bWasDirty);
 				}
 				else
 				{
-					Asset->GetPackage()->ClearPackageFlags(PKG_EditorOnly);
+					Package->ClearPackageFlags(PKG_EditorOnly);
+					Package->SetDirtyFlag(bWasDirty);
 					UE_LOG(LogMLDeformer, Display, TEXT("Unmarked '%s' as editor only asset. Asset can be included in packaging again."), *Asset->GetName());
 				}
 			}
 
-			// Update to the new asset.
+			// Update to the new asset, and mark it as editor only if it isn't yet.
 			Asset = InAsset;
-			bEditorOnly = !Asset.IsNull() ? static_cast<bool>(Asset->GetPackage()->GetPackageFlags() & PKG_EditorOnly) : false;
+			if (!Asset.IsNull())
+			{
+				UPackage* Package = Asset->GetPackage();
+				bEditorOnly = static_cast<bool>(Package->GetPackageFlags() & PKG_EditorOnly);
+				bWasDirty = Package->IsDirty();
+				if (!bEditorOnly)
+				{
+					UE_LOG(LogMLDeformer, Display, TEXT("Marking '%s' as editor only asset. Asset will be excluded from packaging. This requires you to save the asset."), *Asset->GetName());
+					Package->SetPackageFlags(PKG_EditorOnly);
+					Package->SetDirtyFlag(true);
+				}
+			}
+			else
+			{
+				bEditorOnly = false;
+			}
 		}
 
 	private:
@@ -91,6 +110,9 @@ namespace UE::MLDeformer
 
 		/** Was this asset previously marked as editor only already? */
 		bool bEditorOnly = false;
+
+		/** Was the previous asset dirty? */
+		bool bWasDirty = false;
 	};
 
 	/**
diff --git a/Engine/Plugins/Animation/MLDeformer/NearestNeighborModel/Source/NearestNeighborModel/Private/NearestNeighborModelInstance.cpp b/Engine/Plugins/Animation/MLDeformer/NearestNeighborModel/Source/NearestNeighborModel/Private/NearestNeighborModelInstance.cpp
--- a/Engine/Plugins/Animation/MLDeformer/NearestNeighborModel/Source/NearestNeighborModel/Private/NearestNeighborModelInstance.cpp
+++ b/Engine/Plugins/Animation/MLDeformer/NearestNeighborModel/Source/NearestNeighborModel/Private/NearestNeighborModelInstance.cpp
@@ -76,19 +76,18 @@ void AddCopyToCPUPass(FRDGBuilder& GraphBuilder, TRDGBuffer RDGBuffer, T* Ptr, i
         });
 }
 
-bool UNearestNeighborModelInstance::SetupNeuralNetworkForFrame()
+bool UNearestNeighborModelInstance::SetupInputs()
 {
-    bool bSuccess = UMLDeformerModelInstance::SetupNeuralNetworkForFrame();
+    const bool bSuccess = Super::SetupInputs();
 
     // Clip network inputs based on training set min and max
     UNeuralNetwork* NeuralNetwork = Model->GetNeuralNetwork();
     if (NeuralNetwork)
     {
-    	UNearestNeighborModel *NearestNeighborModel = static_cast<UNearestNeighborModel*>(Model.Get());
+    	UNearestNeighborModel* NearestNeighborModel = static_cast<UNearestNeighborModel*>(Model.Get());
     	float* InputDataPointer = static_cast<float*>(NeuralNetwork->GetInputDataPointerMutableForContext(NeuralNetworkInferenceHandle));
     	const int64 NumNeuralNetInputs = NeuralNetwork->GetInputTensorForContext(NeuralNetworkInferenceHandle).Num();
     	NearestNeighborModel->ClipInputs(InputDataPointer, NumNeuralNetInputs);
-
     	return bSuccess;
     }
     else
@@ -97,9 +96,9 @@ bool UNearestNeighborModelInstance::SetupNeuralNetworkForFrame()
     }
 }
 
-void UNearestNeighborModelInstance::RunNeuralNetwork(float ModelWeight)
+void UNearestNeighborModelInstance::Execute(float ModelWeight)
 {
-	Super::RunNeuralNetwork(ModelWeight);
+	Super::Execute(ModelWeight);
 	RunNearestNeighborModel(ModelWeight);
 }
 
@@ -110,23 +109,19 @@ void UNearestNeighborModelInstance::RunNearestNeighborModel(float ModelWeight)
 	{
 		return;
 	}
-
-	const int32 ExternalMorphSetID = NearestNeighborModel->GetExternalMorphSetID();
-	// If this check fails please set this member to some value larger than 0 in your model's constructor.
-	checkf(ExternalMorphSetID != -1, TEXT("Please set the ExternalMorphSetID member value to a unique value for your model type."));
-
+	
 	// Grab the weight data for this morph set.
 	// This could potentially fail if we are applying this deformer to the wrong skeletal mesh component.
 	const int LOD = 0;	// For now we only support LOD 0, as we can't setup an ML Deformer per LOD yet.
-	FExternalMorphSetWeights* WeightData = SkeletalMeshComponent->GetExternalMorphWeights(LOD).MorphSets.Find(ExternalMorphSetID);
+	FExternalMorphSetWeights* WeightData = FindWeightData(LOD);
 	if (WeightData == nullptr)
 	{
 		return;
 	}
 
-	if (ModelWeight > 0.0f)
+	const UNeuralNetwork* NeuralNetwork = Model->GetNeuralNetwork();
+	if (NeuralNetwork)
 	{
-		const UNeuralNetwork* NeuralNetwork = Model->GetNeuralNetwork();
 		const FNeuralTensor& OutputTensor = NeuralNetwork->GetOutputTensorForContext(NeuralNetworkInferenceHandle);
 		const int32 NumNetworkWeights = OutputTensor.Num();
 		const int32 NumMorphTargets = WeightData->Weights.Num();;
@@ -145,15 +140,15 @@ void UNearestNeighborModelInstance::RunNearestNeighborModel(float ModelWeight)
 			for (int32 PartId = 0; PartId < NearestNeighborModel->GetNumParts(); PartId++)
 			{
 				const int32 NearestNeighborId = FindNearestNeighbor(OutputTensor, PartId);
-#if WITH_EDITORONLY_DATA
+	#if WITH_EDITORONLY_DATA
 				NearestNeighborModel->SetNearestNeighborId(PartId, NearestNeighborId);
-#endif
+	#endif
 
 				const int32 NumNeighbors = NearestNeighborModel->GetNumNeighbors(PartId);
 				for (int32 NeighborId = 0; NeighborId < NumNeighbors; NeighborId++)
 				{
-					 const float W = NeighborId == NearestNeighborId ? ModelWeight : 0;
-					 UpdateWeight(WeightData->Weights, NeighborOffset + NeighborId, W);
+					const float W = NeighborId == NearestNeighborId ? ModelWeight : 0;
+					UpdateWeight(WeightData->Weights, NeighborOffset + NeighborId, W);
 				}
 				NeighborOffset += NumNeighbors;
 			}
@@ -161,18 +156,14 @@ void UNearestNeighborModelInstance::RunNearestNeighborModel(float ModelWeight)
 	}
 	else
 	{
-		for (float& Weight : WeightData->Weights)
-		{
-			Weight = 0.0f;
-		}
+		WeightData->ZeroWeights();
 	}
 }
 
-void UNearestNeighborModelInstance::UpdateWeight(TArray<float>& MorphWeights, const int32 Index, const float W)
+void UNearestNeighborModelInstance::UpdateWeight(TArray<float>& MorphWeights, int32 Index, float W)
 {
 	UNearestNeighborModel* NearestNeighborModel = static_cast<UNearestNeighborModel*>(Model);
 	const float DecayFactor = NearestNeighborModel->GetDecayFactor();
-
 	const float PreviousW = NearestNeighborModel->PreviousWeights[Index];
 	const float NewW = (1 - DecayFactor) * W + DecayFactor * PreviousW;
 	MorphWeights[Index] = NewW;
diff --git a/Engine/Plugins/Animation/MLDeformer/NearestNeighborModel/Source/NearestNeighborModel/Public/NearestNeighborModelInstance.h b/Engine/Plugins/Animation/MLDeformer/NearestNeighborModel/Source/NearestNeighborModel/Public/NearestNeighborModelInstance.h
--- a/Engine/Plugins/Animation/MLDeformer/NearestNeighborModel/Source/NearestNeighborModel/Public/NearestNeighborModelInstance.h
+++ b/Engine/Plugins/Animation/MLDeformer/NearestNeighborModel/Source/NearestNeighborModel/Public/NearestNeighborModelInstance.h
@@ -7,14 +7,14 @@
 
 UCLASS()
 class NEARESTNEIGHBORMODEL_API UNearestNeighborModelInstance
-    : public UMLDeformerModelInstance
+    : public UMLDeformerMorphModelInstance
 {
     GENERATED_BODY()
-public:
 
-    // UMLDeformerModelInstance overrides
-    virtual void RunNeuralNetwork(float ModelWeight) override;
-    virtual bool SetupNeuralNetworkForFrame() override;
+public:
+	// UMLDeformerModelInstance overrides
+    virtual void Execute(float ModelWeight) override;
+    virtual bool SetupInputs() override;
 
 protected:
     virtual int64 SetBoneTransforms(float* OutputBuffer, int64 OutputBufferSize, int64 StartIndex) override;
@@ -23,5 +23,5 @@ protected:
 private:
     void RunNearestNeighborModel(float ModelWeight);
     int32 FindNearestNeighbor(const FNeuralTensor& PCACoeffTensor, int32 PartId);
-    void UpdateWeight(TArray<float>& MorphWeights, const int32 Index, const float W);
+    void UpdateWeight(TArray<float>& MorphWeights, int32 Index, float W);
 };
