MLDeformer: Fixed an issue where a deformer that could not get a GPU device could result in a crash

diff --git a/Engine/Plugins/Animation/MLDeformer/MLDeformerFramework/Source/MLDeformerFramework/Private/MLDeformerModelInstance.cpp b/Engine/Plugins/Animation/MLDeformer/MLDeformerFramework/Source/MLDeformerFramework/Private/MLDeformerModelInstance.cpp
--- a/Engine/Plugins/Animation/MLDeformer/MLDeformerFramework/Source/MLDeformerFramework/Private/MLDeformerModelInstance.cpp
+++ b/Engine/Plugins/Animation/MLDeformer/MLDeformerFramework/Source/MLDeformerFramework/Private/MLDeformerModelInstance.cpp
@@ -310,18 +310,22 @@ void UMLDeformerModelInstance::Execute(float ModelWeight)
 
 	if (Model->IsNeuralNetworkOnGPU())
 	{
-		// NOTE: Inputs still come from the CPU.
-		check(NeuralNetwork->GetDeviceType() == ENeuralDeviceType::GPU && NeuralNetwork->GetInputDeviceType() == ENeuralDeviceType::CPU && NeuralNetwork->GetOutputDeviceType() == ENeuralDeviceType::GPU);
-		ENQUEUE_RENDER_COMMAND(RunNeuralNetwork)
-		(
-			[NeuralNetwork, Handle = NeuralNetworkInferenceHandle](FRHICommandListImmediate& RHICmdList)
-			{
-				// Output deltas will be available on GPU for DeformerGraph via UMLDeformerDataProvider.
-				FRDGBuilder GraphBuilder(RHICmdList);
-				NeuralNetwork->Run(GraphBuilder, Handle);
-				GraphBuilder.Execute();
-			}
-		);
+		// Even if the model needs the GPU it is possible that the hardware does not support GPU evaluation
+		if (NeuralNetwork->GetDeviceType() == ENeuralDeviceType::GPU)
+		{
+			// NOTE: Inputs still come from the CPU.
+			check(NeuralNetwork->GetDeviceType() == ENeuralDeviceType::GPU && NeuralNetwork->GetInputDeviceType() == ENeuralDeviceType::CPU && NeuralNetwork->GetOutputDeviceType() == ENeuralDeviceType::GPU);
+			ENQUEUE_RENDER_COMMAND(RunNeuralNetwork)
+				(
+					[NeuralNetwork, Handle = NeuralNetworkInferenceHandle](FRHICommandListImmediate& RHICmdList)
+					{
+						// Output deltas will be available on GPU for DeformerGraph via UMLDeformerDataProvider.
+						FRDGBuilder GraphBuilder(RHICmdList);
+						NeuralNetwork->Run(GraphBuilder, Handle);
+						GraphBuilder.Execute();
+					}
+			);
+		}
 	}
 	else
 	{
