[MLDeformer] Improved handling of incompatible meshes and special case morph targets.

diff --git a/Engine/Plugins/Animation/MLDeformer/MLDeformerFramework/Source/MLDeformerFrameworkEditor/Private/MLDeformerEditorModel.cpp b/Engine/Plugins/Animation/MLDeformer/MLDeformerFramework/Source/MLDeformerFrameworkEditor/Private/MLDeformerEditorModel.cpp
--- a/Engine/Plugins/Animation/MLDeformer/MLDeformerFramework/Source/MLDeformerFrameworkEditor/Private/MLDeformerEditorModel.cpp
+++ b/Engine/Plugins/Animation/MLDeformer/MLDeformerFramework/Source/MLDeformerFrameworkEditor/Private/MLDeformerEditorModel.cpp
@@ -849,7 +849,7 @@ namespace UE::MLDeformer
 				Result = FText::Format(LOCTEXT("BaseMeshMismatch", "Number of vertices in base mesh has changed from {0} to {1} vertices since this ML Deformer Asset was saved! {2}"),
 					Model->GetInputInfo()->GetNumBaseMeshVertices(),
 					Model->GetNumBaseMeshVerts(),
-					IsTrained() ? LOCTEXT("BaseMeshMismatchNN", "Neural network needs to be retrained!") : FText());
+					IsTrained() ? LOCTEXT("BaseMeshMismatchNN", "Model needs to be retrained!") : FText());
 			}
 		}
 		return Result;
@@ -1600,6 +1600,10 @@ namespace UE::MLDeformer
 	void FMLDeformerEditorModel::CreateEngineMorphTargets(TArray<UMorphTarget*>& OutMorphTargets, const TArray<FVector3f>& Deltas, const FString& NamePrefix, int32 LOD, float DeltaThreshold)
 	{
 		OutMorphTargets.Reset();
+		if (Deltas.IsEmpty())
+		{
+			return;
+		}
 
 		const int32 NumBaseMeshVerts = Model->GetNumBaseMeshVerts();
 		check(Deltas.Num() % NumBaseMeshVerts == 0);
diff --git a/Engine/Plugins/Animation/MLDeformer/MLDeformerFramework/Source/MLDeformerFrameworkEditor/Private/MLDeformerGeomCacheEditorModel.cpp b/Engine/Plugins/Animation/MLDeformer/MLDeformerFramework/Source/MLDeformerFrameworkEditor/Private/MLDeformerGeomCacheEditorModel.cpp
--- a/Engine/Plugins/Animation/MLDeformer/MLDeformerFramework/Source/MLDeformerFrameworkEditor/Private/MLDeformerGeomCacheEditorModel.cpp
+++ b/Engine/Plugins/Animation/MLDeformer/MLDeformerFramework/Source/MLDeformerFrameworkEditor/Private/MLDeformerGeomCacheEditorModel.cpp
@@ -142,19 +142,25 @@ namespace UE::MLDeformer
 		// Make sure every skeletal imported mesh has some geometry track.
 		const int32 NumGeomCacheTracks = GeomCache ? GeomCache->Tracks.Num() : 0;
 		int32 NumSkelMeshes = 0;
-		if (SkeletalMesh)
+		check(SkeletalMesh);
+		FSkeletalMeshModel* ImportedModel = SkeletalMesh->GetImportedModel();
+		if (ImportedModel)
 		{
-			FSkeletalMeshModel* ImportedModel = SkeletalMesh->GetImportedModel();
-			if (ImportedModel)
-			{
-				NumSkelMeshes = ImportedModel->LODModels[0].ImportedMeshInfos.Num();
-			}
+			NumSkelMeshes = ImportedModel->LODModels[0].ImportedMeshInfos.Num();
+		}
+
+		// Check if we have any mappings at all.
+		FMLDeformerGeomCacheSampler* GeomCacheSampler = GetGeomCacheSampler();
+		GeomCacheSampler->Init(this);
+		if (GeomCacheSampler->GetMeshMappings().IsEmpty())
+		{
+			return;
 		}
 
 		// Allow the special case where there is just one mesh and track.
 		if (NumGeomCacheTracks != 1 || NumSkelMeshes != 1)
 		{
-			if (!GetGeomCacheSampler()->GetFailedImportedMeshNames().IsEmpty())
+			if (!GeomCacheSampler->GetFailedImportedMeshNames().IsEmpty())
 			{
 				return;
 			}
diff --git a/Engine/Plugins/Animation/MLDeformer/MLDeformerFramework/Source/MLDeformerFrameworkEditor/Private/MLDeformerGeomCacheModelDetails.cpp b/Engine/Plugins/Animation/MLDeformer/MLDeformerFramework/Source/MLDeformerFrameworkEditor/Private/MLDeformerGeomCacheModelDetails.cpp
--- a/Engine/Plugins/Animation/MLDeformer/MLDeformerFramework/Source/MLDeformerFrameworkEditor/Private/MLDeformerGeomCacheModelDetails.cpp
+++ b/Engine/Plugins/Animation/MLDeformer/MLDeformerFramework/Source/MLDeformerFrameworkEditor/Private/MLDeformerGeomCacheModelDetails.cpp
@@ -70,7 +70,7 @@ namespace UE::MLDeformer
 				.Padding(FMargin(0.0f, 4.0f))
 				[
 					SNew(SWarningOrErrorBox)
-					.MessageStyle(EMessageStyle::Error)
+					.MessageStyle(EMessageStyle::Warning)
 					.Message(ChangedErrorText)
 				]
 			];
diff --git a/Engine/Plugins/Animation/MLDeformer/MLDeformerFramework/Source/MLDeformerFrameworkEditor/Private/MLDeformerModelDetails.cpp b/Engine/Plugins/Animation/MLDeformer/MLDeformerFramework/Source/MLDeformerFrameworkEditor/Private/MLDeformerModelDetails.cpp
--- a/Engine/Plugins/Animation/MLDeformer/MLDeformerFramework/Source/MLDeformerFrameworkEditor/Private/MLDeformerModelDetails.cpp
+++ b/Engine/Plugins/Animation/MLDeformer/MLDeformerFramework/Source/MLDeformerFrameworkEditor/Private/MLDeformerModelDetails.cpp
@@ -87,7 +87,7 @@ namespace UE::MLDeformer
 				.Padding(FMargin(0.0f, 4.0f))
 				[
 					SNew(SWarningOrErrorBox)
-					.MessageStyle(EMessageStyle::Error)
+					.MessageStyle(EMessageStyle::Warning)
 					.Message(ChangedErrorText)
 				]
 			];
diff --git a/Engine/Plugins/Animation/MLDeformer/MLDeformerFramework/Source/MLDeformerFrameworkEditor/Private/MLDeformerMorphModelEditorModel.cpp b/Engine/Plugins/Animation/MLDeformer/MLDeformerFramework/Source/MLDeformerFrameworkEditor/Private/MLDeformerMorphModelEditorModel.cpp
--- a/Engine/Plugins/Animation/MLDeformer/MLDeformerFramework/Source/MLDeformerFrameworkEditor/Private/MLDeformerMorphModelEditorModel.cpp
+++ b/Engine/Plugins/Animation/MLDeformer/MLDeformerFramework/Source/MLDeformerFrameworkEditor/Private/MLDeformerMorphModelEditorModel.cpp
@@ -111,6 +111,11 @@ namespace UE::MLDeformer
 		check(MorphModel->GetMorphTargetSet().IsValid());
 		FMorphTargetVertexInfoBuffers& MorphBuffers = MorphModel->GetMorphTargetSet()->MorphBuffers;
 		CompressEngineMorphTargets(MorphBuffers, MorphTargets, LOD, MorphModel->GetMorphTargetErrorTolerance());
+
+		if (MorphBuffers.GetNumBatches() == 0 || MorphBuffers.GetNumMorphs() == 0)
+		{
+			MorphBuffers = FMorphTargetVertexInfoBuffers();
+		}
 	}
 
 	void FMLDeformerMorphModelEditorModel::Render(const FSceneView* View, FViewport* Viewport, FPrimitiveDrawInterface* PDI)
diff --git a/Engine/Plugins/Animation/MLDeformer/MLDeformerFramework/Source/MLDeformerFrameworkEditor/Public/MLDeformerGeomCacheSampler.h b/Engine/Plugins/Animation/MLDeformer/MLDeformerFramework/Source/MLDeformerFrameworkEditor/Public/MLDeformerGeomCacheSampler.h
--- a/Engine/Plugins/Animation/MLDeformer/MLDeformerFramework/Source/MLDeformerFrameworkEditor/Public/MLDeformerGeomCacheSampler.h
+++ b/Engine/Plugins/Animation/MLDeformer/MLDeformerFramework/Source/MLDeformerFrameworkEditor/Public/MLDeformerGeomCacheSampler.h
@@ -45,8 +45,10 @@ namespace UE::MLDeformer
 		 */
 		FMLDeformerGetGeomCacheEvent& OnGetGeometryCache()			{ return GetGeometryCacheEvent; }
 
-		UGeometryCacheComponent* GetGeometryCacheComponent() { return GeometryCacheComponent; }
-		UGeometryCacheComponent* GetGeometryCacheComponent() const { return GeometryCacheComponent; }
+		UGeometryCacheComponent* GetGeometryCacheComponent()		{ return GeometryCacheComponent; }
+		UGeometryCacheComponent* GetGeometryCacheComponent() const	{ return GeometryCacheComponent; }
+
+		const TArray<FMLDeformerGeomCacheMeshMapping>& GetMeshMappings() const	{ return MeshMappings; }
 
 	protected:
 		/**
