MLDeformer: Fix scaling of bone weights, which are now 16-bit on the CPU side.

#rnx

diff --git a/Engine/Plugins/Animation/MLDeformer/MLDeformerFramework/Source/MLDeformerFrameworkEditor/Private/MLDeformerSampler.cpp b/Engine/Plugins/Animation/MLDeformer/MLDeformerFramework/Source/MLDeformerFrameworkEditor/Private/MLDeformerSampler.cpp
--- a/Engine/Plugins/Animation/MLDeformer/MLDeformerFramework/Source/MLDeformerFrameworkEditor/Private/MLDeformerSampler.cpp
+++ b/Engine/Plugins/Animation/MLDeformer/MLDeformerFramework/Source/MLDeformerFrameworkEditor/Private/MLDeformerSampler.cpp
@@ -1,12 +1,14 @@
 // Copyright Epic Games, Inc. All Rights Reserved.
 
 #include "MLDeformerSampler.h"
+
 #include "MLDeformerInputInfo.h"
-#include "MLDeformerEditorModule.h"
 #include "MLDeformerModel.h"
 #include "MLDeformerEditorToolkit.h"
+
 #include "Animation/DebugSkelMeshComponent.h"
 #include "Animation/AnimSequence.h"
+#include "BoneWeights.h"
 #include "Engine/SkeletalMesh.h"
 #include "Rendering/SkeletalMeshModel.h"
 #include "Rendering/SkeletalMeshLODModel.h"
@@ -135,11 +137,11 @@ namespace UE::MLDeformer
 		for (int32 InfluenceIndex = 0; InfluenceIndex < NumInfluences; ++InfluenceIndex)
 		{
 			const int32 BoneIndex = SkinWeightBuffer.GetBoneIndex(VertexIndex, InfluenceIndex);
-			const uint8 WeightByte = SkinWeightBuffer.GetBoneWeight(VertexIndex, InfluenceIndex);
+			const uint16 WeightByte = SkinWeightBuffer.GetBoneWeight(VertexIndex, InfluenceIndex);
 			if (WeightByte > 0)
 			{
 				const int32 RealBoneIndex = LODData.RenderSections[SectionIndex].BoneMap[BoneIndex];
-				const float	Weight = static_cast<float>(WeightByte) / 255.0f;
+				const float	Weight = static_cast<float>(WeightByte) * UE::AnimationCore::InvMaxRawBoneWeightFloat;
 				const FMatrix44f& SkinningTransform = BoneMatrices[RealBoneIndex];
 				InvSkinningTransform += SkinningTransform * Weight;
 			}
