Fix implement module to specify correct module name
#rnx

diff --git a/Engine/Plugins/Animation/MLDeformer/MLDeformerFramework/Source/MLDeformerFrameworkEditor/Private/MLDeformerEditorModule.cpp b/Engine/Plugins/Animation/MLDeformer/MLDeformerFramework/Source/MLDeformerFrameworkEditor/Private/MLDeformerEditorModule.cpp
--- a/Engine/Plugins/Animation/MLDeformer/MLDeformerFramework/Source/MLDeformerFrameworkEditor/Private/MLDeformerEditorModule.cpp
+++ b/Engine/Plugins/Animation/MLDeformer/MLDeformerFramework/Source/MLDeformerFrameworkEditor/Private/MLDeformerEditorModule.cpp
@@ -13,7 +13,7 @@
 
 #define LOCTEXT_NAMESPACE "MLDeformerEditorModule"
 
-IMPLEMENT_MODULE(UE::MLDeformer::FMLDeformerEditorModule, MLDeformerEditor)
+IMPLEMENT_MODULE(UE::MLDeformer::FMLDeformerEditorModule, MLDeformerFrameworkEditor)
 
 namespace UE::MLDeformer
 {
