MLDeformer: Fix heat map kernel compilation on Vulkan.
Need to follow the Vulkan structure packing rules.

diff --git a/Engine/Plugins/Animation/MLDeformer/MLDeformerFramework/Shaders/Private/MLDeformerGraphHeatMapDataInterface.ush b/Engine/Plugins/Animation/MLDeformer/MLDeformerFramework/Shaders/Private/MLDeformerGraphHeatMapDataInterface.ush
--- a/Engine/Plugins/Animation/MLDeformer/MLDeformerFramework/Shaders/Private/MLDeformerGraphHeatMapDataInterface.ush
+++ b/Engine/Plugins/Animation/MLDeformer/MLDeformerFramework/Shaders/Private/MLDeformerGraphHeatMapDataInterface.ush
@@ -6,7 +6,7 @@ uint {DataInterfaceName}_HeatMapMode;
 float {DataInterfaceName}_HeatMapMax;
 float {DataInterfaceName}_GroundTruthLerp;
 uint {DataInterfaceName}_GroundTruthBufferSize;
-StructuredBuffer<float3> {DataInterfaceName}_PositionGroundTruthBuffer;
+StructuredBuffer<float> {DataInterfaceName}_PositionGroundTruthBuffer;
 Buffer<uint> {DataInterfaceName}_VertexMapBuffer;
 
 uint ReadNumVertices_{DataInterfaceName}()
@@ -21,7 +21,8 @@ float3 ReadPositionGroundTruth_{DataInterfaceName}(uint VertexIndex)
 	{
 		return float3(0, 0, 0);
 	}
-	return {DataInterfaceName}_PositionGroundTruthBuffer[BufferIndex];
+	BufferIndex *= 3;
+	return float3({DataInterfaceName}_PositionGroundTruthBuffer[BufferIndex], {DataInterfaceName}_PositionGroundTruthBuffer[BufferIndex + 1], {DataInterfaceName}_PositionGroundTruthBuffer[BufferIndex + 2]);
 }
 
 uint ReadHeatMapMode_{DataInterfaceName}()
diff --git a/Engine/Plugins/Animation/MLDeformer/MLDeformerFramework/Source/MLDeformerFramework/Private/MLDeformerGraphDebugDataInterface.cpp b/Engine/Plugins/Animation/MLDeformer/MLDeformerFramework/Source/MLDeformerFramework/Private/MLDeformerGraphDebugDataInterface.cpp
--- a/Engine/Plugins/Animation/MLDeformer/MLDeformerFramework/Source/MLDeformerFramework/Private/MLDeformerGraphDebugDataInterface.cpp
+++ b/Engine/Plugins/Animation/MLDeformer/MLDeformerFramework/Source/MLDeformerFramework/Private/MLDeformerGraphDebugDataInterface.cpp
@@ -174,7 +174,7 @@ namespace UE::MLDeformer
 
 	void FMLDeformerGraphDebugDataProviderProxy::AllocateResources(FRDGBuilder& GraphBuilder)
 	{
-		GroundTruthBuffer = GraphBuilder.CreateBuffer(FRDGBufferDesc::CreateStructuredDesc(sizeof(FVector3f), GroundTruthPositions.Num()), TEXT("MLDeformer.GroundTruthPositions"));
+		GroundTruthBuffer = GraphBuilder.CreateBuffer(FRDGBufferDesc::CreateStructuredDesc(sizeof(float), 3 * GroundTruthPositions.Num()), TEXT("MLDeformer.GroundTruthPositions"));
 		GroundTruthBufferSRV = GraphBuilder.CreateSRV(GroundTruthBuffer);
 		GraphBuilder.QueueBufferUpload(GroundTruthBuffer, GroundTruthPositions.GetData(), sizeof(FVector3f) * GroundTruthPositions.Num(), ERDGInitialDataFlags::None);
 	}
diff --git a/Engine/Plugins/Animation/MLDeformer/MLDeformerFramework/Source/MLDeformerFramework/Public/MLDeformerGraphDebugDataInterface.h b/Engine/Plugins/Animation/MLDeformer/MLDeformerFramework/Source/MLDeformerFramework/Public/MLDeformerGraphDebugDataInterface.h
--- a/Engine/Plugins/Animation/MLDeformer/MLDeformerFramework/Source/MLDeformerFramework/Public/MLDeformerGraphDebugDataInterface.h
+++ b/Engine/Plugins/Animation/MLDeformer/MLDeformerFramework/Source/MLDeformerFramework/Public/MLDeformerGraphDebugDataInterface.h
@@ -23,7 +23,7 @@ class USkeletalMeshComponent;
 	SHADER_PARAMETER(float, HeatMapMax) \
 	SHADER_PARAMETER(float, GroundTruthLerp) \
 	SHADER_PARAMETER(uint32, GroundTruthBufferSize) \
-	SHADER_PARAMETER_RDG_BUFFER_SRV(StructuredBuffer<float3>, PositionGroundTruthBuffer) \
+	SHADER_PARAMETER_RDG_BUFFER_SRV(StructuredBuffer<float>, PositionGroundTruthBuffer) \
 	SHADER_PARAMETER_SRV(Buffer<uint>, VertexMapBuffer)
 
 #define MLDEFORMER_GRAPH_DISPATCH_DEFAULT_DEBUG_PARAMETERS() \
