[MLDeformer]Changed some GetModuleChecked into LoadModuleChecked.

diff --git a/Engine/Plugins/Animation/MLDeformer/MLDeformerFramework/Source/MLDeformerFrameworkEditor/Private/MLDeformerEditorModel.cpp b/Engine/Plugins/Animation/MLDeformer/MLDeformerFramework/Source/MLDeformerFrameworkEditor/Private/MLDeformerEditorModel.cpp
--- a/Engine/Plugins/Animation/MLDeformer/MLDeformerFramework/Source/MLDeformerFrameworkEditor/Private/MLDeformerEditorModel.cpp
+++ b/Engine/Plugins/Animation/MLDeformer/MLDeformerFramework/Source/MLDeformerFrameworkEditor/Private/MLDeformerEditorModel.cpp
@@ -46,7 +46,7 @@ namespace UE::MLDeformer
 
 		ClearWorld();
 
-		FMLDeformerEditorModule& EditorModule = FModuleManager::GetModuleChecked<FMLDeformerEditorModule>("MLDeformerFrameworkEditor");
+		FMLDeformerEditorModule& EditorModule = FModuleManager::LoadModuleChecked<FMLDeformerEditorModule>("MLDeformerFrameworkEditor");
 		EditorModule.GetModelRegistry().RemoveEditorModelInstance(this);
 	}
 
diff --git a/Engine/Plugins/Animation/MLDeformer/MLDeformerFramework/Source/MLDeformerFrameworkEditor/Private/MLDeformerEditorToolkit.cpp b/Engine/Plugins/Animation/MLDeformer/MLDeformerFramework/Source/MLDeformerFrameworkEditor/Private/MLDeformerEditorToolkit.cpp
--- a/Engine/Plugins/Animation/MLDeformer/MLDeformerFramework/Source/MLDeformerFrameworkEditor/Private/MLDeformerEditorToolkit.cpp
+++ b/Engine/Plugins/Animation/MLDeformer/MLDeformerFramework/Source/MLDeformerFrameworkEditor/Private/MLDeformerEditorToolkit.cpp
@@ -45,7 +45,7 @@ namespace UE::MLDeformer
 		FPersonaToolkitArgs PersonaToolkitArgs;
 		PersonaToolkitArgs.OnPreviewSceneCreated = FOnPreviewSceneCreated::FDelegate::CreateSP(this, &FMLDeformerEditorToolkit::HandlePreviewSceneCreated);
 
-		FMLDeformerEditorModule& EditorModule = FModuleManager::GetModuleChecked<FMLDeformerEditorModule>("MLDeformerFrameworkEditor");
+		FMLDeformerEditorModule& EditorModule = FModuleManager::LoadModuleChecked<FMLDeformerEditorModule>("MLDeformerFrameworkEditor");
 		FMLDeformerEditorModelRegistry& ModelRegistry = EditorModule.GetModelRegistry();	
 
 		// Create a new model if none was selected yet.
@@ -126,7 +126,7 @@ namespace UE::MLDeformer
 
 	void FMLDeformerEditorToolkit::ShowNoModelsWarningIfNeeded()
 	{
-		FMLDeformerEditorModule& EditorModule = FModuleManager::GetModuleChecked<FMLDeformerEditorModule>("MLDeformerFrameworkEditor");
+		FMLDeformerEditorModule& EditorModule = FModuleManager::LoadModuleChecked<FMLDeformerEditorModule>("MLDeformerFrameworkEditor");
 		const FMLDeformerEditorModelRegistry& ModelRegistry = EditorModule.GetModelRegistry();	
 		if (ModelRegistry.GetNumRegisteredModels() == 0)
 		{
@@ -189,7 +189,7 @@ namespace UE::MLDeformer
 		const bool bShouldCloseWindowAfterMenuSelection = true;
 		FMenuBuilder MenuBuilder(bShouldCloseWindowAfterMenuSelection, InCommandList);
 
-		FMLDeformerEditorModule& EditorModule = FModuleManager::GetModuleChecked<FMLDeformerEditorModule>("MLDeformerFrameworkEditor");
+		FMLDeformerEditorModule& EditorModule = FModuleManager::LoadModuleChecked<FMLDeformerEditorModule>("MLDeformerFrameworkEditor");
 		const FMLDeformerEditorModelRegistry& ModelRegistry = EditorModule.GetModelRegistry();
 
 		TArray<UClass*> ModelTypes;
@@ -257,7 +257,7 @@ namespace UE::MLDeformer
 			}
 		}
 
-		FMLDeformerEditorModule& EditorModule = FModuleManager::GetModuleChecked<FMLDeformerEditorModule>("MLDeformerFrameworkEditor");
+		FMLDeformerEditorModule& EditorModule = FModuleManager::LoadModuleChecked<FMLDeformerEditorModule>("MLDeformerFrameworkEditor");
 		FMLDeformerEditorModelRegistry& ModelRegistry = EditorModule.GetModelRegistry();
 
 		// Remove existing actors etc.
diff --git a/Engine/Plugins/Animation/MLDeformer/MLDeformerFramework/Source/MLDeformerFrameworkEditor/Private/MLDeformerModelDetails.cpp b/Engine/Plugins/Animation/MLDeformer/MLDeformerFramework/Source/MLDeformerFrameworkEditor/Private/MLDeformerModelDetails.cpp
--- a/Engine/Plugins/Animation/MLDeformer/MLDeformerFramework/Source/MLDeformerFrameworkEditor/Private/MLDeformerModelDetails.cpp
+++ b/Engine/Plugins/Animation/MLDeformer/MLDeformerFramework/Source/MLDeformerFrameworkEditor/Private/MLDeformerModelDetails.cpp
@@ -40,7 +40,7 @@ namespace UE::MLDeformer
 			Model = static_cast<UMLDeformerModel*>(Objects[0].Get());
 
 			// Get the editor model for this runtime model.
-			FMLDeformerEditorModule& EditorModule = FModuleManager::GetModuleChecked<FMLDeformerEditorModule>("MLDeformerFrameworkEditor");
+			FMLDeformerEditorModule& EditorModule = FModuleManager::LoadModuleChecked<FMLDeformerEditorModule>("MLDeformerFrameworkEditor");
 			if (Model)
 			{
 				EditorModel = EditorModule.GetModelRegistry().GetEditorModel(Model);
diff --git a/Engine/Plugins/Animation/MLDeformer/MLDeformerFramework/Source/MLDeformerFrameworkEditor/Private/MLDeformerVizSettingsDetails.cpp b/Engine/Plugins/Animation/MLDeformer/MLDeformerFramework/Source/MLDeformerFrameworkEditor/Private/MLDeformerVizSettingsDetails.cpp
--- a/Engine/Plugins/Animation/MLDeformer/MLDeformerFramework/Source/MLDeformerFrameworkEditor/Private/MLDeformerVizSettingsDetails.cpp
+++ b/Engine/Plugins/Animation/MLDeformer/MLDeformerFramework/Source/MLDeformerFrameworkEditor/Private/MLDeformerVizSettingsDetails.cpp
@@ -39,7 +39,7 @@ namespace UE::MLDeformer
 
 		if (Objects.Num() == 1)
 		{
-			FMLDeformerEditorModule& EditorModule = FModuleManager::GetModuleChecked<FMLDeformerEditorModule>("MLDeformerFrameworkEditor");
+			FMLDeformerEditorModule& EditorModule = FModuleManager::LoadModuleChecked<FMLDeformerEditorModule>("MLDeformerFrameworkEditor");
 			VizSettings = Cast<UMLDeformerVizSettings>(Objects[0]);	
 			Model = VizSettings ? Cast<UMLDeformerModel>(VizSettings->GetOuter()) : nullptr;
 			EditorModel = Model ? EditorModule.GetModelRegistry().GetEditorModel(Model) : nullptr;
