[MLDeformer] Remove the package flagging to exclude things from packaged builds, but use the soft object pointers instead. This fixes a bug that prevented the automatically exclusion to work correctly.

diff --git a/Engine/Plugins/Animation/MLDeformer/MLDeformerFramework/Source/MLDeformerFramework/Private/MLDeformerGeomCacheModel.cpp b/Engine/Plugins/Animation/MLDeformer/MLDeformerFramework/Source/MLDeformerFramework/Private/MLDeformerGeomCacheModel.cpp
--- a/Engine/Plugins/Animation/MLDeformer/MLDeformerFramework/Source/MLDeformerFramework/Private/MLDeformerGeomCacheModel.cpp
+++ b/Engine/Plugins/Animation/MLDeformer/MLDeformerFramework/Source/MLDeformerFramework/Private/MLDeformerGeomCacheModel.cpp
@@ -13,7 +13,7 @@
 void UMLDeformerGeomCacheModel::Serialize(FArchive& Archive)
 {
 	#if WITH_EDITOR
-		if (Archive.IsSaving() && Archive.IsPersistent() && Archive.IsCooking())
+		if (Archive.IsSaving() && Archive.IsCooking())
 		{
 			GeometryCache = nullptr;
 		}
@@ -23,37 +23,6 @@ void UMLDeformerGeomCacheModel::Serialize(FArchive& Archive)
 }
 
 #if WITH_EDITOR
-void UMLDeformerGeomCacheModel::LogPackagingWarnings()
-{
-	const UGeometryCache* GeomCache = GetGeometryCache();
-	if (GeomCache && !GeomCache->GetPackage()->HasAnyPackageFlags(PKG_EditorOnly))
-	{
-		UE_LOG(
-			LogMLDeformer,
-			Warning, 
-			TEXT("Training geometry cache '%s' is not marked as editor only and will be included during packaging. Reopen the '%s' MLD asset and save this geom cache asset to fix this."),
-			*GeomCache->GetName(),
-			*GetDeformerAsset()->GetName());
-	}
-
-	UMLDeformerGeomCacheVizSettings* Viz = GetGeomCacheVizSettings();
-	if (Viz)
-	{
-		const UGeometryCache* TestGeomCache = Viz->GetTestGroundTruth();
-		if (TestGeomCache && !TestGeomCache->GetPackage()->HasAnyPackageFlags(PKG_EditorOnly))
-		{
-			UE_LOG(
-				LogMLDeformer,
-				Warning, 
-				TEXT("Test geometry cache '%s' is not marked as editor only and will be included during packaging. Reopen the '%s' MLD asset and save this geom cache asset to fix this."),
-				*TestGeomCache->GetName(),
-				*GetDeformerAsset()->GetName());
-		}
-	}
-
-	Super::LogPackagingWarnings();
-}
-
 void UMLDeformerGeomCacheModel::UpdateNumTargetMeshVertices()
 {
 	SetNumTargetMeshVerts(UE::MLDeformer::ExtractNumImportedGeomCacheVertices(GetGeometryCache()));
diff --git a/Engine/Plugins/Animation/MLDeformer/MLDeformerFramework/Source/MLDeformerFramework/Private/MLDeformerModel.cpp b/Engine/Plugins/Animation/MLDeformer/MLDeformerFramework/Source/MLDeformerFramework/Private/MLDeformerModel.cpp
--- a/Engine/Plugins/Animation/MLDeformer/MLDeformerFramework/Source/MLDeformerFramework/Private/MLDeformerModel.cpp
+++ b/Engine/Plugins/Animation/MLDeformer/MLDeformerFramework/Source/MLDeformerFramework/Private/MLDeformerModel.cpp
@@ -64,17 +64,20 @@ void UMLDeformerModel::Serialize(FArchive& Archive)
 {
 	TRACE_CPUPROFILER_EVENT_SCOPE(UMLDeformerModel::Serialize)
 	#if WITH_EDITOR
-		if (Archive.IsSaving() && Archive.IsPersistent())
+		if (Archive.IsSaving())
 		{
 			if (Archive.IsCooking())
 			{
-				LogPackagingWarnings();
-
 				AnimSequence = nullptr;
 				VizSettings = nullptr;
 			}
 
 			InitVertexMap();
+		}
+
+		if (!Archive.IsCooking())
+		{
+			// This also triggers the target mesh to be loaded, don't do that while cooking.
 			UpdateCachedNumVertices();
 		}
 	#endif
@@ -99,10 +102,6 @@ void UMLDeformerModel::PostLoad()
 
 	InitGPUData();
 
-	#if WITH_EDITOR
-		UpdateCachedNumVertices();
-	#endif
-
 	UMLDeformerAsset* MLDeformerAsset = Cast<UMLDeformerAsset>(GetOuter());
 	Init(MLDeformerAsset);
 
@@ -205,20 +204,6 @@ void UMLDeformerModel::FloatArrayToVector3Array(const TArray<float>& FloatArray,
 	{
 		return SkeletalMesh ? SkeletalMesh->GetNumImportedVertices() : 0;
 	}
-
-	void UMLDeformerModel::LogPackagingWarnings()
-	{
-		const UAnimSequence* AnimSeq = GetAnimSequence();
-		if (AnimSeq && !AnimSeq->GetPackage()->HasAnyPackageFlags(PKG_EditorOnly))
-		{
-			UE_LOG(
-				LogMLDeformer,
-				Warning, 
-				TEXT("Training animation sequence '%s' is not marked as editor only and will be included during packaging. Reopen the '%s' MLD asset and save this anim sequence asset to fix this."),
-				*AnimSeq->GetName(),
-				*GetDeformerAsset()->GetName());
-		}
-	}
 #endif	// #if WITH_EDITOR
 
 #if WITH_EDITORONLY_DATA
diff --git a/Engine/Plugins/Animation/MLDeformer/MLDeformerFramework/Source/MLDeformerFramework/Public/MLDeformerGeomCacheModel.h b/Engine/Plugins/Animation/MLDeformer/MLDeformerFramework/Source/MLDeformerFramework/Public/MLDeformerGeomCacheModel.h
--- a/Engine/Plugins/Animation/MLDeformer/MLDeformerFramework/Source/MLDeformerFramework/Public/MLDeformerGeomCacheModel.h
+++ b/Engine/Plugins/Animation/MLDeformer/MLDeformerFramework/Source/MLDeformerFramework/Public/MLDeformerGeomCacheModel.h
@@ -38,7 +38,6 @@ public:
 #endif
 #if WITH_EDITOR
 	virtual void UpdateNumTargetMeshVertices() override;
-	virtual void LogPackagingWarnings() override;
 	// ~END UMLDeformerModel overrides.
 
 	/**
diff --git a/Engine/Plugins/Animation/MLDeformer/MLDeformerFramework/Source/MLDeformerFramework/Public/MLDeformerModel.h b/Engine/Plugins/Animation/MLDeformer/MLDeformerFramework/Source/MLDeformerFramework/Public/MLDeformerModel.h
--- a/Engine/Plugins/Animation/MLDeformer/MLDeformerFramework/Source/MLDeformerFramework/Public/MLDeformerModel.h
+++ b/Engine/Plugins/Animation/MLDeformer/MLDeformerFramework/Source/MLDeformerFramework/Public/MLDeformerModel.h
@@ -163,13 +163,6 @@ public:
 #endif
 
 #if WITH_EDITOR
-	/**
-	 * Log package warnings.
-	 * This can identify things like when the training assets aren't marked as editor only.
-	 * If this flag isn't set, those assets will be included during packaging of your project.
-	 */
-	virtual void LogPackagingWarnings();
-
 	/**
 	 * Update the cached number of vertices of both base and target meshes.
 	 */ 
diff --git a/Engine/Plugins/Animation/MLDeformer/MLDeformerFramework/Source/MLDeformerFrameworkEditor/Private/MLDeformerEditorModel.cpp b/Engine/Plugins/Animation/MLDeformer/MLDeformerFramework/Source/MLDeformerFrameworkEditor/Private/MLDeformerEditorModel.cpp
--- a/Engine/Plugins/Animation/MLDeformer/MLDeformerFramework/Source/MLDeformerFrameworkEditor/Private/MLDeformerEditorModel.cpp
+++ b/Engine/Plugins/Animation/MLDeformer/MLDeformerFramework/Source/MLDeformerFrameworkEditor/Private/MLDeformerEditorModel.cpp
@@ -69,10 +69,6 @@ namespace UE::MLDeformer
 		WorkingRange = TRange<double>(0.0, 100.0);
 		PlaybackRange = TRange<double>(0.0, 100.0);
 
-		Model->LogPackagingWarnings();
-
-		AnimSequenceFlags.SetAsset(Model->GetAnimSequence());
-
 		PostEditPropertyDelegateHandle = Model->OnPostEditChangeProperty().AddRaw(this, &FMLDeformerEditorModel::OnPostEditChangeProperty);
 	}
 
@@ -608,7 +604,6 @@ namespace UE::MLDeformer
 		}
 		else if (Property->GetFName() == UMLDeformerModel::GetAnimSequencePropertyName())
 		{
-			AnimSequenceFlags.SetAsset(Model->GetAnimSequence());
 			TriggerInputAssetChanged();
 			SetResamplingInputOutputsNeeded(true);
 		}
diff --git a/Engine/Plugins/Animation/MLDeformer/MLDeformerFramework/Source/MLDeformerFrameworkEditor/Private/MLDeformerGeomCacheEditorModel.cpp b/Engine/Plugins/Animation/MLDeformer/MLDeformerFramework/Source/MLDeformerFrameworkEditor/Private/MLDeformerGeomCacheEditorModel.cpp
--- a/Engine/Plugins/Animation/MLDeformer/MLDeformerFramework/Source/MLDeformerFrameworkEditor/Private/MLDeformerGeomCacheEditorModel.cpp
+++ b/Engine/Plugins/Animation/MLDeformer/MLDeformerFramework/Source/MLDeformerFrameworkEditor/Private/MLDeformerGeomCacheEditorModel.cpp
@@ -22,14 +22,6 @@
 
 namespace UE::MLDeformer
 {
-	void FMLDeformerGeomCacheEditorModel::Init(const InitSettings& Settings)
-	{
-		FMLDeformerEditorModel::Init(Settings);
-
-		TrainingGroundTruthFlags.SetAsset(GetGeomCacheModel()->GetGeometryCache());
-		TestGroundTruthFlags.SetAsset(GetGeomCacheVizSettings()->GetTestGroundTruth());
-	}
-	
 	FMLDeformerSampler* FMLDeformerGeomCacheEditorModel::CreateSampler() const
 	{
 		FMLDeformerGeomCacheSampler* NewSampler = new FMLDeformerGeomCacheSampler();
@@ -250,13 +242,11 @@ namespace UE::MLDeformer
 		// Properties specific to the geometry cache editor model.
 		if (Property->GetFName() == UMLDeformerGeomCacheModel::GetGeometryCachePropertyName())
 		{
-			TrainingGroundTruthFlags.SetAsset(GetGeomCacheModel()->GetGeometryCache());
 			SetResamplingInputOutputsNeeded(true);
 			TriggerInputAssetChanged(true);
 		}
 		else if (Property->GetFName() == UMLDeformerGeomCacheVizSettings::GetTestGroundTruthPropertyName())
 		{
-			TestGroundTruthFlags.SetAsset(GetGeomCacheVizSettings()->GetTestGroundTruth());
 			TriggerInputAssetChanged(true);
 		}
 	}
diff --git a/Engine/Plugins/Animation/MLDeformer/MLDeformerFramework/Source/MLDeformerFrameworkEditor/Public/MLDeformerEditorModel.h b/Engine/Plugins/Animation/MLDeformer/MLDeformerFramework/Source/MLDeformerFrameworkEditor/Public/MLDeformerEditorModel.h
--- a/Engine/Plugins/Animation/MLDeformer/MLDeformerFramework/Source/MLDeformerFrameworkEditor/Public/MLDeformerEditorModel.h
+++ b/Engine/Plugins/Animation/MLDeformer/MLDeformerFramework/Source/MLDeformerFrameworkEditor/Public/MLDeformerEditorModel.h
@@ -56,65 +56,6 @@ namespace UE::MLDeformer
 	class FMLDeformerSampler;
 	class FMLDeformerEditorModel;
 
-	/**
-	 * A helper template to keep track of the packaging flags for a given object. 
-	 * We use this to restore some flags of the current asset when user changes the asset to another one.
-	 */
-	template <class T>
-	class FMLDeformerEditorOnlyAssetFlags
-	{
-	public:
-		/** Update the current asset and flags state. */
-		void SetAsset(T* InAsset)
-		{
-			// Restore flags of the current asset, before we change to the new.
-			if (Asset)
-			{
-				UPackage* Package = Asset->GetPackage();
-				if (bEditorOnly)
-				{
-					Package->SetPackageFlags(PKG_EditorOnly);
-					Package->SetDirtyFlag(bWasDirty);
-				}
-				else
-				{
-					Package->ClearPackageFlags(PKG_EditorOnly);
-					Package->SetDirtyFlag(bWasDirty);
-					UE_LOG(LogMLDeformer, Display, TEXT("Unmarked '%s' as editor only asset. Asset can be included in packaging again."), *Asset->GetName());
-				}
-			}
-
-			// Update to the new asset, and mark it as editor only if it isn't yet.
-			Asset = InAsset;
-			if (Asset)
-			{
-				UPackage* Package = Asset->GetPackage();
-				bEditorOnly = static_cast<bool>(Package->GetPackageFlags() & PKG_EditorOnly);
-				bWasDirty = Package->IsDirty();
-				if (!bEditorOnly)
-				{
-					UE_LOG(LogMLDeformer, Display, TEXT("Marking '%s' as editor only asset. Asset will be excluded from packaging. This requires you to save the asset."), *Asset->GetName());
-					Package->SetPackageFlags(PKG_EditorOnly);
-					Package->SetDirtyFlag(true);
-				}
-			}
-			else
-			{
-				bEditorOnly = false;
-			}
-		}
-
-	private:
-		/** A pointer to the current asset. */
-		TObjectPtr<T> Asset = nullptr;
-
-		/** Was this asset previously marked as editor only already? */
-		bool bEditorOnly = false;
-
-		/** Was the previous asset dirty? */
-		bool bWasDirty = false;
-	};
-
 	/**
 	 * The base class for the editor side of an UMLDeformerModel.
 	 * The editor model class handles most of the editor interactions, such as property changes, triggering the training process,
@@ -893,9 +834,6 @@ namespace UE::MLDeformer
 		/** A pointer to the sampler, which can sample target meshes to calculate deltas. */
 		FMLDeformerSampler* Sampler = nullptr;
 
-		/** The flags state of the anim sequence. */
-		FMLDeformerEditorOnlyAssetFlags<UAnimSequence> AnimSequenceFlags;
-
 		/**
 		 * The input info as currently setup in the editor.
 		 * This is different from the runtime model's input info, as that is the one that was used to train with.
diff --git a/Engine/Plugins/Animation/MLDeformer/MLDeformerFramework/Source/MLDeformerFrameworkEditor/Public/MLDeformerGeomCacheEditorModel.h b/Engine/Plugins/Animation/MLDeformer/MLDeformerFramework/Source/MLDeformerFrameworkEditor/Public/MLDeformerGeomCacheEditorModel.h
--- a/Engine/Plugins/Animation/MLDeformer/MLDeformerFramework/Source/MLDeformerFrameworkEditor/Public/MLDeformerGeomCacheEditorModel.h
+++ b/Engine/Plugins/Animation/MLDeformer/MLDeformerFramework/Source/MLDeformerFrameworkEditor/Public/MLDeformerGeomCacheEditorModel.h
@@ -27,7 +27,6 @@ namespace UE::MLDeformer
 		// ~END FGCObject overrides.
 
 		// FMLDeformerEditorModel overrides.
-		virtual void Init(const InitSettings& Settings) override;
 		virtual FMLDeformerEditorActor* CreateEditorActor(const FMLDeformerEditorActor::FConstructSettings& Settings) const override;
 		virtual FMLDeformerSampler* CreateSampler() const override;
 		virtual void CreateTrainingGroundTruthActor(UWorld* World) override;
@@ -63,12 +62,5 @@ namespace UE::MLDeformer
 		 * @param bIsTrainingActor Set this to true when the actor is an actor to be used in training mode, or set to false when it is to be used in testing mode.
 		 */
 		void CreateGeomCacheActor(UWorld* World, int32 ActorID, const FName& Name, UGeometryCache* GeomCache, FLinearColor LabelColor, FLinearColor WireframeColor, const FText& LabelText, bool bIsTrainingActor);
-
-	protected:
-		/** Packaging flag handling for the training ground truth. */
-		FMLDeformerEditorOnlyAssetFlags<UGeometryCache> TrainingGroundTruthFlags;
-
-		/** Packaging flag handling for the test ground truth. */
-		FMLDeformerEditorOnlyAssetFlags<UGeometryCache> TestGroundTruthFlags;
 	};
 }	// namespace UE::MLDeformer
