[MLDeformer] Fix some crash caused by calling a pure virtual method.

diff --git a/Engine/Plugins/Animation/MLDeformer/MLDeformerFramework/Source/MLDeformerFramework/Private/MLDeformerModel.cpp b/Engine/Plugins/Animation/MLDeformer/MLDeformerFramework/Source/MLDeformerFramework/Private/MLDeformerModel.cpp
--- a/Engine/Plugins/Animation/MLDeformer/MLDeformerFramework/Source/MLDeformerFramework/Private/MLDeformerModel.cpp
+++ b/Engine/Plugins/Animation/MLDeformer/MLDeformerFramework/Source/MLDeformerFramework/Private/MLDeformerModel.cpp
@@ -40,6 +40,11 @@ namespace UE::MLDeformer
 	}
 }	// namespace UE::MLDeformer
 
+FString UMLDeformerModel::GetDisplayName() const
+{ 
+	return GetClass()->GetFName().ToString();
+}
+
 UMLDeformerInputInfo* UMLDeformerModel::CreateInputInfo()
 { 
 	return NewObject<UMLDeformerInputInfo>(this);
@@ -178,6 +183,11 @@ void UMLDeformerModel::FloatArrayToVector3Array(const TArray<float>& FloatArray,
 }
 
 #if WITH_EDITOR
+	void UMLDeformerModel::UpdateNumTargetMeshVertices()
+	{ 
+		NumTargetMeshVerts = 0;
+	}
+
 	void UMLDeformerModel::UpdateNumBaseMeshVertices()
 	{
 		NumBaseMeshVerts = UMLDeformerModel::ExtractNumImportedSkinnedVertices(GetSkeletalMesh());
diff --git a/Engine/Plugins/Animation/MLDeformer/MLDeformerFramework/Source/MLDeformerFramework/Public/MLDeformerGeomCacheModel.h b/Engine/Plugins/Animation/MLDeformer/MLDeformerFramework/Source/MLDeformerFramework/Public/MLDeformerGeomCacheModel.h
--- a/Engine/Plugins/Animation/MLDeformer/MLDeformerFramework/Source/MLDeformerFramework/Public/MLDeformerGeomCacheModel.h
+++ b/Engine/Plugins/Animation/MLDeformer/MLDeformerFramework/Source/MLDeformerFramework/Public/MLDeformerGeomCacheModel.h
@@ -31,7 +31,6 @@ public:
 	// ~END UObject overrides.
 
 	// UMLDeformerModel overrides.
-	virtual FString GetDisplayName() const override			{ return "GeomCache Based Model"; }
 #if WITH_EDITORONLY_DATA
 	virtual bool HasTrainingGroundTruth() const override	{ return (GetGeometryCache() != nullptr); }
 	virtual void SampleGroundTruthPositions(float SampleTime, TArray<FVector3f>& OutPositions) override;
diff --git a/Engine/Plugins/Animation/MLDeformer/MLDeformerFramework/Source/MLDeformerFramework/Public/MLDeformerModel.h b/Engine/Plugins/Animation/MLDeformer/MLDeformerFramework/Source/MLDeformerFramework/Public/MLDeformerModel.h
--- a/Engine/Plugins/Animation/MLDeformer/MLDeformerFramework/Source/MLDeformerFramework/Public/MLDeformerModel.h
+++ b/Engine/Plugins/Animation/MLDeformer/MLDeformerFramework/Source/MLDeformerFramework/Public/MLDeformerModel.h
@@ -98,9 +98,10 @@ public:
 	/**
 	 * Get the display name of this model.
 	 * This will also define with what name this model will appear inside the UI.
+	 * On default this will return the class name.
 	 * @return The name of the model.
 	 */
-	virtual FString GetDisplayName() const PURE_VIRTUAL(UMLDeformerModel::GetDisplayName, return FString(););
+	virtual FString GetDisplayName() const;
 
 	/**
 	 * Defines whether this model supports bone transforms as input or not.
@@ -180,7 +181,7 @@ public:
 	 * The number of vertices should be the number of DCC vertices, so not the number of render vertices.
 	 * This just updates the NumTargetMeshVerts property.
 	 */ 
-	virtual void UpdateNumTargetMeshVertices() PURE_VIRTUAL(UMLDeformerModel::GetNumTargetMeshVertices; mTargetMeshVerts = 0;);
+	virtual void UpdateNumTargetMeshVertices();
 
 	/**
 	 * Extract the number of imported (DCC) vertices from a skeletal mesh.
diff --git a/Engine/Plugins/Animation/MLDeformer/MLDeformerFramework/Source/MLDeformerFramework/Public/MLDeformerMorphModel.h b/Engine/Plugins/Animation/MLDeformer/MLDeformerFramework/Source/MLDeformerFramework/Public/MLDeformerMorphModel.h
--- a/Engine/Plugins/Animation/MLDeformer/MLDeformerFramework/Source/MLDeformerFramework/Public/MLDeformerMorphModel.h
+++ b/Engine/Plugins/Animation/MLDeformer/MLDeformerFramework/Source/MLDeformerFramework/Public/MLDeformerMorphModel.h
@@ -20,7 +20,6 @@ public:
 	UMLDeformerMorphModel(const FObjectInitializer& ObjectInitializer);
 
 	// UMLDeformerModel overrides.
-	virtual FString GetDisplayName() const override			{ return "Morph Base Model"; }
 	virtual bool IsNeuralNetworkOnGPU() const override		{ return false; }	// CPU based neural network.
 	virtual void Serialize(FArchive& Archive) override;
 	virtual UMLDeformerModelInstance* CreateModelInstance(UMLDeformerComponent* Component) override;
diff --git a/Engine/Plugins/Animation/MLDeformer/MLDeformerFramework/Source/MLDeformerFrameworkEditor/Private/MLDeformerEditorModel.cpp b/Engine/Plugins/Animation/MLDeformer/MLDeformerFramework/Source/MLDeformerFrameworkEditor/Private/MLDeformerEditorModel.cpp
--- a/Engine/Plugins/Animation/MLDeformer/MLDeformerFramework/Source/MLDeformerFrameworkEditor/Private/MLDeformerEditorModel.cpp
+++ b/Engine/Plugins/Animation/MLDeformer/MLDeformerFramework/Source/MLDeformerFrameworkEditor/Private/MLDeformerEditorModel.cpp
@@ -1967,6 +1967,17 @@ namespace UE::MLDeformer
 		return FString::Printf(TEXT("%s%s/%s.onnx"), *FPaths::ProjectIntermediateDir(), *ClassName, *ClassName);
 	}
 
+	ETrainingResult FMLDeformerEditorModel::Train()
+	{ 
+		UE_LOG(LogMLDeformer, Warning, TEXT("Please override the FMLDeformerEditorModel::Train() method inside your derived editor model class."));
+		return ETrainingResult::Success;
+	}
+
+	int32 FMLDeformerEditorModel::GetNumTrainingFrames() const
+	{ 
+		UE_LOG(LogMLDeformer, Warning, TEXT("Please override the FMLDeformerEditorModel::GetNumTrainingFrames() method inside your derived editor model class."));
+		return 0;
+	}
 }	// namespace UE::MLDeformer
 
 #undef LOCTEXT_NAMESPACE
diff --git a/Engine/Plugins/Animation/MLDeformer/MLDeformerFramework/Source/MLDeformerFrameworkEditor/Public/MLDeformerEditorModel.h b/Engine/Plugins/Animation/MLDeformer/MLDeformerFramework/Source/MLDeformerFrameworkEditor/Public/MLDeformerEditorModel.h
--- a/Engine/Plugins/Animation/MLDeformer/MLDeformerFramework/Source/MLDeformerFrameworkEditor/Public/MLDeformerEditorModel.h
+++ b/Engine/Plugins/Animation/MLDeformer/MLDeformerFramework/Source/MLDeformerFrameworkEditor/Public/MLDeformerEditorModel.h
@@ -96,7 +96,7 @@ namespace UE::MLDeformer
 		 * For example this could return the number of frames inside your geometry cache.
 		 * @return The number of frames in the training data.
 		 */
-		virtual int32 GetNumTrainingFrames() const PURE_VIRTUAL(FMLDeformerEditorModel::GetNumTrainingFrames, return 0;);
+		virtual int32 GetNumTrainingFrames() const;
 
 		/**
 		 * Launch the training. This gets executed when the Train button is pressed.
@@ -129,7 +129,7 @@ namespace UE::MLDeformer
 		 * 
 		 * @return The training result.
 		 */
-		virtual ETrainingResult Train() PURE_VIRTUAL(FMLDeformerEditorModel::Train, return ETrainingResult::Success;);
+		virtual ETrainingResult Train();
 
 		// Optional overrides.
 		/**
