[MLDeformer] Fixed some crash when saving ML Deformer asset where the source geom cache asset it uses isn't there.

diff --git a/Engine/Plugins/Animation/MLDeformer/MLDeformerFramework/Source/MLDeformerFramework/Private/MLDeformerGeomCacheModel.cpp b/Engine/Plugins/Animation/MLDeformer/MLDeformerFramework/Source/MLDeformerFramework/Private/MLDeformerGeomCacheModel.cpp
--- a/Engine/Plugins/Animation/MLDeformer/MLDeformerFramework/Source/MLDeformerFramework/Private/MLDeformerGeomCacheModel.cpp
+++ b/Engine/Plugins/Animation/MLDeformer/MLDeformerFramework/Source/MLDeformerFramework/Private/MLDeformerGeomCacheModel.cpp
@@ -25,7 +25,11 @@ void UMLDeformerGeomCacheModel::Serialize(FArchive& Archive)
 #if WITH_EDITOR
 void UMLDeformerGeomCacheModel::UpdateNumTargetMeshVertices()
 {
-	SetNumTargetMeshVerts(UE::MLDeformer::ExtractNumImportedGeomCacheVertices(GetGeometryCache()));
+	UGeometryCache* GeomCache = GetGeometryCache();
+	if (GeomCache)
+	{
+		SetNumTargetMeshVerts(UE::MLDeformer::ExtractNumImportedGeomCacheVertices(GeomCache));
+	}
 }
 
 UMLDeformerGeomCacheVizSettings* UMLDeformerGeomCacheModel::GetGeomCacheVizSettings() const
diff --git a/Engine/Plugins/Animation/MLDeformer/MLDeformerFramework/Source/MLDeformerFramework/Private/MLDeformerModel.cpp b/Engine/Plugins/Animation/MLDeformer/MLDeformerFramework/Source/MLDeformerFramework/Private/MLDeformerModel.cpp
--- a/Engine/Plugins/Animation/MLDeformer/MLDeformerFramework/Source/MLDeformerFramework/Private/MLDeformerModel.cpp
+++ b/Engine/Plugins/Animation/MLDeformer/MLDeformerFramework/Source/MLDeformerFramework/Private/MLDeformerModel.cpp
@@ -80,7 +80,7 @@ void UMLDeformerModel::Serialize(FArchive& Archive)
 			InitVertexMap();
 		}
 
-		if (!Archive.IsCooking())
+		if (!Archive.IsCooking() && Archive.IsLoading())
 		{
 			// This also triggers the target mesh to be loaded, don't do that while cooking.
 			UpdateCachedNumVertices();
@@ -197,7 +197,11 @@ void UMLDeformerModel::FloatArrayToVector3Array(const TArray<float>& FloatArray,
 
 	void UMLDeformerModel::UpdateNumBaseMeshVertices()
 	{
-		NumBaseMeshVerts = UMLDeformerModel::ExtractNumImportedSkinnedVertices(GetSkeletalMesh());
+		USkeletalMesh* SkelMesh = GetSkeletalMesh();
+		if (SkelMesh)
+		{
+			NumBaseMeshVerts = UMLDeformerModel::ExtractNumImportedSkinnedVertices(SkelMesh);
+		}
 	}
 
 	void UMLDeformerModel::PostEditChangeProperty(FPropertyChangedEvent& PropertyChangedEvent)
