[UE] Possible fix for crash in ML DeformerModelInstance
- add check(.IsValidIndex()) before an array access
- replace existing index >= 0 array bounds check with .IsValidIndex before another array access

#rnx

diff --git a/Engine/Plugins/Animation/MLDeformer/MLDeformerFramework/Source/MLDeformerFramework/Private/MLDeformerModelInstance.cpp b/Engine/Plugins/Animation/MLDeformer/MLDeformerFramework/Source/MLDeformerFramework/Private/MLDeformerModelInstance.cpp
--- a/Engine/Plugins/Animation/MLDeformer/MLDeformerFramework/Source/MLDeformerFramework/Private/MLDeformerModelInstance.cpp
+++ b/Engine/Plugins/Animation/MLDeformer/MLDeformerFramework/Source/MLDeformerFramework/Private/MLDeformerModelInstance.cpp
@@ -170,9 +170,11 @@ void UMLDeformerModelInstance::UpdateBoneTransforms()
 		for (int32 Index = 0; Index < NumTrainingBones; ++Index)
 		{
 			const int32 ComponentBoneIndex = AssetBonesToSkelMeshMappings[Index];
+			check(LeaderTransforms.IsValidIndex(ComponentBoneIndex));
+			
 			const FTransform& ComponentSpaceTransform = LeaderTransforms[ComponentBoneIndex];
 			const int32 ParentIndex = SkinnedAsset->GetRefSkeleton().GetParentIndex(ComponentBoneIndex);
-			if (ParentIndex != INDEX_NONE)
+			if (LeaderTransforms.IsValidIndex(ParentIndex))
 			{
 				TrainingBoneTransforms[Index] = ComponentSpaceTransform.GetRelativeTransform(LeaderTransforms[ParentIndex]);
 			}
