[MLDeformer] Add an option to include normals in the generated morph targets.

diff --git a/Engine/Plugins/Animation/MLDeformer/MLDeformerFramework/Source/MLDeformerFramework/Public/MLDeformerMorphModel.h b/Engine/Plugins/Animation/MLDeformer/MLDeformerFramework/Source/MLDeformerFramework/Public/MLDeformerMorphModel.h
--- a/Engine/Plugins/Animation/MLDeformer/MLDeformerFramework/Source/MLDeformerFramework/Public/MLDeformerMorphModel.h
+++ b/Engine/Plugins/Animation/MLDeformer/MLDeformerFramework/Source/MLDeformerFramework/Public/MLDeformerMorphModel.h
@@ -32,13 +32,16 @@ public:
 #if WITH_EDITORONLY_DATA
 	float GetMorphTargetDeltaThreshold() const				{ return MorphTargetDeltaThreshold; }
 	float GetMorphTargetErrorTolerance() const				{ return MorphTargetErrorTolerance; }
+	bool GetIncludeMorphTargetNormals() const				{ return bIncludeNormals; }
 
 	void SetMorphTargetDeltaThreshold(float Threshold)		{ MorphTargetDeltaThreshold = Threshold; }
 	void SetMorphTargetErrorTolerance(float Tolerance)		{ MorphTargetErrorTolerance = Tolerance; }
+	void SetIncludeMorphTargetNormals(bool bInclude)		{ bIncludeNormals = bInclude; }
 
 	// Get property names.
 	static FName GetMorphTargetDeltaThresholdPropertyName() { return GET_MEMBER_NAME_CHECKED(UMLDeformerMorphModel, MorphTargetDeltaThreshold); }
 	static FName GetMorphTargetErrorTolerancePropertyName() { return GET_MEMBER_NAME_CHECKED(UMLDeformerMorphModel, MorphTargetErrorTolerance); }
+	static FName GetIncludeMorphTargetNormalsPropertyName() { return GET_MEMBER_NAME_CHECKED(UMLDeformerMorphModel, bIncludeNormals); }
 #endif
 
 	/**
@@ -95,6 +98,14 @@ private:
 	TArray<FVector3f> MorphTargetDeltas;
 
 #if WITH_EDITORONLY_DATA
+	/**
+	 * Include vertex normals in the morph targets?
+	 * The advantage of this can be that it is higher performance than recomputing the normals.
+	 * The disadvantage is it can result in lower quality and uses more memory for the stored morph targets.
+	 */
+	UPROPERTY(EditAnywhere, Category = "Morph Targets")
+	bool bIncludeNormals = false;
+
 	/**
 	 * Morph target delta values that are smaller than or equal to this threshold will be zeroed out.
 	 * This essentially removes small deltas from morph targets, which will lower the memory usage at runtime, however when set too high it can also introduce visual artifacts.
diff --git a/Engine/Plugins/Animation/MLDeformer/MLDeformerFramework/Source/MLDeformerFrameworkEditor/Private/MLDeformerEditorModel.cpp b/Engine/Plugins/Animation/MLDeformer/MLDeformerFramework/Source/MLDeformerFrameworkEditor/Private/MLDeformerEditorModel.cpp
--- a/Engine/Plugins/Animation/MLDeformer/MLDeformerFramework/Source/MLDeformerFrameworkEditor/Private/MLDeformerEditorModel.cpp
+++ b/Engine/Plugins/Animation/MLDeformer/MLDeformerFramework/Source/MLDeformerFrameworkEditor/Private/MLDeformerEditorModel.cpp
@@ -1589,7 +1589,77 @@ namespace UE::MLDeformer
 		}
 	}
 
-	void FMLDeformerEditorModel::CreateEngineMorphTargets(TArray<UMorphTarget*>& OutMorphTargets, const TArray<FVector3f>& Deltas, const FString& NamePrefix, int32 LOD, float DeltaThreshold)
+	void FMLDeformerEditorModel::CalcMeshNormals(TArrayView<const FVector3f> VertexPositions, TArrayView<const uint32> IndexArray, TArrayView<const int32> VertexMap, TArray<FVector3f>& OutNormals) const
+	{
+		const int32 NumVertices = VertexPositions.Num();
+		OutNormals.Reset();
+		OutNormals.SetNumZeroed(NumVertices);
+
+		checkf(IndexArray.Num() % 3 == 0, TEXT("Expecting a triangle mesh!"));
+		const int32 NumTriangles = IndexArray.Num() / 3;
+		for (int32 TriangleIndex = 0; TriangleIndex < NumTriangles; ++TriangleIndex)
+		{
+			const int32 ImportedIndices[3]	{ VertexMap[IndexArray[TriangleIndex * 3]], VertexMap[IndexArray[TriangleIndex * 3 + 1]], VertexMap[IndexArray[TriangleIndex * 3 + 2]] };
+			const FVector3f Positions[3]	{ VertexPositions[ImportedIndices[0]], VertexPositions[ImportedIndices[1]], VertexPositions[ImportedIndices[2]] };
+
+			const FVector3f EdgeA = (Positions[1] - Positions[0]).GetSafeNormal();
+			const FVector3f EdgeB = (Positions[2] - Positions[0]).GetSafeNormal();
+			if (EdgeA.SquaredLength() > 0.00001f && EdgeB.SquaredLength() > 0.00001f)
+			{	
+				const FVector3f FaceNormal = EdgeB.Cross(EdgeA);
+				OutNormals[ImportedIndices[0]] += FaceNormal;
+				OutNormals[ImportedIndices[1]] += FaceNormal;
+				OutNormals[ImportedIndices[2]] += FaceNormal;
+			}
+		}
+
+		// Renormalize.
+		for (int32 Index = 0; Index < NumVertices; ++Index)
+		{
+			OutNormals[Index] = OutNormals[Index].GetSafeNormal();
+		}
+	}
+
+	void FMLDeformerEditorModel::GenerateNormalsForMorphTarget(int32 LOD, USkeletalMesh* SkelMesh, int32 MorphTargetIndex, TArrayView<const FVector3f> Deltas, TArrayView<const FVector3f> BaseVertexPositions, TArrayView<FVector3f> BaseNormals, TArray<FVector3f>& OutDeltaNormals)
+	{
+		const FSkeletalMeshModel* ImportedModel = SkelMesh->GetImportedModel();
+		if (ImportedModel == nullptr || !ImportedModel->LODModels.IsValidIndex(LOD))
+		{
+			OutDeltaNormals.Reset();
+			OutDeltaNormals.SetNumZeroed(Model->GetNumBaseMeshVerts());
+			return;
+		}
+		const TArrayView<const uint32> IndexArray = ImportedModel->LODModels[LOD].IndexBuffer;
+		const TArrayView<const int32> VertexMap = ImportedModel->LODModels[LOD].MeshToImportVertexMap;
+
+		// Build the array of displaced vertex positions.
+		TArray<FVector3f> MorphedVertexPositions;
+		const int32 NumBaseMeshVerts = Model->GetNumBaseMeshVerts();
+		MorphedVertexPositions.SetNumUninitialized(NumBaseMeshVerts);
+		for (int32 VertexIndex = 0; VertexIndex < NumBaseMeshVerts; ++VertexIndex)
+		{
+			const int32 DeltaIndex = (MorphTargetIndex * NumBaseMeshVerts) + VertexIndex;
+			MorphedVertexPositions[VertexIndex] = BaseVertexPositions[VertexIndex] + Deltas[DeltaIndex];
+		}
+
+		// Calculate the normals of that displaced mesh.
+		TArray<FVector3f> MorphedNormals;
+		CalcMeshNormals(MorphedVertexPositions, IndexArray, VertexMap, MorphedNormals);
+
+		// Calculate and output the difference between the morphed normal and base normal.
+		OutDeltaNormals.Reset();
+		OutDeltaNormals.SetNumUninitialized(NumBaseMeshVerts);
+		for (int32 VertexIndex = 0; VertexIndex < NumBaseMeshVerts; ++VertexIndex)
+		{
+			OutDeltaNormals[VertexIndex] = MorphedNormals[VertexIndex] - BaseNormals[VertexIndex];
+			if (OutDeltaNormals[VertexIndex].SquaredLength() < 0.00001f)
+			{
+				OutDeltaNormals[VertexIndex] = FVector3f::ZeroVector;
+			}
+		}
+	}
+
+	void FMLDeformerEditorModel::CreateEngineMorphTargets(TArray<UMorphTarget*>& OutMorphTargets, const TArray<FVector3f>& Deltas, const FString& NamePrefix, int32 LOD, float DeltaThreshold, bool bIncludeNormals)
 	{
 		OutMorphTargets.Reset();
 		if (Deltas.IsEmpty())
@@ -1609,10 +1679,21 @@ namespace UE::MLDeformer
 		check(!RenderData->LODRenderData.IsEmpty());
 		const int32 NumRenderVertices = RenderData->LODRenderData[LOD].GetNumVertices();
 
+		// Calculate the normals for the base mesh.
+		const FSkeletalMeshModel* ImportedModel = SkelMesh->GetImportedModel();
+		const TArrayView<const uint32> IndexArray = ImportedModel->LODModels[LOD].IndexBuffer;
+		const TArrayView<const int32> VertexMap = ImportedModel->LODModels[LOD].MeshToImportVertexMap;
+		const TArrayView<const FVector3f> BaseVertexPositions = Sampler->GetUnskinnedVertexPositions();
+		TArray<FVector3f> BaseNormals;
+		CalcMeshNormals(BaseVertexPositions, IndexArray, VertexMap, BaseNormals);
+
 		// Initialize an engine morph target for each model morph target.
 		UE_LOG(LogMLDeformer, Display, TEXT("Initializing %d engine morph targets of %d vertices each"), NumMorphTargets, Deltas.Num() / NumMorphTargets);
+		TArray<FVector3f> DeltaNormals;
 		for (int32 MorphTargetIndex = 0; MorphTargetIndex < NumMorphTargets; ++MorphTargetIndex)
 		{
+			GenerateNormalsForMorphTarget(LOD, SkelMesh, MorphTargetIndex, Deltas, BaseVertexPositions, BaseNormals, DeltaNormals);
+
 			const FName MorphName = *FString::Printf(TEXT("%s%.3d"), *NamePrefix, MorphTargetIndex);
 			UMorphTarget* MorphTarget = NewObject<UMorphTarget>(SkelMesh, MorphName);
 			MorphTarget->BaseSkelMesh = SkelMesh;
@@ -1638,7 +1719,6 @@ namespace UE::MLDeformer
 			}
 
 			// Init deltas for this morph target.
-			const TArray<int32>& VertexMap = Model->GetVertexMap();
 			MorphLODModel.Vertices.Reserve(NumRenderVertices);
 			for (int32 VertexIndex = 0; VertexIndex < NumRenderVertices; ++VertexIndex)
 			{
@@ -1652,12 +1732,13 @@ namespace UE::MLDeformer
 						FMorphTargetDelta& MorphTargetDelta = MorphLODModel.Vertices.Last();
 						MorphTargetDelta.PositionDelta = Delta;
 						MorphTargetDelta.SourceIdx = VertexIndex;
-						MorphTargetDelta.TangentZDelta = FVector3f::ZeroVector;
+						MorphTargetDelta.TangentZDelta = bIncludeNormals ? DeltaNormals[ImportedVertexNumber] : FVector3f::ZeroVector;
 					}
 				}
 			}
 
 			MorphLODModel.Vertices.Shrink();
+			UE_LOG(LogMLDeformer, Display, TEXT("Morph #%d: %d deltas"), MorphTargetIndex, MorphLODModel.Vertices.Num());
 		}
 	}
 
@@ -1699,8 +1780,8 @@ namespace UE::MLDeformer
 			const int32 NumVerts = Model->GetNumBaseMeshVerts();
 			check(MorphDeltas.Num() % NumVerts == 0);
 
-			const TArray<FVector3f>& SkinnedPositions = Sampler->GetUnskinnedVertexPositions();
-			check(NumVerts == SkinnedPositions.Num());
+			const TArray<FVector3f>& UnskinnedPositions = Sampler->GetUnskinnedVertexPositions();
+			check(NumVerts == UnskinnedPositions.Num());
 
 			// Draw all deltas.			
 			const int32 NumMorphTargets = MorphDeltas.Num() / NumVerts;
@@ -1709,7 +1790,7 @@ namespace UE::MLDeformer
 			const FLinearColor ExcludedColor(0.1f, 0.1f, 0.1f);
 			for (int32 VertexIndex = 0; VertexIndex < NumVerts; ++VertexIndex)
 			{
-				const FVector StartPoint = FVector(SkinnedPositions[VertexIndex].X, SkinnedPositions[VertexIndex].Y, SkinnedPositions[VertexIndex].Z) + DrawOffset;
+				const FVector StartPoint = FVector(UnskinnedPositions[VertexIndex]) + DrawOffset;
 				const int32 DeltaArrayOffset = NumVerts * FinalMorphTargetIndex + VertexIndex;
 				const FVector Delta(MorphDeltas[DeltaArrayOffset]);
 
diff --git a/Engine/Plugins/Animation/MLDeformer/MLDeformerFramework/Source/MLDeformerFrameworkEditor/Private/MLDeformerMorphModelDetails.cpp b/Engine/Plugins/Animation/MLDeformer/MLDeformerFramework/Source/MLDeformerFrameworkEditor/Private/MLDeformerMorphModelDetails.cpp
--- a/Engine/Plugins/Animation/MLDeformer/MLDeformerFramework/Source/MLDeformerFrameworkEditor/Private/MLDeformerMorphModelDetails.cpp
+++ b/Engine/Plugins/Animation/MLDeformer/MLDeformerFramework/Source/MLDeformerFrameworkEditor/Private/MLDeformerMorphModelDetails.cpp
@@ -37,6 +37,7 @@ namespace UE::MLDeformer
 
 		MorphTargetCategoryBuilder->AddProperty(UMLDeformerMorphModel::GetMorphTargetDeltaThresholdPropertyName(), UMLDeformerMorphModel::StaticClass());
 		MorphTargetCategoryBuilder->AddProperty(UMLDeformerMorphModel::GetMorphTargetErrorTolerancePropertyName(), UMLDeformerMorphModel::StaticClass());
+		MorphTargetCategoryBuilder->AddProperty(UMLDeformerMorphModel::GetIncludeMorphTargetNormalsPropertyName(), UMLDeformerMorphModel::StaticClass());
 	}
 }	// namespace UE::MLDeformer
 
diff --git a/Engine/Plugins/Animation/MLDeformer/MLDeformerFramework/Source/MLDeformerFrameworkEditor/Private/MLDeformerMorphModelEditorModel.cpp b/Engine/Plugins/Animation/MLDeformer/MLDeformerFramework/Source/MLDeformerFrameworkEditor/Private/MLDeformerMorphModelEditorModel.cpp
--- a/Engine/Plugins/Animation/MLDeformer/MLDeformerFramework/Source/MLDeformerFrameworkEditor/Private/MLDeformerMorphModelEditorModel.cpp
+++ b/Engine/Plugins/Animation/MLDeformer/MLDeformerFramework/Source/MLDeformerFrameworkEditor/Private/MLDeformerMorphModelEditorModel.cpp
@@ -106,7 +106,8 @@ namespace UE::MLDeformer
 		// Turn the delta buffer in a set of engine morph targets.
 		UMLDeformerMorphModel* MorphModel = GetMorphModel();
 		TArray<UMorphTarget*> MorphTargets;	// These will be garbage collected.
-		CreateEngineMorphTargets(MorphTargets, Deltas, TEXT("MLDeformerMorph_"), LOD, MorphModel->GetMorphTargetDeltaThreshold());
+		const bool bIncludeNormals = MorphModel->GetIncludeMorphTargetNormals();
+		CreateEngineMorphTargets(MorphTargets, Deltas, TEXT("MLDeformerMorph_"), LOD, MorphModel->GetMorphTargetDeltaThreshold(), bIncludeNormals);
 
 		// Now compress the morph targets to GPU friendly buffers.
 		check(MorphModel->GetMorphTargetSet().IsValid());
diff --git a/Engine/Plugins/Animation/MLDeformer/MLDeformerFramework/Source/MLDeformerFrameworkEditor/Public/MLDeformerEditorModel.h b/Engine/Plugins/Animation/MLDeformer/MLDeformerFramework/Source/MLDeformerFrameworkEditor/Public/MLDeformerEditorModel.h
--- a/Engine/Plugins/Animation/MLDeformer/MLDeformerFramework/Source/MLDeformerFrameworkEditor/Public/MLDeformerEditorModel.h
+++ b/Engine/Plugins/Animation/MLDeformer/MLDeformerFramework/Source/MLDeformerFrameworkEditor/Public/MLDeformerEditorModel.h
@@ -551,6 +551,17 @@ namespace UE::MLDeformer
 		 */
 		virtual FString GetHeatMapDeformerGraphPath() const;
 
+		/**
+		 * Generate the normals for a given morph target.
+		 * @param LOD The LOD level.
+		 * @param SkelMesh The skeletal mesh to get the mesh data from.
+		 * @param Deltas The per vertex deltas. The number of elements in this array is NumMorphTargets * NumImportedVertices.
+		 * @param BaseVertexPositions The positions of the base/neutral mesh, basically the unskinned vertices.
+		 * @param BaseNormals The normals of the base mesh.
+		 * @param OutDeltaNormals The array that we will write the generated normals to. This will automatically be resized by this method.
+		 */
+		virtual void GenerateNormalsForMorphTarget(int32 LOD, USkeletalMesh* SkelMesh, int32 MorphTargetIndex, TArrayView<const FVector3f> Deltas, TArrayView<const FVector3f> BaseVertexPositions, TArrayView<FVector3f> BaseNormals, TArray<FVector3f>& OutDeltaNormals);
+
 		/** Get the current view range. */
 		TRange<double> GetViewRange() const;
 
@@ -783,8 +794,10 @@ namespace UE::MLDeformer
 		 * @param NamePrefix The morph target name prefix. If set to "MorphTarget_" the names will be "MorphTarget_000", "MorphTarget_001", "MorphTarget_002", etc.
 		 * @param LOD The LOD index to generate the morphs for.
 		 * @param DeltaThreshold Only include deltas with a length larger than this threshold in the morph targets.
+		 * @param bIncludeNormals Include normals inside the morph targets? This can be an alternative to recalculating normals at the end, although setting this to true and not 
+		 *        recomputing normals can lead to lower quality results, in trade for faster performance.
 		 */
-		void CreateEngineMorphTargets(TArray<UMorphTarget*>& OutMorphTargets, const TArray<FVector3f>& Deltas, const FString& NamePrefix = TEXT("MorphTarget_"), int32 LOD = 0, float DeltaThreshold = 0.01f);
+		void CreateEngineMorphTargets(TArray<UMorphTarget*>& OutMorphTargets, const TArray<FVector3f>& Deltas, const FString& NamePrefix = TEXT("MorphTarget_"), int32 LOD = 0, float DeltaThreshold = 0.01f, bool bIncludeNormals=false);
 
 		/** 
 		 * Compress morph targets into GPU based morph buffers.
@@ -821,6 +834,16 @@ namespace UE::MLDeformer
 		 */
 		void UpdateRanges();
 
+		/**
+		 * Calculate the normals for each vertex, given the triangle data and positions.
+		 * It computes this by summing up the face normals for each vertex using that face, and normalizing them at the end.
+		 * @param VertexPositions The buffer with vertex positions. This is the size of the number of imported vertices.
+		 * @param IndexArray The index buffer, which contains NumTriangles * 3 number of integers.
+		 * @param VertexMap For each render vertex, an imported vertex number. For example, for a cube these indices go from 0..7.
+		 * @param OutNormals The array that will contain the normals. This will automatically be resized internally by this method.
+		 */
+		void CalcMeshNormals(TArrayView<const FVector3f> VertexPositions, TArrayView<const uint32> IndexArray, TArrayView<const int32> VertexMap, TArray<FVector3f>& OutNormals) const;
+
 	protected:
 		/** The runtime model associated with this editor model. */
 		TObjectPtr<UMLDeformerModel> Model = nullptr;
