Component, world, or bone space simulation for JiggleSpringModel.

diff --git a/Engine/Plugins/Animation/MLDeformer/MLDeformerFramework/Source/MLDeformerFrameworkEditor/Private/MLDeformerEditorToolkit.cpp b/Engine/Plugins/Animation/MLDeformer/MLDeformerFramework/Source/MLDeformerFrameworkEditor/Private/MLDeformerEditorToolkit.cpp
--- a/Engine/Plugins/Animation/MLDeformer/MLDeformerFramework/Source/MLDeformerFrameworkEditor/Private/MLDeformerEditorToolkit.cpp
+++ b/Engine/Plugins/Animation/MLDeformer/MLDeformerFramework/Source/MLDeformerFrameworkEditor/Private/MLDeformerEditorToolkit.cpp
@@ -381,7 +381,10 @@ namespace UE::MLDeformer
 						ActiveModel->OnPreTraining();
 
 						// Change the interpolation type for the training sequence to step.
-						Model->GetAnimSequence()->Interpolation = EAnimInterpolationType::Step;
+						if (UAnimSequence* AnimSequence = Model->GetAnimSequence())
+						{
+							AnimSequence->Interpolation = EAnimInterpolationType::Step;
+						}
 
 						// Initialize the training inputs.
 						ActiveModel->InitInputInfo(Model->GetInputInfo());
