[MLDeformer]
- Add support for masking the ML Deformer using color channels.
- Fix some crash related to opening an MLD asset, closing the asset editor, and modifying or reimporting the mesh.

diff --git a/Engine/Plugins/Animation/MLDeformer/MLDeformerFramework/Source/MLDeformerFramework/Public/MLDeformerModel.h b/Engine/Plugins/Animation/MLDeformer/MLDeformerFramework/Source/MLDeformerFramework/Public/MLDeformerModel.h
--- a/Engine/Plugins/Animation/MLDeformer/MLDeformerFramework/Source/MLDeformerFramework/Public/MLDeformerModel.h
+++ b/Engine/Plugins/Animation/MLDeformer/MLDeformerFramework/Source/MLDeformerFramework/Public/MLDeformerModel.h
@@ -21,6 +21,26 @@ class UMLDeformerModelInstance;
 class UMLDeformerComponent;
 class UMLDeformerInputInfo;
 
+UENUM()
+enum class EMLDeformerMaskChannel : uint8
+{
+	/** Disable the weight mask. */
+	Disabled,
+
+	/** The red vertex color channel. */
+	VertexColorRed,
+
+	/** The green vertex color channel. */
+	VertexColorGreen,
+
+	/** The blue vertex color channel. */
+	VertexColorBlue,
+
+	/** The alpha vertex color channel. */
+	VertexColorAlpha
+};
+
+
 namespace UE::MLDeformer
 {
 	/**
diff --git a/Engine/Plugins/Animation/MLDeformer/MLDeformerFramework/Source/MLDeformerFramework/Public/MLDeformerMorphModel.h b/Engine/Plugins/Animation/MLDeformer/MLDeformerFramework/Source/MLDeformerFramework/Public/MLDeformerMorphModel.h
--- a/Engine/Plugins/Animation/MLDeformer/MLDeformerFramework/Source/MLDeformerFramework/Public/MLDeformerMorphModel.h
+++ b/Engine/Plugins/Animation/MLDeformer/MLDeformerFramework/Source/MLDeformerFramework/Public/MLDeformerMorphModel.h
@@ -29,20 +29,24 @@ public:
 	void BeginDestroy() override;
 	// ~END UObject overrides.
 
-#if WITH_EDITORONLY_DATA
 	float GetMorphTargetDeltaThreshold() const				{ return MorphTargetDeltaThreshold; }
 	float GetMorphTargetErrorTolerance() const				{ return MorphTargetErrorTolerance; }
 	bool GetIncludeMorphTargetNormals() const				{ return bIncludeNormals; }
+	EMLDeformerMaskChannel GetMaskChannel() const			{ return MaskChannel; }
+	bool GetInvertMaskChannel() const						{ return bInvertMaskChannel; }
 
 	void SetMorphTargetDeltaThreshold(float Threshold)		{ MorphTargetDeltaThreshold = Threshold; }
 	void SetMorphTargetErrorTolerance(float Tolerance)		{ MorphTargetErrorTolerance = Tolerance; }
 	void SetIncludeMorphTargetNormals(bool bInclude)		{ bIncludeNormals = bInclude; }
+	void SetWeightMask(EMLDeformerMaskChannel Channel)		{ MaskChannel = Channel; }
+	void SetInvertMaskChannel(bool bInvert)					{ bInvertMaskChannel = bInvert; }
 
 	// Get property names.
 	static FName GetMorphTargetDeltaThresholdPropertyName() { return GET_MEMBER_NAME_CHECKED(UMLDeformerMorphModel, MorphTargetDeltaThreshold); }
 	static FName GetMorphTargetErrorTolerancePropertyName() { return GET_MEMBER_NAME_CHECKED(UMLDeformerMorphModel, MorphTargetErrorTolerance); }
 	static FName GetIncludeMorphTargetNormalsPropertyName() { return GET_MEMBER_NAME_CHECKED(UMLDeformerMorphModel, bIncludeNormals); }
-#endif
+	static FName GetMaskChannelPropertyName()				{ return GET_MEMBER_NAME_CHECKED(UMLDeformerMorphModel, MaskChannel); }
+	static FName GetInvertMaskChannelPropertyName()			{ return GET_MEMBER_NAME_CHECKED(UMLDeformerMorphModel, bInvertMaskChannel); }
 
 	/**
 	 * Set the per vertex deltas, as a set of floats. Each vertex delta must have 3 floats.
@@ -97,7 +101,6 @@ private:
 	 */
 	TArray<FVector3f> MorphTargetDeltas;
 
-#if WITH_EDITORONLY_DATA
 	/**
 	 * Include vertex normals in the morph targets?
 	 * The advantage of this can be that it is higher performance than recomputing the normals.
@@ -117,5 +120,19 @@ private:
 	/** The morph target error tolerance. Higher values result in larger compression, but could result in visual artifacts. */
 	UPROPERTY(EditAnywhere, Category = "Morph Targets", meta = (ClampMin = "0.01", ClampMax = "500"))
 	float MorphTargetErrorTolerance = 20.0f;
-#endif // WITH_EDITORONLY_DATA
+
+	/**
+	 * The channel data that represents the delta mask multipliers.
+	 * You can use this feather out influence of the ML Deformer in specific areas, such as neck line seams, where the head mesh connects with the body.
+	 * The painted vertex color values will be like a weight multiplier on the ML deformer deltas applied to that vertex. You can invert the mask as well.
+	 */
+	UPROPERTY(EditAnywhere, Category = "Morph Targets")
+	EMLDeformerMaskChannel MaskChannel = EMLDeformerMaskChannel::Disabled;
+
+	/** 
+	 * Enable this if you want to invert the mask channel values. For example if you painted the neck seam vertices in red, and you wish the vertices that got painted to NOT move, you have to invert the mask.
+	 * On default you paint areas where the deformer should be active. If you enable the invert option, you paint areas where the deformer will not be active.
+	 */
+	UPROPERTY(EditAnywhere, Category = "Morph Targets", meta = (EditCondition = "MaskChannel != EMLDeformerMaskChannel::Disabled"))
+	bool bInvertMaskChannel = false;
 };
diff --git a/Engine/Plugins/Animation/MLDeformer/MLDeformerFramework/Source/MLDeformerFrameworkEditor/MLDeformerFrameworkEditor.Build.cs b/Engine/Plugins/Animation/MLDeformer/MLDeformerFramework/Source/MLDeformerFrameworkEditor/MLDeformerFrameworkEditor.Build.cs
--- a/Engine/Plugins/Animation/MLDeformer/MLDeformerFramework/Source/MLDeformerFrameworkEditor/MLDeformerFrameworkEditor.Build.cs
+++ b/Engine/Plugins/Animation/MLDeformer/MLDeformerFramework/Source/MLDeformerFrameworkEditor/MLDeformerFrameworkEditor.Build.cs
@@ -8,6 +8,13 @@ namespace UnrealBuildTool.Rules
 	{
 		public MLDeformerFrameworkEditor(ReadOnlyTargetRules Target) : base(Target)
 		{
+			PrivateIncludePaths.AddRange(
+				new string[] 
+				{
+					System.IO.Path.Combine(GetModuleDirectory("Persona"), "Private"),
+				}
+			);
+
 			PublicDependencyModuleNames.AddRange(
 				new string[]
 				{
diff --git a/Engine/Plugins/Animation/MLDeformer/MLDeformerFramework/Source/MLDeformerFrameworkEditor/Private/MLDeformerEditorModel.cpp b/Engine/Plugins/Animation/MLDeformer/MLDeformerFramework/Source/MLDeformerFrameworkEditor/Private/MLDeformerEditorModel.cpp
--- a/Engine/Plugins/Animation/MLDeformer/MLDeformerFramework/Source/MLDeformerFrameworkEditor/Private/MLDeformerEditorModel.cpp
+++ b/Engine/Plugins/Animation/MLDeformer/MLDeformerFramework/Source/MLDeformerFrameworkEditor/Private/MLDeformerEditorModel.cpp
@@ -15,6 +15,7 @@
 #include "MLDeformerSampler.h"
 #include "MLDeformerModelInstance.h"
 #include "AnimationEditorPreviewActor.h"
+#include "AnimationEditorViewportClient.h"
 #include "AnimPreviewInstance.h"
 #include "Animation/MeshDeformer.h"
 #include "Animation/DebugSkelMeshComponent.h"
@@ -236,11 +237,10 @@ namespace UE::MLDeformer
 
 	void FMLDeformerEditorModel::ClearWorld()
 	{
-		TSharedRef<IPersonaPreviewScene> PreviewScene = Editor->GetPersonaToolkit()->GetPreviewScene();
-
 		// First remove actors from the world and add them to the destroy list.
 		TArray<AActor*> ActorsToDestroy;
 		ActorsToDestroy.Reserve(EditorActors.Num());
+		TSharedRef<IPersonaPreviewScene> PreviewScene = Editor->GetPersonaToolkit()->GetPreviewScene();
 		UWorld* World = PreviewScene->GetWorld();
 		for (FMLDeformerEditorActor* EditorActor : EditorActors)
 		{
@@ -254,15 +254,27 @@ namespace UE::MLDeformer
 		// Now delete the editor actors, which removes components.
 		DeleteEditorActors();
 
+		for (AActor* Actor : ActorsToDestroy)
+		{
+			Actor->Destroy();
+		}
+	}
+
+	void FMLDeformerEditorModel::ClearPersonaPreviewScene()
+	{
+		TSharedRef<IPersonaPreviewScene> PreviewScene = Editor->GetPersonaToolkit()->GetPreviewScene();
+		UWorld* World = PreviewScene->GetWorld();
+
 		PreviewScene->DeselectAll();
 		PreviewScene->SetPreviewAnimationAsset(nullptr);
 		PreviewScene->SetPreviewMeshComponent(nullptr);
 		PreviewScene->SetActor(nullptr);
+	}
 
-		for (AActor* Actor : ActorsToDestroy)
-		{
-			Actor->Destroy();
-		}
+	void FMLDeformerEditorModel::ClearWorldAndPersonaPreviewScene()
+	{
+		ClearWorld();
+		ClearPersonaPreviewScene();
 	}
 
 	FMLDeformerEditorActor* FMLDeformerEditorModel::CreateEditorActor(const FMLDeformerEditorActor::FConstructSettings& Settings) const
@@ -1659,7 +1671,7 @@ namespace UE::MLDeformer
 		}
 	}
 
-	void FMLDeformerEditorModel::CreateEngineMorphTargets(TArray<UMorphTarget*>& OutMorphTargets, const TArray<FVector3f>& Deltas, const FString& NamePrefix, int32 LOD, float DeltaThreshold, bool bIncludeNormals)
+	void FMLDeformerEditorModel::CreateEngineMorphTargets(TArray<UMorphTarget*>& OutMorphTargets, const TArray<FVector3f>& Deltas, const FString& NamePrefix, int32 LOD, float DeltaThreshold, bool bIncludeNormals, EMLDeformerMaskChannel MaskChannel, bool bInvertMaskChannel)
 	{
 		OutMorphTargets.Reset();
 		if (Deltas.IsEmpty())
@@ -1687,6 +1699,8 @@ namespace UE::MLDeformer
 		TArray<FVector3f> BaseNormals;
 		CalcMeshNormals(BaseVertexPositions, IndexArray, VertexMap, BaseNormals);
 
+		const FColorVertexBuffer& ColorBuffer = RenderData->LODRenderData[LOD].StaticVertexBuffers.ColorVertexBuffer;
+
 		// Initialize an engine morph target for each model morph target.
 		UE_LOG(LogMLDeformer, Display, TEXT("Initializing %d engine morph targets of %d vertices each"), NumMorphTargets, Deltas.Num() / NumMorphTargets);
 		TArray<FVector3f> DeltaNormals;
@@ -1725,20 +1739,40 @@ namespace UE::MLDeformer
 				const int32 ImportedVertexNumber = VertexMap[VertexIndex];
 				if (ImportedVertexNumber != INDEX_NONE)
 				{
+					float VertexWeight = 1.0f;
+					if (ColorBuffer.GetNumVertices() != 0 && MaskChannel != EMLDeformerMaskChannel::Disabled)
+					{
+						const FLinearColor& VertexColor = ColorBuffer.VertexColor(VertexIndex);
+						switch (MaskChannel)
+						{
+							case EMLDeformerMaskChannel::VertexColorRed:	{ VertexWeight = VertexColor.R; break; }
+							case EMLDeformerMaskChannel::VertexColorGreen:	{ VertexWeight = VertexColor.G; break; }
+							case EMLDeformerMaskChannel::VertexColorBlue:	{ VertexWeight = VertexColor.B; break; }
+							case EMLDeformerMaskChannel::VertexColorAlpha:	{ VertexWeight = VertexColor.A; break; }
+							default: 
+								checkf(false, TEXT("Unexpected weight mask value."));
+								break;
+						};
+
+						if (bInvertMaskChannel)
+						{
+							VertexWeight = FMath::Clamp<float>(1.0f - VertexWeight, 0.0f, 1.0f);
+						}
+					}
+
 					const FVector3f Delta = Deltas[ImportedVertexNumber + MorphTargetIndex * NumBaseMeshVerts];
-					if (Delta.Length() > DeltaThreshold)
+					if (Delta.Length() * VertexWeight > DeltaThreshold)
 					{
 						MorphLODModel.Vertices.AddDefaulted();
 						FMorphTargetDelta& MorphTargetDelta = MorphLODModel.Vertices.Last();
-						MorphTargetDelta.PositionDelta = Delta;
+						MorphTargetDelta.PositionDelta = Delta * VertexWeight;
 						MorphTargetDelta.SourceIdx = VertexIndex;
-						MorphTargetDelta.TangentZDelta = bIncludeNormals ? DeltaNormals[ImportedVertexNumber] : FVector3f::ZeroVector;
+						MorphTargetDelta.TangentZDelta = bIncludeNormals ? DeltaNormals[ImportedVertexNumber] * VertexWeight : FVector3f::ZeroVector;
 					}
 				}
 			}
 
 			MorphLODModel.Vertices.Shrink();
-			UE_LOG(LogMLDeformer, Display, TEXT("Morph #%d: %d deltas"), MorphTargetIndex, MorphLODModel.Vertices.Num());
 		}
 	}
 
diff --git a/Engine/Plugins/Animation/MLDeformer/MLDeformerFramework/Source/MLDeformerFrameworkEditor/Private/MLDeformerEditorToolkit.cpp b/Engine/Plugins/Animation/MLDeformer/MLDeformerFramework/Source/MLDeformerFrameworkEditor/Private/MLDeformerEditorToolkit.cpp
--- a/Engine/Plugins/Animation/MLDeformer/MLDeformerFramework/Source/MLDeformerFrameworkEditor/Private/MLDeformerEditorToolkit.cpp
+++ b/Engine/Plugins/Animation/MLDeformer/MLDeformerFramework/Source/MLDeformerFrameworkEditor/Private/MLDeformerEditorToolkit.cpp
@@ -13,6 +13,7 @@
 #include "MLDeformerEditorStyle.h"
 #include "MLDeformerVizSettings.h"
 #include "MLDeformerSampler.h"
+#include "AnimationEditorViewportClient.h"
 #include "EditorModeManager.h"
 #include "EditorViewportClient.h"
 #include "Modules/ModuleManager.h"
@@ -86,7 +87,7 @@ namespace UE::MLDeformer
 			DeformerAsset);
 
 		// Create and set the application mode.
-		FMLDeformerApplicationMode* ApplicationMode = new FMLDeformerApplicationMode(SharedThis(this), PersonaToolkit->GetPreviewScene());
+		ApplicationMode = new FMLDeformerApplicationMode(SharedThis(this), PersonaToolkit->GetPreviewScene());
 		AddApplicationMode(MLDeformerEditorModes::Editor, MakeShareable(ApplicationMode));
 		SetCurrentMode(MLDeformerEditorModes::Editor);
 
@@ -261,7 +262,7 @@ namespace UE::MLDeformer
 		// Remove existing actors etc.
 		if (ActiveModel)
 		{
-			ActiveModel->ClearWorld();
+			ActiveModel->ClearWorldAndPersonaPreviewScene();
 		}
 
 		// Get the runtime model type based on the index, and create an instance of it.
@@ -606,7 +607,7 @@ namespace UE::MLDeformer
 	{
 		FFormatNamedArguments Args;
 		Args.Add(TEXT("AssetName"), FText::FromString(DeformerAsset->GetName()));
-		return FText::Format(LOCTEXT("DemoEditorToolkitName", "{AssetName}"), Args);
+		return FText::Format(LOCTEXT("MLDeformerEditorToolkitName", "{AssetName}"), Args);
 	}
 
 	FLinearColor FMLDeformerEditorToolkit::GetWorldCentricTabColorScale() const
@@ -828,7 +829,6 @@ namespace UE::MLDeformer
 			ActiveModel->SetViewRange(ViewRange);
 
 		}
-
 	}
 }	// namespace UE::MLDeformer
 
diff --git a/Engine/Plugins/Animation/MLDeformer/MLDeformerFramework/Source/MLDeformerFrameworkEditor/Private/MLDeformerMorphModelDetails.cpp b/Engine/Plugins/Animation/MLDeformer/MLDeformerFramework/Source/MLDeformerFrameworkEditor/Private/MLDeformerMorphModelDetails.cpp
--- a/Engine/Plugins/Animation/MLDeformer/MLDeformerFramework/Source/MLDeformerFrameworkEditor/Private/MLDeformerMorphModelDetails.cpp
+++ b/Engine/Plugins/Animation/MLDeformer/MLDeformerFramework/Source/MLDeformerFrameworkEditor/Private/MLDeformerMorphModelDetails.cpp
@@ -37,6 +37,8 @@ namespace UE::MLDeformer
 
 		MorphTargetCategoryBuilder->AddProperty(UMLDeformerMorphModel::GetMorphTargetDeltaThresholdPropertyName(), UMLDeformerMorphModel::StaticClass());
 		MorphTargetCategoryBuilder->AddProperty(UMLDeformerMorphModel::GetMorphTargetErrorTolerancePropertyName(), UMLDeformerMorphModel::StaticClass());
+		MorphTargetCategoryBuilder->AddProperty(UMLDeformerMorphModel::GetMaskChannelPropertyName(), UMLDeformerMorphModel::StaticClass());
+		MorphTargetCategoryBuilder->AddProperty(UMLDeformerMorphModel::GetInvertMaskChannelPropertyName(), UMLDeformerMorphModel::StaticClass());
 		MorphTargetCategoryBuilder->AddProperty(UMLDeformerMorphModel::GetIncludeMorphTargetNormalsPropertyName(), UMLDeformerMorphModel::StaticClass());
 	}
 }	// namespace UE::MLDeformer
diff --git a/Engine/Plugins/Animation/MLDeformer/MLDeformerFramework/Source/MLDeformerFrameworkEditor/Private/MLDeformerMorphModelEditorModel.cpp b/Engine/Plugins/Animation/MLDeformer/MLDeformerFramework/Source/MLDeformerFrameworkEditor/Private/MLDeformerMorphModelEditorModel.cpp
--- a/Engine/Plugins/Animation/MLDeformer/MLDeformerFramework/Source/MLDeformerFrameworkEditor/Private/MLDeformerMorphModelEditorModel.cpp
+++ b/Engine/Plugins/Animation/MLDeformer/MLDeformerFramework/Source/MLDeformerFrameworkEditor/Private/MLDeformerMorphModelEditorModel.cpp
@@ -14,6 +14,7 @@
 #include "Animation/MorphTarget.h"
 #include "Rendering/SkeletalMeshModel.h"
 #include "Rendering/SkeletalMeshLODModel.h"
+#include "Rendering/SkeletalMeshRenderData.h"
 #include "Engine/SkeletalMesh.h"
 
 #define LOCTEXT_NAMESPACE "MLDeformerMorphModelEditorModel"
@@ -103,11 +104,15 @@ namespace UE::MLDeformer
 	{
 		const int32 LOD = 0;
 
-		// Turn the delta buffer in a set of engine morph targets.
+		// Get the weight mask setting.
 		UMLDeformerMorphModel* MorphModel = GetMorphModel();
+		const EMLDeformerMaskChannel MaskChannel = MorphModel->GetMaskChannel();
+		const bool bInvertMaskChannel = MorphModel->GetInvertMaskChannel();
+
+		// Turn the delta buffer in a set of engine morph targets.
 		TArray<UMorphTarget*> MorphTargets;	// These will be garbage collected.
 		const bool bIncludeNormals = MorphModel->GetIncludeMorphTargetNormals();
-		CreateEngineMorphTargets(MorphTargets, Deltas, TEXT("MLDeformerMorph_"), LOD, MorphModel->GetMorphTargetDeltaThreshold(), bIncludeNormals);
+		CreateEngineMorphTargets(MorphTargets, Deltas, TEXT("MLDeformerMorph_"), LOD, MorphModel->GetMorphTargetDeltaThreshold(), bIncludeNormals, MaskChannel, bInvertMaskChannel);
 
 		// Now compress the morph targets to GPU friendly buffers.
 		check(MorphModel->GetMorphTargetSet().IsValid());
diff --git a/Engine/Plugins/Animation/MLDeformer/MLDeformerFramework/Source/MLDeformerFrameworkEditor/Public/MLDeformerEditorModel.h b/Engine/Plugins/Animation/MLDeformer/MLDeformerFramework/Source/MLDeformerFrameworkEditor/Public/MLDeformerEditorModel.h
--- a/Engine/Plugins/Animation/MLDeformer/MLDeformerFramework/Source/MLDeformerFrameworkEditor/Public/MLDeformerEditorModel.h
+++ b/Engine/Plugins/Animation/MLDeformer/MLDeformerFramework/Source/MLDeformerFrameworkEditor/Public/MLDeformerEditorModel.h
@@ -11,6 +11,7 @@
 #include "Misc/FrameTime.h"
 #include "MLDeformerVizSettings.h"
 #include "MLDeformerModule.h"
+#include "MLDeformerModel.h"
 
 class UMLDeformerModel;
 class UMLDeformerInputInfo;
@@ -151,10 +152,20 @@ namespace UE::MLDeformer
 
 		/**
 		 * Clears the world, which basically means it removes all the editor actors and engine actors from the editor world.
-		 * It also resets the Persona preview scene.
 		 */
 		virtual void ClearWorld();
 
+		/**
+		 * Clears the Persona preview scene. Resetting the preview mesh, preview mesh component etc.
+		 */
+		virtual void ClearPersonaPreviewScene();
+
+		/**
+		 * Clear both the world and preview scene.
+		 * This calls ClearWorld and ClearPersonaPreviewScene.
+		 */
+		virtual void ClearWorldAndPersonaPreviewScene();
+
 		/**
 		 * Create an actor in the editor viewport.
 		 * Each actor will have a specific ID and a label that is displayed above it, with a specific color.
@@ -796,8 +807,18 @@ namespace UE::MLDeformer
 		 * @param DeltaThreshold Only include deltas with a length larger than this threshold in the morph targets.
 		 * @param bIncludeNormals Include normals inside the morph targets? This can be an alternative to recalculating normals at the end, although setting this to true and not 
 		 *        recomputing normals can lead to lower quality results, in trade for faster performance.
-		 */
-		void CreateEngineMorphTargets(TArray<UMorphTarget*>& OutMorphTargets, const TArray<FVector3f>& Deltas, const FString& NamePrefix = TEXT("MorphTarget_"), int32 LOD = 0, float DeltaThreshold = 0.01f, bool bIncludeNormals=false);
+		 * @param MaskChannel The weight mask mode, which specifies what channel to get the weight data from. Such channel allows the user to define what areas the deformer should for example not be active in.
+		 * @param bInvertMaskChannel Specifies whether the weight mask should be inverted or not.
+		 */
+		void CreateEngineMorphTargets(
+			TArray<UMorphTarget*>& OutMorphTargets,
+			const TArray<FVector3f>& Deltas,
+			const FString& NamePrefix = TEXT("MorphTarget_"),
+			int32 LOD = 0,
+			float DeltaThreshold = 0.01f,
+			bool bIncludeNormals=false,
+			const EMLDeformerMaskChannel MaskChannel = EMLDeformerMaskChannel::Disabled,
+			bool bInvertMaskChannel = false);
 
 		/** 
 		 * Compress morph targets into GPU based morph buffers.
diff --git a/Engine/Plugins/Animation/MLDeformer/MLDeformerFramework/Source/MLDeformerFrameworkEditor/Public/MLDeformerEditorToolkit.h b/Engine/Plugins/Animation/MLDeformer/MLDeformerFramework/Source/MLDeformerFrameworkEditor/Public/MLDeformerEditorToolkit.h
--- a/Engine/Plugins/Animation/MLDeformer/MLDeformerFramework/Source/MLDeformerFrameworkEditor/Public/MLDeformerEditorToolkit.h
+++ b/Engine/Plugins/Animation/MLDeformer/MLDeformerFramework/Source/MLDeformerFrameworkEditor/Public/MLDeformerEditorToolkit.h
@@ -21,6 +21,7 @@ class SSimpleTimeSlider;
 namespace UE::MLDeformer
 {
 	class SMLDeformerTimeline;
+	class FMLDeformerApplicationMode;
 
 	namespace MLDeformerEditorModes
 	{
@@ -88,11 +89,15 @@ namespace UE::MLDeformer
 		TWeakPtr<FMLDeformerEditorModel> GetActiveModelPointer(){ return TWeakPtr<FMLDeformerEditorModel>(ActiveModel); }
 		const FMLDeformerEditorModel* GetActiveModel() const { return ActiveModel.Get(); }
 
+		FMLDeformerApplicationMode* GetApplicationMode() const { return ApplicationMode; }
+
 		double CalcTimelinePosition() const;
 		void OnTimeSliderScrubPositionChanged(double NewScrubTime, bool bIsScrubbing);
 		void UpdateTimeSliderRange();
 		void SetTimeSliderRange(double StartTime, double EndTime);
 
+		TSharedPtr<IPersonaViewport> GetViewport() const { return PersonaViewport; }
+
 	private:
 		/* Toolbar related. */
 		void ExtendToolbar();
@@ -142,6 +147,8 @@ namespace UE::MLDeformer
 		/** The ML Deformer Asset. */
 		TObjectPtr<UMLDeformerAsset> DeformerAsset = nullptr;
 
+		FMLDeformerApplicationMode* ApplicationMode = nullptr;
+
 		/** Has the asset editor been initialized? */
 		bool bIsInitialized = false;
 	};
diff --git a/Engine/Plugins/Animation/MLDeformer/MLDeformerFramework/Source/MLDeformerFrameworkEditor/Public/MLDeformerMorphModelEditorModel.h b/Engine/Plugins/Animation/MLDeformer/MLDeformerFramework/Source/MLDeformerFrameworkEditor/Public/MLDeformerMorphModelEditorModel.h
--- a/Engine/Plugins/Animation/MLDeformer/MLDeformerFramework/Source/MLDeformerFrameworkEditor/Public/MLDeformerMorphModelEditorModel.h
+++ b/Engine/Plugins/Animation/MLDeformer/MLDeformerFramework/Source/MLDeformerFrameworkEditor/Public/MLDeformerMorphModelEditorModel.h
@@ -46,7 +46,7 @@ namespace UE::MLDeformer
 		 * the first morph target, etc. In other words, the layout is: [morph0_deltas][morph1_deltas][morph2_deltas][...].
 		 * @param Deltas The deltas for all morph targets concatenated. So the number of items in this array is a multiple of Model->GetNumBaseVerts().
 		 */
-		void InitEngineMorphTargets(const TArray<FVector3f>& Deltas);
+		virtual void InitEngineMorphTargets(const TArray<FVector3f>& Deltas);
 
 		/**
 		 * Clamp the current morph target number inside the visualization settings to a valid range.
