[UE] MLDeformer Component - SetWeight and SetAsset can now be overriden by derived classes

#rnx

diff --git a/Engine/Plugins/Animation/MLDeformer/MLDeformerFramework/Source/MLDeformerFramework/Public/MLDeformerComponent.h b/Engine/Plugins/Animation/MLDeformer/MLDeformerFramework/Source/MLDeformerFramework/Public/MLDeformerComponent.h
--- a/Engine/Plugins/Animation/MLDeformer/MLDeformerFramework/Source/MLDeformerFramework/Public/MLDeformerComponent.h
+++ b/Engine/Plugins/Animation/MLDeformer/MLDeformerFramework/Source/MLDeformerFramework/Public/MLDeformerComponent.h
@@ -62,7 +62,7 @@ public:
 	 * @param NormalizedWeightValue The weight value.
 	 */
 	UFUNCTION(BlueprintCallable, Category = "MLDeformer")
-	void SetWeight(float NormalizedWeightValue)					{ Weight = FMath::Clamp<float>(NormalizedWeightValue, 0.0f, 1.0f);  }
+	void SetWeight(float NormalizedWeightValue)					{ SetWeightInternal(NormalizedWeightValue); }
 
 	/**
 	 * Get the ML Deformer asset that is used by this component.
@@ -77,7 +77,7 @@ public:
 	 * @param InDeformerAsset A pointer to the deformer asset.
 	 */
 	UFUNCTION(BlueprintCallable, Category = "MLDeformer")
-	void SetDeformerAsset(UMLDeformerAsset* InDeformerAsset)	{ DeformerAsset = InDeformerAsset; UpdateSkeletalMeshComponent(); }
+	void SetDeformerAsset(UMLDeformerAsset* InDeformerAsset) { SetDeformerAssetInternal(InDeformerAsset); }
 
 	/**
 	 * Find the skeletal mesh component to apply the deformer on.
@@ -139,6 +139,12 @@ protected:
 	/** Unbind from the MLDeformerModel's NeuralNetworkModifyDelegate. */
 	void RemoveNeuralNetworkModifyDelegate();
 
+	/** Set the ML Deformer weight. */
+	virtual void SetWeightInternal(const float NormalizedWeightValue) { Weight = FMath::Clamp<float>(NormalizedWeightValue, 0.0f, 1.0f); }
+
+	/** Set the ML Deformer asset. */
+	virtual void SetDeformerAssetInternal(UMLDeformerAsset* const InDeformerAsset) { DeformerAsset = InDeformerAsset; UpdateSkeletalMeshComponent(); }
+
 protected:
 	/** Render command fence that let's us wait for all other commands to finish. */
 	FRenderCommandFence RenderCommandFence;
