[MLDeformer]
- Extract information about the importance level of each morph target, used only in Neural Morph Model for now.
- Introducing deformer quality levels, which allows models to opt into. Morph models will disable morph targets based on this quality level.
- Dynamically update and recompress morph targets without having to retrain.
- Store raw deltas in uasset, remove on cook. This allows recompression even without training first. Also allows visualizing morph targets at any time.
- Refactor the python code, placing shared code into the ML Deformer framework, which plugins can then import and use.
- Added documentation to the python code.
- Skip saving some properties, like the deformer weight and ground truth lerp, to prevent saving a state where you wouldn't see the deformer being active anymore in your next session.
- Added statistics showing performance and memory usage.
- Renamed and grouped some properties to make things more clear.
- Fixed a bug where deleting or editing a curve or bone in the include list would close the UI, which was very annoying.

diff --git a/Engine/Plugins/Animation/MLDeformer/MLDeformerFramework/Config/BaseMLDeformerFramework.ini b/Engine/Plugins/Animation/MLDeformer/MLDeformerFramework/Config/BaseMLDeformerFramework.ini
--- a/Engine/Plugins/Animation/MLDeformer/MLDeformerFramework/Config/BaseMLDeformerFramework.ini
+++ b/Engine/Plugins/Animation/MLDeformer/MLDeformerFramework/Config/BaseMLDeformerFramework.ini
@@ -6,3 +6,10 @@
 ; Redirects from BoneNameStrings to BoneNames and from CurveNameStrings to CurveNames.
 +PropertyRedirects=(OldName="/Script/MLDeformerFramework.MLDeformerInputInfo.BoneNameStrings",NewName="BoneNames")
 +PropertyRedirects=(OldName="/Script/MLDeformerFramework.MLDeformerInputInfo.CurveNameStrings",NewName="CurveNames")
+
+; Redirect some morph target settings.
++PropertyRedirects=(OldName="/Script/MLDeformerFramework.MLDeformerMorphModel.MorphTargetErrorTolerance",NewName="MorphCompressionLevel")
++PropertyRedirects=(OldName="/Script/MLDeformerFramework.MLDeformerMorphModel.MorphTargetDeltaThreshold",NewName="MorphDeltaZeroThreshold")
+
+; Deprecated value inside the morph model's viz settings.
++PropertyRedirects=(OldName="/Script/MLDeformerFramework.MLDeformerMorphModelVizSettings.MorphTargetDeltaThreshold",NewName="MorphTargetDeltaThreshold_DEPRECATED")
diff --git a/Engine/Plugins/Animation/MLDeformer/MLDeformerFramework/Source/MLDeformerFramework/Private/MLDeformerComponent.cpp b/Engine/Plugins/Animation/MLDeformer/MLDeformerFramework/Source/MLDeformerFramework/Private/MLDeformerComponent.cpp
--- a/Engine/Plugins/Animation/MLDeformer/MLDeformerFramework/Source/MLDeformerFramework/Private/MLDeformerComponent.cpp
+++ b/Engine/Plugins/Animation/MLDeformer/MLDeformerFramework/Source/MLDeformerFramework/Private/MLDeformerComponent.cpp
@@ -100,6 +100,10 @@ void UMLDeformerComponent::SetupComponent(UMLDeformerAsset* InDeformerAsset, USk
 			);
 		}
 	}
+
+	#if WITH_EDITOR
+		TickPerfCounter.Reset();
+	#endif
 }
 
 void UMLDeformerComponent::AddNeuralNetworkModifyDelegate()
@@ -199,6 +203,10 @@ void UMLDeformerComponent::Activate(bool bReset)
 
 void UMLDeformerComponent::Deactivate()
 {
+	#if WITH_EDITOR
+		TickPerfCounter.Reset();
+	#endif
+
 	RemoveNeuralNetworkModifyDelegate();
 	if (ModelInstance)
 	{
@@ -210,6 +218,10 @@ void UMLDeformerComponent::Deactivate()
 
 void UMLDeformerComponent::TickComponent(float DeltaTime, enum ELevelTick TickType, FActorComponentTickFunction* ThisTickFunction)
 {
+	#if WITH_EDITOR
+		TickPerfCounter.BeginSample();
+	#endif
+
 	Super::TickComponent(DeltaTime, TickType, ThisTickFunction);
 	if (TickType != ELevelTick::LEVELTICK_PauseTick)
 	{
@@ -219,8 +231,22 @@ void UMLDeformerComponent::TickComponent(float DeltaTime, enum ELevelTick TickTy
 		{
 			TRACE_CPUPROFILER_EVENT_SCOPE(UMLDeformerComponent::TickComponent)
 			ModelInstance->Tick(DeltaTime, Weight);
+
+			#if WITH_EDITOR
+				// Update our memory usage if desired.
+				// This can be pretty slow.
+				UMLDeformerModel* Model = ModelInstance->GetModel();
+				if (Model->IsMemUsageInvalidated())
+				{
+					Model->UpdateMemoryUsage();
+				}
+			#endif
 		}
 	}
+
+	#if WITH_EDITOR
+		TickPerfCounter.EndSample();
+	#endif
 }
 
 #if WITH_EDITOR
diff --git a/Engine/Plugins/Animation/MLDeformer/MLDeformerFramework/Source/MLDeformerFramework/Private/MLDeformerModel.cpp b/Engine/Plugins/Animation/MLDeformer/MLDeformerFramework/Source/MLDeformerFramework/Private/MLDeformerModel.cpp
--- a/Engine/Plugins/Animation/MLDeformer/MLDeformerFramework/Source/MLDeformerFramework/Private/MLDeformerModel.cpp
+++ b/Engine/Plugins/Animation/MLDeformer/MLDeformerFramework/Source/MLDeformerFramework/Private/MLDeformerModel.cpp
@@ -141,6 +141,10 @@ void UMLDeformerModel::PostLoad()
 			Network->SetDeviceType(ENeuralDeviceType::CPU, ENeuralDeviceType::CPU, ENeuralDeviceType::CPU);
 		}
 	}
+
+	#if WITH_EDITOR
+		UpdateMemoryUsage();
+	#endif
 }
 
 void UMLDeformerModel::SetNeuralNetwork(UNeuralNetwork* InNeuralNetwork)
@@ -227,6 +231,59 @@ void UMLDeformerModel::FloatArrayToVector3Array(const TArray<float>& FloatArray,
 	{
 		return SkeletalMesh ? SkeletalMesh->GetNumImportedVertices() : 0;
 	}
+
+	void UMLDeformerModel::InvalidateMemUsage()
+	{
+		bInvalidateMemUsage = true;
+	}
+
+	uint64 UMLDeformerModel::GetMemUsageInBytes(UE::MLDeformer::EMemUsageRequestFlags Flags) const
+	{
+		return (Flags == UE::MLDeformer::EMemUsageRequestFlags::Cooked) ? CookedMemUsageInBytes : MemUsageInBytes;
+	}
+
+	uint64 UMLDeformerModel::GetGPUMemUsageInBytes() const
+	{
+		return GPUMemUsageInBytes;
+	}
+
+	bool UMLDeformerModel::IsMemUsageInvalidated() const
+	{
+		return bInvalidateMemUsage;
+	}
+
+	void UMLDeformerModel::UpdateMemoryUsage()
+	{
+		bInvalidateMemUsage = false;
+		MemUsageInBytes = 0;
+		GPUMemUsageInBytes = 0;
+
+		// Get the size of the model to get an approximate size.
+		UMLDeformerModel* Model = const_cast<UMLDeformerModel*>(this);
+		MemUsageInBytes += static_cast<uint64>(Model->GetResourceSizeBytes(EResourceSizeMode::Type::EstimatedTotal));
+
+		// Init the cooked size. We can subtract bytes from this to simulate a cook.
+		CookedMemUsageInBytes = MemUsageInBytes;
+
+		// Add GPU memory.
+		if (VertexMapBuffer.VertexBufferRHI.IsValid())
+		{
+			GPUMemUsageInBytes += VertexMapBuffer.VertexBufferRHI->GetSize();
+		}
+
+		// Check if the neural network is on the GPU, if so, count the memory to the GPU and remove it from Main memory, as we added it to 
+		// that already when we did the Model->GetResourceSizeBytes.
+		if (NeuralNetwork)
+		{
+			if (NeuralNetwork->GetDeviceType() == ENeuralDeviceType::GPU)
+			{
+				const uint64 NeuralNetSize = static_cast<uint64>(NeuralNetwork->GetResourceSizeBytes(EResourceSizeMode::Type::EstimatedTotal));
+				GPUMemUsageInBytes += NeuralNetSize;
+				MemUsageInBytes -= NeuralNetSize;
+				CookedMemUsageInBytes -= NeuralNetSize;
+			}
+		}
+	}
 #endif	// #if WITH_EDITOR
 
 #if WITH_EDITORONLY_DATA
diff --git a/Engine/Plugins/Animation/MLDeformer/MLDeformerFramework/Source/MLDeformerFramework/Private/MLDeformerModelInstance.cpp b/Engine/Plugins/Animation/MLDeformer/MLDeformerFramework/Source/MLDeformerFramework/Private/MLDeformerModelInstance.cpp
--- a/Engine/Plugins/Animation/MLDeformer/MLDeformerFramework/Source/MLDeformerFramework/Private/MLDeformerModelInstance.cpp
+++ b/Engine/Plugins/Animation/MLDeformer/MLDeformerFramework/Source/MLDeformerFramework/Private/MLDeformerModelInstance.cpp
@@ -5,6 +5,7 @@
 #include "MLDeformerModule.h"
 #include "MLDeformerAsset.h"
 #include "MLDeformerInputInfo.h"
+#include "MLDeformerComponent.h"
 #include "Components/SkeletalMeshComponent.h"
 #include "Engine/SkeletalMesh.h"
 #include "Animation/AnimInstance.h"
@@ -49,6 +50,11 @@ void UMLDeformerModelInstance::SetModel(UMLDeformerModel* InModel)
 	Model = InModel;
 }
 
+UMLDeformerComponent* UMLDeformerModelInstance::GetMLDeformerComponent() const
+{ 
+	return Cast<UMLDeformerComponent>(GetOuter());
+}
+
 int32 UMLDeformerModelInstance::GetNeuralNetworkInferenceHandle() const
 { 
 	return NeuralNetworkInferenceHandle;
diff --git a/Engine/Plugins/Animation/MLDeformer/MLDeformerFramework/Source/MLDeformerFramework/Private/MLDeformerMorphModel.cpp b/Engine/Plugins/Animation/MLDeformer/MLDeformerFramework/Source/MLDeformerFramework/Private/MLDeformerMorphModel.cpp
--- a/Engine/Plugins/Animation/MLDeformer/MLDeformerFramework/Source/MLDeformerFramework/Private/MLDeformerMorphModel.cpp
+++ b/Engine/Plugins/Animation/MLDeformer/MLDeformerFramework/Source/MLDeformerFramework/Private/MLDeformerMorphModel.cpp
@@ -29,6 +29,11 @@ void UMLDeformerMorphModel::Serialize(FArchive& Archive)
 	bool bHasMorphData = false;
 	if (Archive.IsSaving())
 	{
+		if (Archive.IsCooking())
+		{
+			MorphTargetDeltas.Empty();
+		}
+
 		bHasMorphData = MorphTargetSet.IsValid() ? MorphTargetSet->MorphBuffers.IsMorphCPUDataValid() : false;
 	}
 	Archive << bHasMorphData;
@@ -41,6 +46,16 @@ void UMLDeformerMorphModel::Serialize(FArchive& Archive)
 	}
 }
 
+void UMLDeformerMorphModel::PostLoad()
+{
+	Super::PostLoad();
+	UpdateStatistics();
+
+#if WITH_EDITOR
+	InvalidateMemUsage();
+#endif
+}
+
 UMLDeformerModelInstance* UMLDeformerMorphModel::CreateModelInstance(UMLDeformerComponent* Component)
 {
 	return NewObject<UMLDeformerMorphModelInstance>(Component);
@@ -79,4 +94,82 @@ void UMLDeformerMorphModel::BeginDestroy()
 	Super::BeginDestroy();
 }
 
+void UMLDeformerMorphModel::SetMorphTargetsErrorOrder(const TArray<int32>& MorphTargetOrder, const TArray<float>& ErrorValues)
+{
+	MorphTargetErrorOrder = MorphTargetOrder;
+	MorphTargetErrors = ErrorValues;
+}
+
+void UMLDeformerMorphModel::UpdateStatistics()
+{
+	NumMorphTargets = MorphTargetSet.IsValid() ? MorphTargetSet->MorphBuffers.GetNumMorphs() : 0;
+	CompressedMorphDataSizeInBytes = MorphTargetSet.IsValid() ? MorphTargetSet->MorphBuffers.GetMorphDataSizeInBytes() : 0;
+	UncompressedMorphDataSizeInBytes = GetMorphTargetDeltas().Num() * sizeof(FVector3f);
+}
+
+void UMLDeformerMorphModel::SetMorphTargetsMaxWeights(const TArray<float>& MaxWeights)
+{
+	MaxMorphWeights = MaxWeights;
+}
+
+TArrayView<const float> UMLDeformerMorphModel::GetMorphTargetMaxWeights() const
+{
+	return MaxMorphWeights;
+}
+
+TArrayView<const float> UMLDeformerMorphModel::GetMorphTargetErrorValues() const
+{
+	return MorphTargetErrors;
+}
+
+TArrayView<const int32> UMLDeformerMorphModel::GetMorphTargetErrorOrder() const
+{ 
+	return MorphTargetErrorOrder;
+}
+
+TArrayView<const FMLDeformerMorphModelQualityLevel> UMLDeformerMorphModel::GetQualityLevels() const
+{
+	return QualityLevels;
+}
+
+TArray<FMLDeformerMorphModelQualityLevel>& UMLDeformerMorphModel::GetQualityLevelsArray()
+{
+	return QualityLevels;
+}
+
+float UMLDeformerMorphModel::GetMorphTargetError(int32 MorphIndex) const
+{ 
+	return MorphTargetErrors[MorphIndex];
+}
+
+void UMLDeformerMorphModel::SetMorphTargetError(int32 MorphIndex, float Error)
+{ 
+	MorphTargetErrors[MorphIndex] = Error;
+}
+
+int32 UMLDeformerMorphModel::GetNumActiveMorphs(int32 QualityLevel) const
+{
+	if (QualityLevels.IsEmpty() || MorphTargetErrorOrder.IsEmpty() || MorphTargetErrors.IsEmpty())
+	{
+		return FMath::Max<int32>(0, NumMorphTargets - 1);	// -1 Because this number includes the means morph target.
+	}
+
+	const int32 ClampedQualityLevel = FMath::Clamp<int32>(QualityLevel, 0, QualityLevels.Num() - 1);
+	const int32 NumActiveMorphs = FMath::Clamp<int32>(QualityLevels[ClampedQualityLevel].GetMaxActiveMorphs(), 0, NumMorphTargets - 1);	// -1 Because we want to exclude the means morph target.
+	return NumActiveMorphs;
+}
+
+#if WITH_EDITOR
+void UMLDeformerMorphModel::UpdateMemoryUsage()
+{
+	Super::UpdateMemoryUsage();
+
+	// We strip the deltas when cooking.
+	CookedMemUsageInBytes -= MorphTargetDeltas.Num() * sizeof(FVector3f);
+
+	// Add the compressed morph size.
+	GPUMemUsageInBytes += GetCompressedMorphDataSizeInBytes();
+}
+#endif
+
 #undef LOCTEXT_NAMESPACE
diff --git a/Engine/Plugins/Animation/MLDeformer/MLDeformerFramework/Source/MLDeformerFramework/Private/MLDeformerMorphModelInstance.cpp b/Engine/Plugins/Animation/MLDeformer/MLDeformerFramework/Source/MLDeformerFramework/Private/MLDeformerMorphModelInstance.cpp
--- a/Engine/Plugins/Animation/MLDeformer/MLDeformerFramework/Source/MLDeformerFramework/Private/MLDeformerMorphModelInstance.cpp
+++ b/Engine/Plugins/Animation/MLDeformer/MLDeformerFramework/Source/MLDeformerFramework/Private/MLDeformerMorphModelInstance.cpp
@@ -2,6 +2,7 @@
 
 #include "MLDeformerMorphModelInstance.h"
 #include "MLDeformerMorphModel.h" 
+#include "MLDeformerComponent.h" 
 #include "NeuralNetwork.h"
 #include "Components/ExternalMorphSet.h"
 #include "Components/SkeletalMeshComponent.h"
@@ -23,6 +24,11 @@ bool UMLDeformerMorphModelInstance::IsValidForDataProvider() const
 	return true;
 }
 
+int32 UMLDeformerMorphModelInstance::GetExternalMorphSetID() const
+{ 
+	return ExternalMorphSetID;
+}
+
 void UMLDeformerMorphModelInstance::Release()
 {
 	// Try to unregister the morph target and morph target set.
@@ -97,6 +103,39 @@ void UMLDeformerMorphModelInstance::PostMLDeformerComponentInit()
 	}
 }
 
+void UMLDeformerMorphModelInstance::Tick(float DeltaTime, float ModelWeight)
+{
+	// Detect changes in quality level.
+	const int32 CurrentQualityLevel = GetMLDeformerComponent()->GetQualityLevel();
+	if (CurrentQualityLevel != LastQualityLevel)
+	{
+		LastQualityLevel = CurrentQualityLevel;
+		MorphLerpAlpha = 0.0f;
+
+		const int LOD = 0;	// For now we only support LOD 0, as we can't setup an ML Deformer per LOD yet.
+		FExternalMorphSetWeights* WeightData = FindWeightData(LOD);
+		if (WeightData)
+		{
+			StartMorphWeights = WeightData->Weights;
+		}
+	}
+
+	if (StartMorphWeights.IsEmpty())
+	{
+		const UMLDeformerMorphModel* MorphModel = Cast<UMLDeformerMorphModel>(Model);
+		StartMorphWeights.SetNumZeroed(MorphModel->GetNumMorphTargets());
+	}
+
+	Super::Tick(DeltaTime, ModelWeight);
+
+	// Update the morph lerp towards the target.
+	MorphLerpAlpha += DeltaTime * 10.0f;
+	if (MorphLerpAlpha > 1.0f)
+	{
+		MorphLerpAlpha = 1.0f;
+	}
+}
+
 FExternalMorphSetWeights* UMLDeformerMorphModelInstance::FindWeightData(int32 LOD) const
 {
 	const UMLDeformerMorphModel* MorphModel = Cast<UMLDeformerMorphModel>(Model);
@@ -160,18 +199,40 @@ void UMLDeformerMorphModelInstance::Execute(float ModelWeight)
 		if (NumMorphTargets == NumNetworkWeights + 1)
 		{
 			// Set the first morph target, which represents the means, to a weight of 1.0, as it always needs to be fully active.
-			WeightData->Weights[0] = 1.0f * ModelWeight;
+			WeightData->Weights[0] = ModelWeight;
 
 			// Update all generated morph target weights with the values calculated by our neural network.
-			for (int32 MorphIndex = 0; MorphIndex < NumNetworkWeights; ++MorphIndex)
+			UMLDeformerMorphModel* MorphModel = Cast<UMLDeformerMorphModel>(Model);
+			const TArrayView<const float> ErrorValues = MorphModel->GetMorphTargetErrorValues();
+			const TArrayView<const int32> ErrorOrder = MorphModel->GetMorphTargetErrorOrder();
+			if (!ErrorValues.IsEmpty())
 			{
-				const float OutputTensorValue = OutputTensor.At<float>(MorphIndex);
-				WeightData->Weights[MorphIndex + 1] = OutputTensorValue * ModelWeight;
+				const int32 QualityLevel = GetMLDeformerComponent()->GetQualityLevel();
+				const int32 NumActiveMorphs = MorphModel->GetNumActiveMorphs(QualityLevel);
+				for (int32 Index = 0; Index < NumActiveMorphs; ++Index)
+				{
+					const int32 MorphIndex = ErrorOrder[Index];
+					const float TargetWeight = OutputTensor.At<float>(MorphIndex) * ModelWeight;
+					WeightData->Weights[MorphIndex + 1] = FMath::Lerp<float, float>(StartMorphWeights[MorphIndex + 1], TargetWeight, MorphLerpAlpha);
+				}
+
+				// Disable all inactive morphs.
+				for (int32 Index = NumActiveMorphs; Index < NumNetworkWeights; ++Index)
+				{
+					const int32 MorphIndex = ErrorOrder[Index];
+					WeightData->Weights[MorphIndex + 1] = FMath::Lerp<float, float>(StartMorphWeights[MorphIndex + 1], 0.0f, MorphLerpAlpha);
+				}
 			}
+			else
+			{
+				for (int32 MorphIndex = 0; MorphIndex < NumNetworkWeights; ++MorphIndex)
+				{
+					WeightData->Weights[MorphIndex + 1] = OutputTensor.At<float>(MorphIndex) * ModelWeight;
+				}
+			}
+			return;
 		}
 	}
-	else
-	{
-		WeightData->ZeroWeights();
-	}
+
+	WeightData->ZeroWeights();
 }
diff --git a/Engine/Plugins/Animation/MLDeformer/MLDeformerFramework/Source/MLDeformerFramework/Private/MLDeformerPerfCounter.cpp b/Engine/Plugins/Animation/MLDeformer/MLDeformerFramework/Source/MLDeformerFramework/Private/MLDeformerPerfCounter.cpp
new file mode 100644
--- /dev/null
+++ b/Engine/Plugins/Animation/MLDeformer/MLDeformerFramework/Source/MLDeformerFramework/Private/MLDeformerPerfCounter.cpp
@@ -0,0 +1,95 @@
+// Copyright Epic Games, Inc. All Rights Reserved.
+
+#include "MLDeformerPerfCounter.h"
+#include "HAL/PlatformTime.h"
+
+namespace UE::MLDeformer
+{
+	FMLDeformerPerfCounter::FMLDeformerPerfCounter()
+	{
+		SetHistorySize(100);
+	}
+
+	void FMLDeformerPerfCounter::SetHistorySize(int32 NumHistoryItems)
+	{
+		CycleHistory.Reset();
+		CycleHistory.SetNumZeroed(NumHistoryItems);
+		Reset();
+	}
+
+	void FMLDeformerPerfCounter::BeginSample()
+	{
+		StartTime = FPlatformTime::Cycles();
+	}
+
+	void FMLDeformerPerfCounter::EndSample()
+	{
+		Cycles = FPlatformTime::Cycles() - StartTime;
+
+		if (NumSamples == 0 || (!CycleHistory.IsEmpty() && (NumSamples % CycleHistory.Num()) == 0))
+		{
+			CyclesMin = Cycles;
+			CyclesMax = Cycles;
+		}
+		else
+		{
+			CyclesMin = FMath::Min<int32>(Cycles, CyclesMin);
+			CyclesMax = FMath::Max<int32>(Cycles, CyclesMax);
+		}
+
+		if (!CycleHistory.IsEmpty())
+		{
+			CycleHistory[NumSamples % CycleHistory.Num()] = Cycles;
+		}
+
+		NumSamples++;
+	}
+
+	int32 FMLDeformerPerfCounter::GetCycles() const
+	{
+		return Cycles;
+	}
+
+	int32 FMLDeformerPerfCounter::GetCyclesMin() const
+	{
+		return CyclesMin;
+	}
+
+	int32 FMLDeformerPerfCounter::GetCyclesMax() const
+	{
+		return CyclesMax;
+	}
+
+	int32 FMLDeformerPerfCounter::GetCyclesAverage() const
+	{
+		if (CycleHistory.IsEmpty() || NumSamples == 0)
+		{
+			return Cycles;
+		}
+
+		int32 Sum = 0;
+		const int32 NumRecordedSamples = FMath::Min<int32>(NumSamples, CycleHistory.Num());
+		for (int32 Index = 0; Index < NumRecordedSamples; ++Index)
+		{
+			Sum += CycleHistory[Index];
+		}
+		return Sum / NumRecordedSamples;
+	}
+
+	int32 FMLDeformerPerfCounter::GetNumSamples() const
+	{
+		return NumSamples;
+	}
+
+	int32 FMLDeformerPerfCounter::GetHistorySize() const
+	{
+		return CycleHistory.Num();
+	}
+
+	void FMLDeformerPerfCounter::Reset()
+	{
+		NumSamples = 0;
+		CyclesMax = 0;
+		CyclesMin = 0;
+	}
+}	// namespace UE::MLDeformer
diff --git a/Engine/Plugins/Animation/MLDeformer/MLDeformerFramework/Source/MLDeformerFramework/Public/MLDeformerComponent.h b/Engine/Plugins/Animation/MLDeformer/MLDeformerFramework/Source/MLDeformerFramework/Public/MLDeformerComponent.h
--- a/Engine/Plugins/Animation/MLDeformer/MLDeformerFramework/Source/MLDeformerFramework/Public/MLDeformerComponent.h
+++ b/Engine/Plugins/Animation/MLDeformer/MLDeformerFramework/Source/MLDeformerFramework/Public/MLDeformerComponent.h
@@ -6,6 +6,7 @@
 #include "UObject/Object.h"
 #include "Components/ActorComponent.h"
 #include "RenderCommandFence.h"
+#include "MLDeformerPerfCounter.h"
 #include "MLDeformerComponent.generated.h"
 
 class UMLDeformerAsset;
@@ -64,6 +65,28 @@ public:
 	UFUNCTION(BlueprintCallable, Category = "MLDeformer")
 	void SetWeight(float NormalizedWeightValue)					{ SetWeightInternal(NormalizedWeightValue); }
 
+	/** 
+	 * The quality level of the deformer. A value of 0 is the highest quality, 1 is a step lower, etc.
+	 * It is up to the models to configure what each quality level means visually.
+	 * If the quality level will be clamped to the available quality levels, so if you choose quality level 100, while there are only 3 levels, then quality
+	 * level 3 will be used in this case, which represents the lowest available quality.
+	 * In morph based models each quality level defines how many morph targets are active at most.
+	 * @param InQualityLevel The quality level to switch to.
+	 */
+	UFUNCTION(BlueprintCallable, Category = "MLDeformer")
+	void SetQualityLevel(int32 InQualityLevel)					{ QualityLevel = FMath::Max<int32>(InQualityLevel, 0); }
+
+	/** 
+	 * The quality level of the deformer. A value of 0 is the highest quality, 1 is a step lower, etc.
+	 * It is up to the models to configure what each quality level means visually.
+	 * If the quality level will be clamped to the available quality levels, so if you choose quality level 100, while there are only 3 levels, then quality
+	 * level 3 will be used in this case, which represents the lowest available quality.
+	 * In morph based models each quality level defines how many morph targets are active at most.
+	 * @param InQualityLevel The quality level to switch to.
+	 */
+	UFUNCTION(BlueprintCallable, Category = "MLDeformer")
+	int32 GetQualityLevel() const								{ return QualityLevel; }
+
 	/**
 	 * Get the ML Deformer asset that is used by this component.
 	 * @return A pointer to the ML Deformer asset.
@@ -116,9 +139,17 @@ public:
 	 */
 	void SetSuppressMeshDeformerLogWarnings(bool bSuppress)		{ bSuppressMeshDeformerLogWarnings = bSuppress; }
 
+#if WITH_EDITOR
+	/**
+	 * Get the performance counter that measures how much time is spent inside the Tick function.
+	 */
+	const UE::MLDeformer::FMLDeformerPerfCounter& GetTickPerfCounter() const	{ return TickPerfCounter; }
+#endif
+
 	// Get property names.
 	static FName GetDeformerAssetPropertyName()					{ return GET_MEMBER_NAME_CHECKED(UMLDeformerComponent, DeformerAsset); }
 	static FName GetWeightPropertyName()						{ return GET_MEMBER_NAME_CHECKED(UMLDeformerComponent, Weight); }
+	static FName GetQualityLevelPropertyName()					{ return GET_MEMBER_NAME_CHECKED(UMLDeformerComponent, QualityLevel); }
 
 protected:
 	// AActorComponent overrides.
@@ -145,7 +176,17 @@ protected:
 	/** Set the ML Deformer asset. */
 	virtual void SetDeformerAssetInternal(UMLDeformerAsset* const InDeformerAsset) { DeformerAsset = InDeformerAsset; UpdateSkeletalMeshComponent(); }
 
+#if WITH_EDITOR
+	/** Reset the tick cycle counters. */
+	void ResetTickCycleCounters();
+#endif
+
 protected:
+#if WITH_EDITOR
+	/** The performance counter that measures timing of the Tick function. */
+	UE::MLDeformer::FMLDeformerPerfCounter TickPerfCounter;
+#endif
+
 	/** Render command fence that let's us wait for all other commands to finish. */
 	FRenderCommandFence RenderCommandFence;
 
@@ -170,6 +211,16 @@ protected:
 	UPROPERTY(EditAnywhere, Category = "ML Deformer", meta = (ClampMin = "0.0", ClampMax = "1.0"))
 	float Weight = 1.0f;
 
+	/**
+	 * The quality level of the deformer. A value of 0 is the highest quality, 1 is a step lower, etc.
+	 * It is up to the models to configure what each quality level means visually.
+	 * If the quality level will be clamped to the available quality levels, so if you choose quality level 100, while there are only 3 levels, then quality
+	 * level 3 will be used in this case, which represents the lowest available quality.
+	 * In morph based models each quality level defines how many morph targets are active at most.
+	 */
+	UPROPERTY(EditAnywhere, Category = "ML Deformer", meta = (ClampMin = "0"))
+	int32 QualityLevel = 0;
+
 	/** The deformation model instance. This is used to perform the runtime updates and run the inference. */
 	UPROPERTY(Transient)
 	TObjectPtr<UMLDeformerModelInstance> ModelInstance = nullptr;
diff --git a/Engine/Plugins/Animation/MLDeformer/MLDeformerFramework/Source/MLDeformerFramework/Public/MLDeformerModel.h b/Engine/Plugins/Animation/MLDeformer/MLDeformerFramework/Source/MLDeformerFramework/Public/MLDeformerModel.h
--- a/Engine/Plugins/Animation/MLDeformer/MLDeformerFramework/Source/MLDeformerFramework/Public/MLDeformerModel.h
+++ b/Engine/Plugins/Animation/MLDeformer/MLDeformerFramework/Source/MLDeformerFramework/Public/MLDeformerModel.h
@@ -21,6 +21,7 @@ class UMLDeformerModelInstance;
 class UMLDeformerComponent;
 class UMLDeformerInputInfo;
 
+/** The channel to get the mask data from. */
 UENUM()
 enum class EMLDeformerMaskChannel : uint8
 {
@@ -43,6 +44,16 @@ enum class EMLDeformerMaskChannel : uint8
 
 namespace UE::MLDeformer
 {
+	/** The memory usage request flags, which you pass UMLDeformer::GetMemUsageInBytes method. */
+	enum class EMemUsageRequestFlags : uint8
+	{
+		/** Get the uncooked memory, so inside the editor. */
+		Uncooked,
+
+		/** Get the memory usage as we if are cooked. */
+		Cooked
+	};
+
 	/**
 	 * The vertex map, but in a GPU buffer. 
 	 * This map basically has a DCC vertex number for every render vertex.
@@ -137,6 +148,15 @@ public:
 	 */
 	virtual bool DoesSupportCurves() const					{ return true; }
 
+	/**
+	 * Does this model support deformer quality levels?
+	 * For example Morph based models can disable certain morph targets based on this quality level.
+	 * On default this is disabled for models. You can override this method and make it return true to support it.
+	 * Morph based models on enable this on default.
+	 * @return Returns true when Deformer Quality is supported, otherwise false is returned.
+	 */
+	virtual bool DoesSupportQualityLevels() const			{ return false; }
+
 	/**
 	 * Check whether the neural network of this model should run on the GPU or not.
 	 * This is false on default, which makes it a CPU based neural network.
@@ -166,7 +186,6 @@ public:
 	 */
 	void SetSkeletalMesh(USkeletalMesh* SkelMesh)			{ SkeletalMesh = SkelMesh; }
 
-
 #if WITH_EDITORONLY_DATA
 	/**
 	 * Check whether this model currently has a training mesh setup or not.
@@ -209,6 +228,31 @@ public:
 	 * @return The number of imported vertices, which is the vertex count as seen in the DCC.
 	 */ 
 	static int32 ExtractNumImportedSkinnedVertices(const USkeletalMesh* SkeletalMesh);
+
+	/** Call this if you want it to recalculate the memory usage again after the tick has completed. */
+	void InvalidateMemUsage();
+
+	/** 
+	 * Check whether we invalidated the memory usage. If this returns true, we need to recompute it.
+	 * @return Returns true when we should make a call to UpdateMemoryUsage.
+	 */
+	bool IsMemUsageInvalidated() const;
+
+	/**
+	 * Get the memory usage for this model.
+	 * @param Flags The memory request flags.
+	 * @return The number of bytes that this model uses.
+	 */
+	uint64 GetMemUsageInBytes(UE::MLDeformer::EMemUsageRequestFlags Flags) const;
+
+	/**
+	 * Get the GPU memory usage for this model.
+	 * @return The number of bytes that this model uses.
+	 */
+	uint64 GetGPUMemUsageInBytes() const;
+
+	/** Force update the cached memory usage. */
+	virtual void UpdateMemoryUsage();
 #endif
 
 	// UObject overrides.
@@ -495,6 +539,21 @@ protected:
 	 */
 	void SetNumTargetMeshVerts(int32 NumVerts)			{ NumTargetMeshVerts = NumVerts; }
 
+protected:
+#if WITH_EDITOR
+	/** Should we recalculate the memory usage? */
+	bool bInvalidateMemUsage = true;
+
+	/** The computed memory usage. */
+	uint64 MemUsageInBytes = 0;
+
+	/** The cooked memory usage. */
+	uint64 CookedMemUsageInBytes = 0;
+
+	/** GPU memory usage. */
+	uint64 GPUMemUsageInBytes = 0;
+#endif
+
 private:
 	/** The deformer asset that this model is part of. */
 	TObjectPtr<UMLDeformerAsset> DeformerAsset = nullptr;
diff --git a/Engine/Plugins/Animation/MLDeformer/MLDeformerFramework/Source/MLDeformerFramework/Public/MLDeformerModelInstance.h b/Engine/Plugins/Animation/MLDeformer/MLDeformerFramework/Source/MLDeformerFramework/Public/MLDeformerModelInstance.h
--- a/Engine/Plugins/Animation/MLDeformer/MLDeformerFramework/Source/MLDeformerFramework/Public/MLDeformerModelInstance.h
+++ b/Engine/Plugins/Animation/MLDeformer/MLDeformerFramework/Source/MLDeformerFramework/Public/MLDeformerModelInstance.h
@@ -11,6 +11,7 @@
 class UMLDeformerModel;
 class UNeuralNetwork;
 class USkeletalMeshComponent;
+class UMLDeformerComponent;
 
 /**
  * An instance of the ML Deformer model.
@@ -119,6 +120,12 @@ public:
 	/** Check whether we already called UMLDeformerModel::PostMLDeformerComponentInit. */
 	bool HasPostInitialized() const;
 
+	/**
+	 * Get the ML Deformer component that this instance is part of.
+	 * @return A pointer to the ML Deformer component.
+	 */
+	UMLDeformerComponent* GetMLDeformerComponent() const;
+
 protected:
 	/**
 	 * Update the neural network input values directly inside its input tensor.
diff --git a/Engine/Plugins/Animation/MLDeformer/MLDeformerFramework/Source/MLDeformerFramework/Public/MLDeformerMorphModel.h b/Engine/Plugins/Animation/MLDeformer/MLDeformerFramework/Source/MLDeformerFramework/Public/MLDeformerMorphModel.h
--- a/Engine/Plugins/Animation/MLDeformer/MLDeformerFramework/Source/MLDeformerFramework/Public/MLDeformerMorphModel.h
+++ b/Engine/Plugins/Animation/MLDeformer/MLDeformerFramework/Source/MLDeformerFramework/Public/MLDeformerMorphModel.h
@@ -4,6 +4,7 @@
 
 #include "CoreMinimal.h"
 #include "MLDeformerGeomCacheModel.h"
+#include "MLDeformerMorphModelQualityLevel.h"
 #include "MLDeformerMorphModel.generated.h"
 
 struct FExternalMorphSet;
@@ -20,33 +21,68 @@ public:
 	UMLDeformerMorphModel(const FObjectInitializer& ObjectInitializer);
 
 	// UMLDeformerModel overrides.
-	virtual bool IsNeuralNetworkOnGPU() const override		{ return false; }	// CPU based neural network.
+	virtual bool IsNeuralNetworkOnGPU() const override				{ return false; }	// CPU based neural network.
+	virtual bool DoesSupportQualityLevels() const override			{ return true; }	// We can disable morph targets based on the Deformer LOD float value.
 	virtual void Serialize(FArchive& Archive) override;
 	virtual UMLDeformerModelInstance* CreateModelInstance(UMLDeformerComponent* Component) override;
+#if WITH_EDITOR
+	virtual void UpdateMemoryUsage() override;
+#endif
 	// ~END UMLDeformerModel overrides.
 
 	// UObject overrides.
-	void BeginDestroy() override;
+	virtual void BeginDestroy() override;
+	virtual void PostLoad() override;
 	// ~END UObject overrides.
 
-	float GetMorphTargetDeltaThreshold() const				{ return MorphTargetDeltaThreshold; }
-	float GetMorphTargetErrorTolerance() const				{ return MorphTargetErrorTolerance; }
-	bool GetIncludeMorphTargetNormals() const				{ return bIncludeNormals; }
-	EMLDeformerMaskChannel GetMaskChannel() const			{ return MaskChannel; }
-	bool GetInvertMaskChannel() const						{ return bInvertMaskChannel; }
+	int32 GetNumMorphTargets() const								{ return NumMorphTargets; }
+	uint64 GetCompressedMorphDataSizeInBytes() const				{ return CompressedMorphDataSizeInBytes; }
+	uint64 GetUncompressedMorphDataSizeInBytes() const				{ return UncompressedMorphDataSizeInBytes; }
 
-	void SetMorphTargetDeltaThreshold(float Threshold)		{ MorphTargetDeltaThreshold = Threshold; }
-	void SetMorphTargetErrorTolerance(float Tolerance)		{ MorphTargetErrorTolerance = Tolerance; }
-	void SetIncludeMorphTargetNormals(bool bInclude)		{ bIncludeNormals = bInclude; }
-	void SetWeightMask(EMLDeformerMaskChannel Channel)		{ MaskChannel = Channel; }
-	void SetInvertMaskChannel(bool bInvert)					{ bInvertMaskChannel = bInvert; }
+	float GetMorphDeltaZeroThreshold() const						{ return MorphDeltaZeroThreshold; }
+	float GetMorphCompressionLevel() const							{ return MorphCompressionLevel; }
+	float GetMaxMorphTargetErrorInUnits() const						{ return !MorphTargetErrorOrder.IsEmpty() ? MorphTargetErrors[MorphTargetErrorOrder[0]] : 0.0f; }
+	bool GetIncludeMorphTargetNormals() const						{ return bIncludeNormals; }
+	EMLDeformerMaskChannel GetMaskChannel() const					{ return MaskChannel; }
+	bool GetInvertMaskChannel() const								{ return bInvertMaskChannel; }
+
+	void SetMorphDeltaZeroThreshold(float Threshold)				{ MorphDeltaZeroThreshold = Threshold; }
+	void SetMorphCompressionlevel(float Tolerance)					{ MorphCompressionLevel = Tolerance; }
+	void SetIncludeMorphTargetNormals(bool bInclude)				{ bIncludeNormals = bInclude; }
+	void SetWeightMask(EMLDeformerMaskChannel Channel)				{ MaskChannel = Channel; }
+	void SetInvertMaskChannel(bool bInvert)							{ bInvertMaskChannel = bInvert; }
+
+	UE_DEPRECATED(5.2, "Please use GetMorphDeltaZeroThreshold instead.")
+	float GetMorphTargetDeltaThreshold() const						{ return MorphDeltaZeroThreshold; }
+
+	UE_DEPRECATED(5.2, "Please use SetMorphDeltaZeroThreshold instead.")
+	void SetMorphTargetDeltaThreshold(float Threshold)				{ MorphDeltaZeroThreshold = Threshold; }
+
+	UE_DEPRECATED(5.2, "Please use GetMorphCompressionLevel instead.")
+	float GetMorphTargetErrorTolerance() const						{ return MorphCompressionLevel; }
+
+	UE_DEPRECATED(5.2, "Please use SetMorphCompressionLevel instead.")
+	void SetMorphTargetErrorTolerance(float Tolerance)				{ MorphCompressionLevel = Tolerance; }
 
 	// Get property names.
-	static FName GetMorphTargetDeltaThresholdPropertyName() { return GET_MEMBER_NAME_CHECKED(UMLDeformerMorphModel, MorphTargetDeltaThreshold); }
-	static FName GetMorphTargetErrorTolerancePropertyName() { return GET_MEMBER_NAME_CHECKED(UMLDeformerMorphModel, MorphTargetErrorTolerance); }
-	static FName GetIncludeMorphTargetNormalsPropertyName() { return GET_MEMBER_NAME_CHECKED(UMLDeformerMorphModel, bIncludeNormals); }
-	static FName GetMaskChannelPropertyName()				{ return GET_MEMBER_NAME_CHECKED(UMLDeformerMorphModel, MaskChannel); }
-	static FName GetInvertMaskChannelPropertyName()			{ return GET_MEMBER_NAME_CHECKED(UMLDeformerMorphModel, bInvertMaskChannel); }
+	static FName GetMorphDeltaZeroThresholdPropertyName()			{ return GET_MEMBER_NAME_CHECKED(UMLDeformerMorphModel, MorphDeltaZeroThreshold); }
+	static FName GetMorphCompressionLevelPropertyName()				{ return GET_MEMBER_NAME_CHECKED(UMLDeformerMorphModel, MorphCompressionLevel); }
+	static FName GetIncludeMorphTargetNormalsPropertyName()			{ return GET_MEMBER_NAME_CHECKED(UMLDeformerMorphModel, bIncludeNormals); }
+	static FName GetMaskChannelPropertyName()						{ return GET_MEMBER_NAME_CHECKED(UMLDeformerMorphModel, MaskChannel); }
+	static FName GetInvertMaskChannelPropertyName()					{ return GET_MEMBER_NAME_CHECKED(UMLDeformerMorphModel, bInvertMaskChannel); }
+	static FName GetNumMorphTargetsPropertyName()					{ return GET_MEMBER_NAME_CHECKED(UMLDeformerMorphModel, NumMorphTargets); }
+	static FName GetCompressedMorphDataSizeInBytesPropertyName()	{ return GET_MEMBER_NAME_CHECKED(UMLDeformerMorphModel, CompressedMorphDataSizeInBytes); }
+	static FName GetUncompressedMorphDataSizeInBytesPropertyName()	{ return GET_MEMBER_NAME_CHECKED(UMLDeformerMorphModel, UncompressedMorphDataSizeInBytes); }
+	static FName GetQualityLevelsPropertyName()						{ return GET_MEMBER_NAME_CHECKED(UMLDeformerMorphModel, QualityLevels); }
+
+	UE_DEPRECATED(5.2, "Please use GetMorphDeltaZeroThresholdPropertyName instead.")
+	static FName GetMorphTargetDeltaThresholdPropertyName()			{ return GET_MEMBER_NAME_CHECKED(UMLDeformerMorphModel, MorphDeltaZeroThreshold); }
+
+	UE_DEPRECATED(5.2, "Please use GetMorphCompressionLevelPropertyName instead.")
+	static FName GetMorphTargetErrorTolerancePropertyName()			{ return GET_MEMBER_NAME_CHECKED(UMLDeformerMorphModel, MorphCompressionLevel); }
+
+	/** Update the statistics properties of this model. */
+	virtual void UpdateStatistics();
 
 	/**
 	 * Set the per vertex deltas, as a set of floats. Each vertex delta must have 3 floats.
@@ -71,17 +107,46 @@ public:
 	UFUNCTION(BlueprintCallable, Category = "MLDeformerMorphModel")
 	void SetMorphTargetDeltas(const TArray<FVector3f>& Deltas);
 
+	/**
+	 * Set the order of importance during LOD, for the morph targets.
+	 * Basically this specifies the sorted order of the morph targets, based on the error they introduce when disabling them.
+	 * @param MorphTargetOrder The array of morph target indices, starting with the most important morph target, and ending with the least important morph target's index.
+	 * @param ErrorValues The error value for each morph target (not sorted), so index 0 contains the error value of the first morph target, which isn't necesarrily the one with highest error.
+	 */
+	UFUNCTION(BlueprintCallable, Category = "MLDeformerMorphModel")
+	void SetMorphTargetsErrorOrder(const TArray<int32>& MorphTargetOrder, const TArray<float>& ErrorValues);
+
+	/**
+	 * Set the maximum weights for each morph target when we pass the training inputs into the trained neural network.
+	 * After training we run the inputs through the trained network, and get the weights for each morph target.
+	 * The array passed to this function must be the maximum of the absolute weight value over all training samples for each specific morph.
+	 * We use this to scale the morph target deltas to get an estimate of the maximum length they will ever have, which we can in turn use
+	 * to estimate how important a given morph target is.
+	 * Keep in mind that these are the maximum of the absolute weights. As weights can be negative, we first take the absolute value of them.
+	 * Training python scripts mostly will call this function to set the values.
+	 * @param MaxWeights The maximum of the absolute weight values.
+	 */
+	UFUNCTION(BlueprintCallable, Category = "MLDeformerMorphModel")
+	void SetMorphTargetsMaxWeights(const TArray<float>& MaxWeights);
+
+	/**
+	 * Get the estimated maximum weight values that the morph targets will ever have. It is not guaranteed that the weights will never be larger than these values though.
+	 * These values are used to estimate the importance levels of the morph targets.
+	 * @see SetMorphTargetsMaxWeights.
+	 */
+	TArrayView<const float> GetMorphTargetMaxWeights() const;
+
 	/**
 	 * Get the morph target delta vectors array.
 	 * The layout of this array is [morphdeltas_target0, morphdeltas_target1, ..., morphdeltas_targetN].
 	 * So the total number of items in the array returned equals (NumMorphTargets * NumBaseMeshVerts).
 	 */
-	const TArray<FVector3f>& GetMorphTargetDeltas() const	{ return MorphTargetDeltas; }
+	const TArray<FVector3f>& GetMorphTargetDeltas() const		{ return MorphTargetDeltas; }
 
 	/**
 	 * Get the morph target set.
 	 */
-	TSharedPtr<FExternalMorphSet> GetMorphTargetSet() const { return MorphTargetSet; }
+	TSharedPtr<FExternalMorphSet> GetMorphTargetSet() const		{ return MorphTargetSet; }
 
 	/**
 	 * Get the start index into the array of deltas (vectors3's), for a given morph target.
@@ -91,16 +156,106 @@ public:
 	 */
 	int32 GetMorphTargetDeltaStartIndex(int32 MorphTargetIndex) const;
 
+	/**
+	 * Get the morph target error values.
+	 * This is an error value for each morph target, which indicates how much error you would introduce when disabling that specific morph target.
+	 * It can be used as some sort of LOD system, where we disable morph targets that are marked with an error below a specific threshold.
+	 * Keep in mind that the first morph target always contains the means, and is always active. So actually you have to +1 the index.
+	 * @return The array view of the error values.
+	 */
+	TArrayView<const float> GetMorphTargetErrorValues() const;
+
+	/**
+	 * Get the array of morph target indices, sorted by error value, from largest to smallest error.
+	 * So the first array value contains the index of the most important morph target and the last element contains the least important morph target.
+	 * @return An array view of the morph target indices.
+	 */
+	TArrayView<const int32> GetMorphTargetErrorOrder() const;
+
+	/**
+	 * Get the quality levels.
+	 * @return An array view of the quality levels.
+	 */
+	TArrayView<const FMLDeformerMorphModelQualityLevel> GetQualityLevels() const;
+
+	/**
+	 * Get the quality levels, with access to modify the quality levels.
+	 * @return The array containing the quality levels.
+	 */
+	TArray<FMLDeformerMorphModelQualityLevel>& GetQualityLevelsArray();
+	
+	/**
+	 * Get the estimated error that disabling a specific morph target would introduce.
+	 * This error is represented in units, so in cm.
+	 * @param MorphIndex The morph target index to get the error for.
+	 * @return The estimated error value, in cm.
+	 */
+	float GetMorphTargetError(int32 MorphIndex) const;
+
+	/**
+	 * Set the estimated error that disabling a specific morph target would introduce.
+	 * This error is represented in units, so in cm.
+	 * @param MorphIndex The morph target index to get the error for.
+	 * @param Error The estimated error value, in cm.
+	 */
+	void SetMorphTargetError(int32 MorphIndex, float Error);
+
+	/** 
+	 * Get the number of active morph targets for a given quality level.
+	 * If no morph error data is available, or if no quality levels have been setup, the total number of morph targets is returned.
+	 * If there is such data, the quality level will be clamped to be within a valid range.
+	 * Finally, if the number of active morph targets specified by the quality level will be clamped within a valid range as well.
+	 * So basically this method will always return a valid number of morph targets, whatever the input is.
+	 * @param QualityLevel The quality level to get the number of active morph targets for.
+	 */
+	int32 GetNumActiveMorphs(int32 QualityLevel) const;
+
 private:
 	/** The compressed morph target data, ready for the GPU. */
 	TSharedPtr<FExternalMorphSet> MorphTargetSet;
 
 	/**
-	 * The entire set of morph target deltas, 3 per vertex, for each morph target, as one flattened buffer.
+	 * The entire set of uncompressed morph target deltas, 3 per vertex, for each morph target, as one flattened buffer.
 	 * So the size of this buffer is: (NumVertsPerMorphTarget * 3 * NumMorphTargets).
 	 */
+	UPROPERTY()
 	TArray<FVector3f> MorphTargetDeltas;
 
+	/** The number of morph targets. */
+	UPROPERTY()
+	int32 NumMorphTargets = 0;
+
+	/** The compressed memory usage of the morph targets. This is approximately what this MLD asset will use in your packaged project. */
+	UPROPERTY()
+	uint64 CompressedMorphDataSizeInBytes = 0;
+
+	/** The uncompressed memory usage of the morph targets. Please note that your packaged project will use the compressed amount of memory, and not the uncompresed amount of memory. */
+	UPROPERTY()
+	uint64 UncompressedMorphDataSizeInBytes = 0;
+
+	/** 
+	 * An array of integers where the first element is the most important morph target and the last element is the least important one.
+	 * The number of elements in the array is equal to the total number of morph targets.
+	 */
+	UPROPERTY()
+	TArray<int32> MorphTargetErrorOrder;
+
+	/** An error value, for each morph target. */
+	UPROPERTY()
+	TArray<float> MorphTargetErrors;
+
+	/** The maximum absolute weight values of the morph targets, during training. One value for each morph target. */
+	UPROPERTY()
+	TArray<float> MaxMorphWeights;
+
+	/** 
+	 * The list of quality levels, where the first item represents the highest quality and the last element the lowest quality level.
+	 * The number in each quality level represents the number of active morph targets for that quality level.
+	 * These numbers will be clamped internally to be within valid ranges in case they go beyond the amount of morphs that exist.
+	 */
+	UPROPERTY(EditAnywhere, Category = "Deformer Quality", meta = (ClampMin = "1"))
+	TArray<FMLDeformerMorphModelQualityLevel> QualityLevels;
+
 	/**
 	 * Include vertex normals in the morph targets?
 	 * The advantage of this can be that it is higher performance than recomputing the normals.
@@ -114,12 +269,15 @@ private:
 	 * This essentially removes small deltas from morph targets, which will lower the memory usage at runtime, however when set too high it can also introduce visual artifacts.
 	 * A value of 0 will result in the highest quality morph targets, at the cost of higher runtime memory usage.
 	 */
-	UPROPERTY(EditAnywhere, Category = "Morph Targets", meta = (ClampMin = "0.0", ClampMax = "1.0", ForceUnits="cm"))
-	float MorphTargetDeltaThreshold = 0.0025f;
+	UPROPERTY(EditAnywhere, Category = "Morph Targets", DisplayName = "Delta Zero Threshold", meta = (ClampMin = "0.0", ClampMax = "1.0", ForceUnits="cm"))
+	float MorphDeltaZeroThreshold = 0.0025f;
 
-	/** The morph target error tolerance. Higher values result in larger compression, but could result in visual artifacts. */
-	UPROPERTY(EditAnywhere, Category = "Morph Targets", meta = (ClampMin = "0.01", ClampMax = "500"))
-	float MorphTargetErrorTolerance = 20.0f;
+	/** 
+	 * The morph target compression level. Higher values result in larger compression, but could result in visual artifacts.
+	 * Most of the times this is a value between 20 and 200.
+	 */
+	UPROPERTY(EditAnywhere, Category = "Morph Targets", DisplayName = "Compression Level", meta = (ClampMin = "0.01", ClampMax = "1000"))
+	float MorphCompressionLevel = 20.0f;
 
 	/**
 	 * The channel data that represents the delta mask multipliers.
diff --git a/Engine/Plugins/Animation/MLDeformer/MLDeformerFramework/Source/MLDeformerFramework/Public/MLDeformerMorphModelInstance.h b/Engine/Plugins/Animation/MLDeformer/MLDeformerFramework/Source/MLDeformerFramework/Public/MLDeformerMorphModelInstance.h
--- a/Engine/Plugins/Animation/MLDeformer/MLDeformerFramework/Source/MLDeformerFramework/Public/MLDeformerMorphModelInstance.h
+++ b/Engine/Plugins/Animation/MLDeformer/MLDeformerFramework/Source/MLDeformerFramework/Public/MLDeformerMorphModelInstance.h
@@ -27,14 +27,14 @@ public:
 	virtual void Init(USkeletalMeshComponent* SkelMeshComponent) override;
 	virtual void PostMLDeformerComponentInit() override;
 	virtual void Release() override;
+	virtual void Tick(float DeltaTime, float ModelWeight) override;
 	virtual void Execute(float ModelWeight) override;
 	virtual void HandleZeroModelWeight() override;
 	virtual bool IsValidForDataProvider() const override;
 	// ~END UMLDeformerModelInstance overrides.
 
-	int32 GetExternalMorphSetID() const { return ExternalMorphSetID; }
+	int32 GetExternalMorphSetID() const;
 
-protected:
 	/**
 	 * Find the external morph target weight data for this model instance.
 	 * @param LOD The LOD level to get the weight data for.
@@ -48,4 +48,13 @@ protected:
 
 	/** The ID of the external morph target set for this instance. This gets initialized during Init. */
 	int32 ExternalMorphSetID = -1;
+
+	/** The morph target weights we are blending towards. */
+	TArray<float> StartMorphWeights;
+
+	/** The current lerp alpha (between 0 and 1). */
+	float MorphLerpAlpha = 1.0f;
+
+	/** The quality level in the previous tick. */
+	int32 LastQualityLevel = -1;
 };
diff --git a/Engine/Plugins/Animation/MLDeformer/MLDeformerFramework/Source/MLDeformerFramework/Public/MLDeformerMorphModelQualityLevel.h b/Engine/Plugins/Animation/MLDeformer/MLDeformerFramework/Source/MLDeformerFramework/Public/MLDeformerMorphModelQualityLevel.h
new file mode 100644
--- /dev/null
+++ b/Engine/Plugins/Animation/MLDeformer/MLDeformerFramework/Source/MLDeformerFramework/Public/MLDeformerMorphModelQualityLevel.h
@@ -0,0 +1,25 @@
+// Copyright Epic Games, Inc. All Rights Reserved.
+
+#pragma once
+
+#include "CoreTypes.h"
+#include "UObject/ObjectMacros.h"
+#include "MLDeformerMorphModelQualityLevel.generated.h"
+
+/** 
+ * A quality level for the Morph Model.
+ * Each quality level contains the maximum number of active morph targets.
+ */
+USTRUCT()
+struct MLDEFORMERFRAMEWORK_API FMLDeformerMorphModelQualityLevel
+{
+	GENERATED_USTRUCT_BODY()
+
+public:
+	int32 GetMaxActiveMorphs() const						{ return MaxActiveMorphs; }
+	void SetMaxActiveMorphs(int32 NumMorphs)				{ MaxActiveMorphs = NumMorphs; }
+
+public:
+	UPROPERTY(EditAnywhere, Category = "Quality Level", meta = (ClampMin = "1"))
+	int32 MaxActiveMorphs = 0;
+};
diff --git a/Engine/Plugins/Animation/MLDeformer/MLDeformerFramework/Source/MLDeformerFramework/Public/MLDeformerMorphModelVizSettings.h b/Engine/Plugins/Animation/MLDeformer/MLDeformerFramework/Source/MLDeformerFramework/Public/MLDeformerMorphModelVizSettings.h
--- a/Engine/Plugins/Animation/MLDeformer/MLDeformerFramework/Source/MLDeformerFramework/Public/MLDeformerMorphModelVizSettings.h
+++ b/Engine/Plugins/Animation/MLDeformer/MLDeformerFramework/Source/MLDeformerFramework/Public/MLDeformerMorphModelVizSettings.h
@@ -16,32 +16,37 @@ class MLDEFORMERFRAMEWORK_API UMLDeformerMorphModelVizSettings
 	GENERATED_BODY()
 
 public:
-#if WITH_EDITORONLY_DATA
 	void SetMorphTargetNumber(int32 MorphIndexToPreview)	{ MorphTargetNumber = MorphIndexToPreview; }
+	void SetDrawMorphTargets(bool bDraw)					{ bDrawMorphTargets = bDraw; }
+
+	UE_DEPRECATED(5.2, "This method will be removed, please stop using it.")
+	float GetMorphTargetDeltaThreshold() const				{ return MorphTargetDeltaThreshold_DEPRECATED; }
+
 	int32 GetMorphTargetNumber() const						{ return MorphTargetNumber; }
-	float GetMorphTargetDeltaThreshold() const				{ return MorphTargetDeltaThreshold; }
 	bool GetDrawMorphTargets() const						{ return bDrawMorphTargets; }
 
 	// Get property names.
 	static FName GetMorphTargetNumberPropertyName()			{ return GET_MEMBER_NAME_CHECKED(UMLDeformerMorphModelVizSettings, MorphTargetNumber); }
-	static FName GetMorphTargetDeltaThresholdPropertyName() { return GET_MEMBER_NAME_CHECKED(UMLDeformerMorphModelVizSettings, MorphTargetDeltaThreshold); }
 	static FName GetDrawMorphTargetsPropertyName()			{ return GET_MEMBER_NAME_CHECKED(UMLDeformerMorphModelVizSettings, bDrawMorphTargets); }
 
+	UE_DEPRECATED(5.2, "This method will be removed, as well as the MorphTargetDeltaThreshold member, please stop using it.")
+	static FName GetMorphTargetDeltaThresholdPropertyName() { return GET_MEMBER_NAME_CHECKED(UMLDeformerMorphModelVizSettings, MorphTargetDeltaThreshold_DEPRECATED); }
+
 protected:
 	/**
-	 * The morph target to visualize. The first one always being the means, so not a sparse target.
+	 * The morph target delta threshold. This is a preview of what deltas would be included in the selected morph target
+	 * when using a delta threshold during training that is equal to this value.
 	 * This only can be used after you trained, in the same editor session directly after training.
 	 */
-	UPROPERTY(EditAnywhere, Category = "Live Settings", meta = (ClampMin = "0"))
-	int32 MorphTargetNumber = 0;
+	UPROPERTY()
+	float MorphTargetDeltaThreshold_DEPRECATED = 0.01f;
 
 	/**
-	 * The morph target delta threshold. This is a preview of what deltas would be included in the selected morph target
-	 * when using a delta threshold during training that is equal to this value.
+	 * The morph target to visualize. The first one always being the means, so not a sparse target.
 	 * This only can be used after you trained, in the same editor session directly after training.
 	 */
-	UPROPERTY(EditAnywhere, Category = "Live Settings", meta = (ClampMin = "0.001", ClampMax = "1.0", ForceUnits="cm"))
-	float MorphTargetDeltaThreshold = 0.01f;
+	UPROPERTY(EditAnywhere, Category = "Live Settings", meta = (ClampMin = "0"))
+	int32 MorphTargetNumber = 0;
 
 	/**
 	 * Specify whether we want to debug draw the morph targets.
@@ -51,5 +56,4 @@ protected:
 	 */
 	UPROPERTY(EditAnywhere, Category = "Live Settings")
 	bool bDrawMorphTargets = false;
-#endif // WITH_EDITORONLY_DATA
 };
diff --git a/Engine/Plugins/Animation/MLDeformer/MLDeformerFramework/Source/MLDeformerFramework/Public/MLDeformerPerfCounter.h b/Engine/Plugins/Animation/MLDeformer/MLDeformerFramework/Source/MLDeformerFramework/Public/MLDeformerPerfCounter.h
new file mode 100644
--- /dev/null
+++ b/Engine/Plugins/Animation/MLDeformer/MLDeformerFramework/Source/MLDeformerFramework/Public/MLDeformerPerfCounter.h
@@ -0,0 +1,38 @@
+// Copyright Epic Games, Inc. All Rights Reserved.
+
+#pragma once
+
+#include "CoreTypes.h"
+#include "Containers/Array.h"
+
+
+namespace UE::MLDeformer
+{
+	class MLDEFORMERFRAMEWORK_API FMLDeformerPerfCounter
+	{
+	public:
+		FMLDeformerPerfCounter();
+
+		// Main methods.
+		void SetHistorySize(int32 NumHistoryItems);
+		void Reset();
+		void BeginSample();
+		void EndSample();
+
+		// Get statistics.
+		int32 GetCycles() const;
+		int32 GetCyclesMin() const;
+		int32 GetCyclesMax() const;
+		int32 GetCyclesAverage() const;
+		int32 GetNumSamples() const;
+		int32 GetHistorySize() const;
+
+	private:
+		int32 StartTime = 0;
+		int32 NumSamples = 0;
+		int32 Cycles = 0;
+		int32 CyclesMin = 0;
+		int32 CyclesMax = 0;
+		TArray<int32> CycleHistory;
+	};
+}
diff --git a/Engine/Plugins/Animation/MLDeformer/MLDeformerFramework/Source/MLDeformerFramework/Public/MLDeformerVizSettings.h b/Engine/Plugins/Animation/MLDeformer/MLDeformerFramework/Source/MLDeformerFramework/Public/MLDeformerVizSettings.h
--- a/Engine/Plugins/Animation/MLDeformer/MLDeformerFramework/Source/MLDeformerFramework/Public/MLDeformerVizSettings.h
+++ b/Engine/Plugins/Animation/MLDeformer/MLDeformerFramework/Source/MLDeformerFramework/Public/MLDeformerVizSettings.h
@@ -58,6 +58,7 @@ public:
 	void SetTrainingFrameNumber(int32 FrameNumber)			{ TrainingFrameNumber = FrameNumber; }
 	void SetTestingFrameNumber(int32 FrameNumber)			{ TestingFrameNumber = FrameNumber; }
 	void SetWeight(float InWeight)							{ Weight = InWeight; }
+	void SetQualityLevel(int32 InQualityLevel)				{ QualityLevel = FMath::Max<int32>(InQualityLevel, 0); }
 
 	FVector GetMeshSpacingOffsetVector() const				{ return FVector(MeshSpacing, 0.0f, 0.0f); }
 	float GetMeshSpacing() const							{ return MeshSpacing; }
@@ -81,6 +82,7 @@ public:
 	float GetWeight() const									{ return Weight; }
 	bool GetXRayDeltas() const								{ return bXRayDeltas; }
 	bool GetDrawVertexDeltas() const						{ return bDrawDeltas; }
+	int32 GetQualityLevel() const							{ return QualityLevel; }
 
 	// Get property names.
 	static FName GetVisualizationModePropertyName()			{ return GET_MEMBER_NAME_CHECKED(UMLDeformerVizSettings, VisualizationMode); }
@@ -103,6 +105,7 @@ public:
 	static FName GetWeightPropertyName()					{ return GET_MEMBER_NAME_CHECKED(UMLDeformerVizSettings, Weight); }
 	static FName GetXRayDeltasPropertyName()				{ return GET_MEMBER_NAME_CHECKED(UMLDeformerVizSettings, bXRayDeltas); }
 	static FName GetDrawVertexDeltasPropertyName()			{ return GET_MEMBER_NAME_CHECKED(UMLDeformerVizSettings, bDrawDeltas); }
+	static FName GetQualityLevelPropertyName()				{ return GET_MEMBER_NAME_CHECKED(UMLDeformerVizSettings, QualityLevel); }
 #endif
 
 protected:
@@ -147,6 +150,14 @@ protected:
 	UPROPERTY(EditAnywhere, Category = "Live Settings", meta = (ClampMin = "0"))
 	int32 TestingFrameNumber = 0;
 
+	/*
+	 * The ML Deformer LOD value. This is a continuous value between 0 and 1, where 0 means maximum quality and 1 means the lowest quality.
+	 * Morph based models will disable certain morph targets when increasing this value, which will lead to lower quality, but better GPU performance.
+	 * Each model can decide how to use this value. Some models might not support this.
+	 */
+	UPROPERTY(EditAnywhere, Transient, Category = "Live Settings", meta = (ClampMin = "0"))
+	int32 QualityLevel = 0;
+
 	/** Specify whether the heatmap is enabled or not. */
 	UPROPERTY(EditAnywhere, Category = "Live Settings")
 	bool bShowHeatMap = false;
@@ -170,7 +181,7 @@ protected:
 	 * A value of 0 means that we look exactly the same as the ML Deformed output, while a value of 1 means
 	 * the output is exactly the same as the ground truth. This is the blend factor between those two.
 	 */
-	UPROPERTY(EditAnywhere, Category = "Live Settings", meta = (EditCondition = "HeatMapMode==EMLDeformerHeatMapMode::GroundTruth && bShowHeatMap", UIMin = "0.0", UIMax = "1.0", ClampMin = "0.0", ClampMax = "1.0"))
+	UPROPERTY(EditAnywhere, Transient, Category = "Live Settings", meta = (EditCondition = "HeatMapMode==EMLDeformerHeatMapMode::GroundTruth && bShowHeatMap", UIMin = "0.0", UIMax = "1.0", ClampMin = "0.0", ClampMax = "1.0"))
 	float GroundTruthLerp = 0.0f;
 
 	/** Specifies whether we draw the linear skinned model or not. */
@@ -186,7 +197,7 @@ protected:
 	bool bDrawGroundTruthActor = true;
 
 	/** The scale factor of the ML deformer deltas being applied on top of the linear skinned results. */
-	UPROPERTY(EditAnywhere, Category = "Live Settings", meta = (ClampMin = "0.0", ClampMax = "1.0"))
+	UPROPERTY(EditAnywhere, Transient, Category = "Live Settings", meta = (ClampMin = "0.0", ClampMax = "1.0"))
 	float Weight = 1.0f;
 
 	/** Specifiy whether we want to draw the vertex deltas or not. */
diff --git a/Engine/Plugins/Animation/MLDeformer/MLDeformerFramework/Source/MLDeformerFrameworkEditor/Private/MLDeformerEditorModel.cpp b/Engine/Plugins/Animation/MLDeformer/MLDeformerFramework/Source/MLDeformerFrameworkEditor/Private/MLDeformerEditorModel.cpp
--- a/Engine/Plugins/Animation/MLDeformer/MLDeformerFramework/Source/MLDeformerFrameworkEditor/Private/MLDeformerEditorModel.cpp
+++ b/Engine/Plugins/Animation/MLDeformer/MLDeformerFramework/Source/MLDeformerFrameworkEditor/Private/MLDeformerEditorModel.cpp
@@ -339,14 +339,15 @@ namespace UE::MLDeformer
 		UpdateLabels();
 		CheckTrainingDataFrameChanged();
 
-		// Update the ML Deformer component's weight.
+		// Update the ML Deformer component of the ML Deformed test actor.
 		FMLDeformerEditorActor* EditorActor = FindEditorActor(ActorID_Test_MLDeformed);
 		if (EditorActor)
 		{
 			UMLDeformerComponent* DeformerComponent = EditorActor->GetMLDeformerComponent();
 			if (DeformerComponent)
 			{		
-				DeformerComponent->SetWeight(Model->GetVizSettings()->GetWeight());
+				DeformerComponent->SetWeight(VizSettings->GetWeight());
+				DeformerComponent->SetQualityLevel(VizSettings->GetQualityLevel());
 			}
 		}
 	}
@@ -653,7 +654,6 @@ namespace UE::MLDeformer
 		{
 			SetResamplingInputOutputsNeeded(true);
 			UpdateEditorInputInfo();
-			GetEditor()->GetModelDetailsView()->ForceRefresh();
 		}
 		else if (Property->GetFName() == UMLDeformerVizSettings::GetAnimPlaySpeedPropertyName())
 		{
@@ -1487,6 +1487,8 @@ namespace UE::MLDeformer
 		Sampler->SetVertexDeltaSpace(EVertexDeltaSpace::PostSkinning);
 		SampleDeltas();
 		Model->InitGPUData();
+
+		UpdateMemoryUsage();
 	}
 
 	FMLDeformerEditorActor* FMLDeformerEditorModel::GetTimelineEditorActor() const
@@ -1579,10 +1581,16 @@ namespace UE::MLDeformer
 		return (Model->GetNeuralNetwork() != nullptr);
 	}
 
+	void FMLDeformerEditorModel::UpdateMemoryUsage()
+	{
+		Model->InvalidateMemUsage();
+	}
+
 	void FMLDeformerEditorModel::TriggerInputAssetChanged(bool bRefreshVizSettings)
 	{
 		OnInputAssetsChanged();
 		OnPostInputAssetChanged();
+		UpdateMemoryUsage();
 		GetEditor()->GetModelDetailsView()->ForceRefresh();
 		if (bRefreshVizSettings)
 		{
@@ -1696,8 +1704,12 @@ namespace UE::MLDeformer
 		const TArrayView<const uint32> IndexArray = ImportedModel->LODModels[LOD].IndexBuffer;
 		const TArrayView<const int32> VertexMap = ImportedModel->LODModels[LOD].MeshToImportVertexMap;
 		const TArrayView<const FVector3f> BaseVertexPositions = Sampler->GetUnskinnedVertexPositions();
+
 		TArray<FVector3f> BaseNormals;
-		CalcMeshNormals(BaseVertexPositions, IndexArray, VertexMap, BaseNormals);
+		if (bIncludeNormals)
+		{
+			CalcMeshNormals(BaseVertexPositions, IndexArray, VertexMap, BaseNormals);
+		}
 
 		const FColorVertexBuffer& ColorBuffer = RenderData->LODRenderData[LOD].StaticVertexBuffers.ColorVertexBuffer;
 
@@ -1706,7 +1718,10 @@ namespace UE::MLDeformer
 		TArray<FVector3f> DeltaNormals;
 		for (int32 MorphTargetIndex = 0; MorphTargetIndex < NumMorphTargets; ++MorphTargetIndex)
 		{
-			GenerateNormalsForMorphTarget(LOD, SkelMesh, MorphTargetIndex, Deltas, BaseVertexPositions, BaseNormals, DeltaNormals);
+			if (bIncludeNormals)
+			{
+				GenerateNormalsForMorphTarget(LOD, SkelMesh, MorphTargetIndex, Deltas, BaseVertexPositions, BaseNormals, DeltaNormals);
+			}
 
 			const FName MorphName = *FString::Printf(TEXT("%s%.3d"), *NamePrefix, MorphTargetIndex);
 			UMorphTarget* MorphTarget = NewObject<UMorphTarget>(SkelMesh, MorphName);
@@ -1750,7 +1765,7 @@ namespace UE::MLDeformer
 							case EMLDeformerMaskChannel::VertexColorBlue:	{ VertexWeight = VertexColor.B; break; }
 							case EMLDeformerMaskChannel::VertexColorAlpha:	{ VertexWeight = VertexColor.A; break; }
 							default: 
-								checkf(false, TEXT("Unexpected weight mask value."));
+								checkf(false, TEXT("Unexpected mask channel value."));
 								break;
 						};
 
@@ -1804,6 +1819,25 @@ namespace UE::MLDeformer
 			LOD,
 			MorphErrorTolerance
 		);
+
+		// Reinit the render resources.
+		if (OutMorphBuffers.IsMorphCPUDataValid() && OutMorphBuffers.GetNumMorphs() > 0 && OutMorphBuffers.GetNumBatches() > 0)
+		{
+			BeginInitResource(&OutMorphBuffers);
+		}
+
+		// Update the editor actor skel mesh components for all the ones that also have an ML Deformer on it.
+		for (FMLDeformerEditorActor* EditorActor : EditorActors)
+		{
+			const UMLDeformerComponent* MLDeformerComponent = EditorActor->GetMLDeformerComponent();
+			USkeletalMeshComponent* SkelMeshComponent = EditorActor->GetSkeletalMeshComponent();
+			if (SkelMeshComponent && MLDeformerComponent)
+			{
+				SkelMeshComponent->RefreshExternalMorphTargetWeights();
+			}
+		}
+
+		UpdateMemoryUsage();
 	}
 
 	void FMLDeformerEditorModel::DrawMorphTarget(FPrimitiveDrawInterface* PDI, const TArray<FVector3f>& MorphDeltas, float DeltaThreshold, int32 MorphTargetIndex, const FVector& DrawOffset)
@@ -1820,8 +1854,9 @@ namespace UE::MLDeformer
 			// Draw all deltas.			
 			const int32 NumMorphTargets = MorphDeltas.Num() / NumVerts;
 			const int32 FinalMorphTargetIndex = FMath::Clamp<int32>(MorphTargetIndex, 0, NumMorphTargets - 1);
-			const FLinearColor IncludedColor(1.0f, 0.0f, 1.0f);
-			const FLinearColor ExcludedColor(0.1f, 0.1f, 0.1f);
+
+			const FLinearColor IncludedColor = FMLDeformerEditorStyle::Get().GetColor("MLDeformer.Morphs.IncludedVertexColor");
+			const FLinearColor ExcludedColor = FMLDeformerEditorStyle::Get().GetColor("MLDeformer.Morphs.ExcludedVertexColor");
 			for (int32 VertexIndex = 0; VertexIndex < NumVerts; ++VertexIndex)
 			{
 				const FVector StartPoint = FVector(UnskinnedPositions[VertexIndex]) + DrawOffset;
@@ -2088,6 +2123,16 @@ namespace UE::MLDeformer
 		UE_LOG(LogMLDeformer, Warning, TEXT("Please override the FMLDeformerEditorModel::GetNumTrainingFrames() method inside your derived editor model class."));
 		return 0;
 	}
+
+	UMLDeformerComponent* FMLDeformerEditorModel::FindMLDeformerComponent(int32 ActorID) const
+	{
+		const FMLDeformerEditorActor* EditorActor = FindEditorActor(ActorID);
+		if (EditorActor)
+		{
+			return EditorActor->GetMLDeformerComponent();
+		}
+		return nullptr;
+	}
 }	// namespace UE::MLDeformer
 
 #undef LOCTEXT_NAMESPACE
diff --git a/Engine/Plugins/Animation/MLDeformer/MLDeformerFramework/Source/MLDeformerFrameworkEditor/Private/MLDeformerEditorModule.cpp b/Engine/Plugins/Animation/MLDeformer/MLDeformerFramework/Source/MLDeformerFrameworkEditor/Private/MLDeformerEditorModule.cpp
--- a/Engine/Plugins/Animation/MLDeformer/MLDeformerFramework/Source/MLDeformerFrameworkEditor/Private/MLDeformerEditorModule.cpp
+++ b/Engine/Plugins/Animation/MLDeformer/MLDeformerFramework/Source/MLDeformerFrameworkEditor/Private/MLDeformerEditorModule.cpp
@@ -7,6 +7,7 @@
 #include "MLDeformerEditorModel.h"
 #include "MLDeformerModel.h"
 #include "MLDeformerCurveReferenceCustomization.h"
+#include "MLDeformerMorphModelQualityLevelCustomization.h"
 #include "Modules/ModuleManager.h"
 #include "EditorModeRegistry.h"
 #include "PropertyEditorDelegates.h"
@@ -26,6 +27,7 @@ namespace UE::MLDeformer
 		// Register object detail customizations.
 		FPropertyEditorModule& PropertyModule = FModuleManager::GetModuleChecked<FPropertyEditorModule>("PropertyEditor");
 		PropertyModule.RegisterCustomPropertyTypeLayout("CurveReference", FOnGetPropertyTypeCustomizationInstance::CreateStatic(&FMLDeformerCurveReferenceCustomization::MakeInstance) );
+		PropertyModule.RegisterCustomPropertyTypeLayout("MLDeformerMorphModelQualityLevel", FOnGetPropertyTypeCustomizationInstance::CreateStatic(&FMLDeformerMorphModelQualityLevelCustomization::MakeInstance) );
 		PropertyModule.NotifyCustomizationModuleChanged();
 	}
 
@@ -46,6 +48,7 @@ namespace UE::MLDeformer
 		{
 			FPropertyEditorModule& PropertyModule = FModuleManager::GetModuleChecked<FPropertyEditorModule>(TEXT("PropertyEditor"));
 			PropertyModule.UnregisterCustomPropertyTypeLayout(TEXT("CurveReference"));
+			PropertyModule.UnregisterCustomPropertyTypeLayout(TEXT("MLDeformerMorphModelQualityLevel"));
 			PropertyModule.NotifyCustomizationModuleChanged();
 		}
 	}
diff --git a/Engine/Plugins/Animation/MLDeformer/MLDeformerFramework/Source/MLDeformerFrameworkEditor/Private/MLDeformerEditorStyle.cpp b/Engine/Plugins/Animation/MLDeformer/MLDeformerFramework/Source/MLDeformerFrameworkEditor/Private/MLDeformerEditorStyle.cpp
--- a/Engine/Plugins/Animation/MLDeformer/MLDeformerFramework/Source/MLDeformerFrameworkEditor/Private/MLDeformerEditorStyle.cpp
+++ b/Engine/Plugins/Animation/MLDeformer/MLDeformerFramework/Source/MLDeformerFrameworkEditor/Private/MLDeformerEditorStyle.cpp
@@ -48,6 +48,10 @@ namespace UE::MLDeformer
 		Set("MLDeformer.DebugVectors.Color", FLinearColor(1.0f, 0.0f, 0.0f));
 		Set("MLDeformer.DebugVectors.Color2", FLinearColor(1.0f, 1.0f, 0.0f));
 
+		Set("MLDeformer.Morphs.IncludedVertexColor", FLinearColor(1.0f, 0.0f, 1.0f));
+		Set("MLDeformer.Morphs.ExcludedVertexColor", FLinearColor(0.1f, 0.1f, 0.1f));
+
+		Set("MLDeformer.Statistics.Performance", FLinearColor(0.0f, 1.0f, 0.0f));
 		Set("MLDeformer.DefaultLabelScale", 0.65f);
 
 		FSlateStyleRegistry::RegisterSlateStyle(*this);
diff --git a/Engine/Plugins/Animation/MLDeformer/MLDeformerFramework/Source/MLDeformerFrameworkEditor/Private/MLDeformerModelDetails.cpp b/Engine/Plugins/Animation/MLDeformer/MLDeformerFramework/Source/MLDeformerFrameworkEditor/Private/MLDeformerModelDetails.cpp
--- a/Engine/Plugins/Animation/MLDeformer/MLDeformerFramework/Source/MLDeformerFrameworkEditor/Private/MLDeformerModelDetails.cpp
+++ b/Engine/Plugins/Animation/MLDeformer/MLDeformerFramework/Source/MLDeformerFrameworkEditor/Private/MLDeformerModelDetails.cpp
@@ -194,13 +194,6 @@ namespace UE::MLDeformer
 		// Bone include list group.
 		if (Model->DoesSupportBones())
 		{
-			const int NumBones = EditorModel->GetEditorInputInfo() ? EditorModel->GetEditorInputInfo()->GetNumBones() : 0;
-			const FText BonesText = FText::Format(FTextFormat(LOCTEXT("BonesGroupName", "Bones ({0})")), NumBones);;
-			FText AllBonesText;
-			if (Model->GetSkeletalMesh() && (Model->GetSkeletalMesh()->GetRefSkeleton().GetNum() == NumBones || NumBones == 0))
-			{
-				AllBonesText = LOCTEXT("BoneGroupValue", "All Bones Included");
-			}
 			IDetailGroup& BoneIncludeGroup = InputOutputCategoryBuilder->AddGroup("BoneIncludeGroup", FText(), false, false);
 			BoneIncludeGroup.HeaderRow()
 			.NameContent()
@@ -210,7 +203,14 @@ namespace UE::MLDeformer
 				.HAlign(HAlign_Right)
 				[
 					SNew(STextBlock)
-					.Text(BonesText)
+					.Text_Lambda
+					(
+						[this]()
+						{
+							const int NumBonesIncluded = EditorModel->GetEditorInputInfo() ? EditorModel->GetEditorInputInfo()->GetNumBones() : 0;
+							return FText::Format(FTextFormat(LOCTEXT("BonesGroupName", "Bones ({0})")), NumBonesIncluded);
+						}
+					)
 					.Font(IDetailLayoutBuilder::GetDetailFont())
 				]
 			]
@@ -221,7 +221,19 @@ namespace UE::MLDeformer
 				.HAlign(HAlign_Left)
 				[
 					SNew(STextBlock)
-					.Text(AllBonesText)
+					.Text_Lambda
+					(
+						[this]()
+						{
+							const int NumBonesIncluded = EditorModel->GetEditorInputInfo() ? EditorModel->GetEditorInputInfo()->GetNumBones() : 0;
+							FText AllBonesText;
+							if (Model->GetSkeletalMesh() && (Model->GetSkeletalMesh()->GetRefSkeleton().GetNum() == NumBonesIncluded || NumBonesIncluded == 0))
+							{
+								AllBonesText = LOCTEXT("BoneGroupValue", "All Bones Included");
+							}
+							return AllBonesText;
+						}
+					)
 					.Font(IDetailLayoutBuilder::GetDetailFontItalic())
 				]
 			];
@@ -249,18 +261,6 @@ namespace UE::MLDeformer
 		// Curve include list group.
 		if (Model->DoesSupportCurves())
 		{
-			const int NumCurves = EditorModel->GetEditorInputInfo() ? EditorModel->GetEditorInputInfo()->GetNumCurves() : 0;
-			const FText CurvesText = FText::Format(FTextFormat(LOCTEXT("CurvesGroupName", "Curves ({0})")), NumCurves);
-			FText AllCurvesText;
-			const int32 NumCurvesOnSkelMesh = EditorModel->GetNumCurvesOnSkeletalMesh(Model->GetSkeletalMesh());
-			if (NumCurvesOnSkelMesh == 0)
-			{
-				AllCurvesText = LOCTEXT("CurvesGroupValueNoCurves", "No Curves Found");
-			}
-			else if (NumCurves == NumCurvesOnSkelMesh)
-			{
-				AllCurvesText = LOCTEXT("CurvesGroupValue", "All Curves Included");
-			}
 			IDetailGroup& CurvesIncludeGroup = InputOutputCategoryBuilder->AddGroup("CurveIncludeGroup", LOCTEXT("CurveIncludeGroup", "Curves"), false, false);
 			CurvesIncludeGroup.HeaderRow()
 			.NameContent()
@@ -270,7 +270,14 @@ namespace UE::MLDeformer
 				.HAlign(HAlign_Right)
 				[
 					SNew(STextBlock)
-					.Text(CurvesText)
+					.Text_Lambda
+					(
+						[this]()
+						{
+							const int NumCurves = EditorModel->GetEditorInputInfo() ? EditorModel->GetEditorInputInfo()->GetNumCurves() : 0;
+							return FText::Format(FTextFormat(LOCTEXT("CurvesGroupName", "Curves ({0})")), NumCurves);
+						}
+					)
 					.Font(IDetailLayoutBuilder::GetDetailFont())
 				]
 			]
@@ -281,7 +288,23 @@ namespace UE::MLDeformer
 				.HAlign(HAlign_Left)
 				[
 					SNew(STextBlock)
-					.Text(AllCurvesText)
+					.Text_Lambda
+					(
+						[this]()
+						{
+							const int NumCurves = EditorModel->GetEditorInputInfo() ? EditorModel->GetEditorInputInfo()->GetNumCurves() : 0;
+							const int32 NumCurvesOnSkelMesh = EditorModel->GetNumCurvesOnSkeletalMesh(Model->GetSkeletalMesh());
+							if (NumCurvesOnSkelMesh == 0)
+							{
+								return LOCTEXT("CurvesGroupValueNoCurves", "No Curves Found");
+							}
+							else if (NumCurves == NumCurvesOnSkelMesh)
+							{
+								return LOCTEXT("CurvesGroupValue", "All Curves Included");
+							}
+							return FText();
+						}
+					)
 					.Font(IDetailLayoutBuilder::GetDetailFontItalic())
 				]
 			];
diff --git a/Engine/Plugins/Animation/MLDeformer/MLDeformerFramework/Source/MLDeformerFrameworkEditor/Private/MLDeformerMorphModelDetails.cpp b/Engine/Plugins/Animation/MLDeformer/MLDeformerFramework/Source/MLDeformerFrameworkEditor/Private/MLDeformerMorphModelDetails.cpp
--- a/Engine/Plugins/Animation/MLDeformer/MLDeformerFramework/Source/MLDeformerFrameworkEditor/Private/MLDeformerMorphModelDetails.cpp
+++ b/Engine/Plugins/Animation/MLDeformer/MLDeformerFramework/Source/MLDeformerFrameworkEditor/Private/MLDeformerMorphModelDetails.cpp
@@ -5,6 +5,8 @@
 #include "MLDeformerMorphModel.h"
 #include "DetailLayoutBuilder.h"
 #include "DetailCategoryBuilder.h"
+#include "DetailWidgetRow.h"
+#include "IDetailGroup.h"
 
 #define LOCTEXT_NAMESPACE "MLDeformerMorphModelDetails"
 
@@ -35,11 +37,16 @@ namespace UE::MLDeformer
 		// Create all the detail categories and add the properties of the base class.
 		FMLDeformerGeomCacheModelDetails::CustomizeDetails(DetailBuilder);
 
-		MorphTargetCategoryBuilder->AddProperty(UMLDeformerMorphModel::GetMorphTargetDeltaThresholdPropertyName(), UMLDeformerMorphModel::StaticClass());
-		MorphTargetCategoryBuilder->AddProperty(UMLDeformerMorphModel::GetMorphTargetErrorTolerancePropertyName(), UMLDeformerMorphModel::StaticClass());
-		MorphTargetCategoryBuilder->AddProperty(UMLDeformerMorphModel::GetMaskChannelPropertyName(), UMLDeformerMorphModel::StaticClass());
-		MorphTargetCategoryBuilder->AddProperty(UMLDeformerMorphModel::GetInvertMaskChannelPropertyName(), UMLDeformerMorphModel::StaticClass());
-		MorphTargetCategoryBuilder->AddProperty(UMLDeformerMorphModel::GetIncludeMorphTargetNormalsPropertyName(), UMLDeformerMorphModel::StaticClass());
+		IDetailGroup& CompressionGroup = MorphTargetCategoryBuilder->AddGroup("Compression", LOCTEXT("MorphCompressionGroupLabel", "Compression"), false, true);
+		CompressionGroup.AddPropertyRow(DetailLayoutBuilder->GetProperty(UMLDeformerMorphModel::GetMorphDeltaZeroThresholdPropertyName(), UMLDeformerMorphModel::StaticClass()));
+		CompressionGroup.AddPropertyRow(DetailLayoutBuilder->GetProperty(UMLDeformerMorphModel::GetMorphCompressionLevelPropertyName(), UMLDeformerMorphModel::StaticClass()));
+		CompressionGroup.AddPropertyRow(DetailLayoutBuilder->GetProperty(UMLDeformerMorphModel::GetIncludeMorphTargetNormalsPropertyName(), UMLDeformerMorphModel::StaticClass()));
+
+		IDetailGroup& MaskGroup = MorphTargetCategoryBuilder->AddGroup("Mask", LOCTEXT("MorphMaskGroupLabel", "Masking"), false, false);
+		MaskGroup.AddPropertyRow(DetailLayoutBuilder->GetProperty(UMLDeformerMorphModel::GetMaskChannelPropertyName(), UMLDeformerMorphModel::StaticClass()));
+		MaskGroup.AddPropertyRow(DetailLayoutBuilder->GetProperty(UMLDeformerMorphModel::GetInvertMaskChannelPropertyName(), UMLDeformerMorphModel::StaticClass()));
+
+		MorphTargetCategoryBuilder->AddProperty(DetailLayoutBuilder->GetProperty(UMLDeformerMorphModel::GetQualityLevelsPropertyName(), UMLDeformerMorphModel::StaticClass()));
 	}
 }	// namespace UE::MLDeformer
 
diff --git a/Engine/Plugins/Animation/MLDeformer/MLDeformerFramework/Source/MLDeformerFrameworkEditor/Private/MLDeformerMorphModelEditorModel.cpp b/Engine/Plugins/Animation/MLDeformer/MLDeformerFramework/Source/MLDeformerFrameworkEditor/Private/MLDeformerMorphModelEditorModel.cpp
--- a/Engine/Plugins/Animation/MLDeformer/MLDeformerFramework/Source/MLDeformerFrameworkEditor/Private/MLDeformerMorphModelEditorModel.cpp
+++ b/Engine/Plugins/Animation/MLDeformer/MLDeformerFramework/Source/MLDeformerFrameworkEditor/Private/MLDeformerMorphModelEditorModel.cpp
@@ -3,9 +3,12 @@
 #include "MLDeformerMorphModelEditorModel.h"
 #include "MLDeformerMorphModel.h"
 #include "MLDeformerMorphModelVizSettings.h"
+#include "MLDeformerMorphModelInstance.h"
+#include "MLDeformerMorphModelInstance.h"
 #include "MLDeformerGeomCacheActor.h"
 #include "MLDeformerEditorToolkit.h"
 #include "MLDeformerInputInfo.h"
+#include "MLDeformerComponent.h"
 #include "MLDeformerGeomCacheHelpers.h"
 #include "Components/ExternalMorphSet.h"
 #include "GeometryCache.h"
@@ -36,11 +39,39 @@ namespace UE::MLDeformer
 
 		FMLDeformerGeomCacheEditorModel::OnPropertyChanged(PropertyChangedEvent);
 
-		// Clamp the morph target number, if it changes.
-		if (Property->GetFName() == UMLDeformerMorphModelVizSettings::GetMorphTargetNumberPropertyName())
+		if (Property->GetFName() == UMLDeformerMorphModel::GetMorphDeltaZeroThresholdPropertyName() || 
+			Property->GetFName() == UMLDeformerMorphModel::GetMorphCompressionLevelPropertyName() ||
+			Property->GetFName() == UMLDeformerMorphModel::GetIncludeMorphTargetNormalsPropertyName() ||
+			Property->GetFName() == UMLDeformerMorphModel::GetMaskChannelPropertyName() ||
+			Property->GetFName() == UMLDeformerMorphModel::GetInvertMaskChannelPropertyName())
+		{
+			if (PropertyChangedEvent.ChangeType == EPropertyChangeType::ValueSet)
+			{
+				InitEngineMorphTargets(GetMorphModel()->GetMorphTargetDeltas());
+			}
+		}
+		else if (Property->GetFName() == UMLDeformerMorphModelVizSettings::GetMorphTargetNumberPropertyName())
 		{
 			ClampMorphTargetNumber();
 		}
+		else if (Property->GetFName() == UMLDeformerMorphModel::GetQualityLevelsPropertyName())
+		{
+			if (PropertyChangedEvent.ChangeType == EPropertyChangeType::ArrayAdd)
+			{
+				TArray<FMLDeformerMorphModelQualityLevel>& QualityLevels = GetMorphModel()->GetQualityLevelsArray();				
+				if (QualityLevels.Num() == 1)
+				{
+					const int32 NewQualityValue = FMath::Max(GetMorphModel()->GetNumMorphTargets() - 1, 1);
+					QualityLevels[QualityLevels.Num() - 1].SetMaxActiveMorphs(NewQualityValue);
+				}
+				else
+				{
+					int32 NewQuality = static_cast<int32>(QualityLevels[QualityLevels.Num() - 2].GetMaxActiveMorphs() * 0.75f);
+					NewQuality = FMath::Clamp<int32>(NewQuality, 1, GetMorphModel()->GetNumMorphTargets() - 1);
+					QualityLevels[QualityLevels.Num() - 1].SetMaxActiveMorphs(NewQuality);
+				}
+			}
+		}
 	}
 
 	void FMLDeformerMorphModelEditorModel::ClampMorphTargetNumber()
@@ -78,51 +109,157 @@ namespace UE::MLDeformer
 		// We aborted and don't want to use partially trained results, we should restore the deltas that we just overwrote after training.
 		if (TrainingResult == ETrainingResult::Aborted && !bUsePartiallyTrainedWhenAborted)
 		{
-			// Restore the blendshape backup.
+			// Restore the morph target vertex deltas backup.
 			GetMorphModel()->SetMorphTargetDeltas(MorphTargetDeltasBackup);
 		}
 		else if (TrainingResult == ETrainingResult::Success || (TrainingResult == ETrainingResult::Aborted && bUsePartiallyTrainedWhenAborted))
 		{
-			const UMLDeformerMorphModel* MorphModel = GetMorphModel();
-			if (!MorphModel->GetMorphTargetDeltas().IsEmpty())
-			{
-				// Set deltas with a length equal or below a given threshold to zero, for better compression.
-				TArray<FVector3f> MorphTargetDeltas = MorphModel->GetMorphTargetDeltas();
-				ZeroDeltasByThreshold(MorphTargetDeltas, MorphModel->GetMorphTargetDeltaThreshold());
-
-				// Build morph targets inside the engine, using the engine's compression scheme.
-				// Add one as we included the means now as extra morph target.
-				InitEngineMorphTargets(MorphTargetDeltas);
-			}
+			// Build morph targets inside the engine, using the engine's compression scheme.
+			// Add one as we included the means now as extra morph target.
+			InitEngineMorphTargets(GetMorphModel()->GetMorphTargetDeltas());
 		}
 
 		// This internally calls InitGPUData() which updates the GPU buffer with the deltas.
 		FMLDeformerGeomCacheEditorModel::OnPostTraining(TrainingResult, bUsePartiallyTrainedWhenAborted);
 	}
 
-	void FMLDeformerMorphModelEditorModel::InitEngineMorphTargets(const TArray<FVector3f>& Deltas)
+	float CalcStandardDeviation(TArrayView<const float> Values)
+	{
+		if (Values.IsEmpty())
+		{
+			return 0.0f;
+		}
+
+		// First calculate the mean.
+		float Mean = 0.0f;
+		for (float Value : Values)
+		{
+			Mean += Value;
+		}
+		Mean /= static_cast<float>(Values.Num());
+
+		// Now calculate the standard deviation.
+		float Sum = 0.0f;
+		for (float Value : Values)
+		{
+			Sum += FMath::Square(Value - Mean);
+		}
+		Sum /= static_cast<float>(Values.Num());
+
+		return FMath::Sqrt(Sum);
+	}
+
+	void FMLDeformerMorphModelEditorModel::UpdateMorphErrorValues(TArrayView<UMorphTarget*> MorphTargets)
 	{
+		UMLDeformerMorphModel* MorphModel = GetMorphModel();
+
+		// Check if we have max morph weight information.
+		// If we do not have this we can't really update the morph error values.
+		TArrayView<const float> MaxMorphWeights = MorphModel->GetMorphTargetMaxWeights();
+		if (MaxMorphWeights.IsEmpty())
+		{
+			return;
+		}
+
+		// Preallocate space for the standard deviation of each morph target.
+		TArray<float> ErrorValues;
+		ErrorValues.SetNumZeroed(MorphTargets.Num() - 1);
+
 		const int32 LOD = 0;
+		TArray<float> DeltaLengths;
+		for (int32 MorphIndex = 0; MorphIndex < MorphTargets.Num() - 1; ++MorphIndex)	// We have one extra morph for the means, skip that one.
+		{
+			const UMorphTarget* MorphTarget = MorphTargets[MorphIndex + 1];
+			const float MaxWeight = MaxMorphWeights[MorphIndex];
 
-		// Get the weight mask setting.
+			// Get the array of deltas.
+			int32 NumDeltas = 0;
+			const FMorphTargetDelta* Deltas = MorphTarget->GetMorphTargetDelta(LOD, NumDeltas);
+
+			// Build the array of position delta lengths.
+			DeltaLengths.Reset();
+			DeltaLengths.SetNumUninitialized(NumDeltas);
+			for (int32 DeltaIndex = 0; DeltaIndex < NumDeltas; ++DeltaIndex)
+			{
+				DeltaLengths[DeltaIndex] = Deltas[DeltaIndex].PositionDelta.Length() * MaxWeight;
+			}
+
+			// Now calculate the standard deviation of those lengths.
+			const float StandardDeviation = CalcStandardDeviation(DeltaLengths);
+			ErrorValues[MorphIndex] = StandardDeviation;
+		}
+
+		// Build a list of array indices, so we know the order in which things got sorted.
+		TArray<int32> SortedIndices;
+		SortedIndices.SetNumUninitialized(MorphTargets.Num() - 1);
+		for (int32 Index = 0; Index < SortedIndices.Num(); ++Index)
+		{
+			SortedIndices[Index] = Index;
+		}
+
+		// Now that we have a list of standard deviations, sort them.
+		SortedIndices.Sort
+		(
+			[&ErrorValues](const int32& IndexA, const int32& IndexB)
+			{
+				return ErrorValues[IndexA] > ErrorValues[IndexB];
+			}
+		);
+
+		// Update the morph model with the newly calculated error values.
+		MorphModel->SetMorphTargetsErrorOrder(SortedIndices, ErrorValues);
+	}
+
+	void FMLDeformerMorphModelEditorModel::InitEngineMorphTargets(const TArray<FVector3f>& Deltas)
+	{
 		UMLDeformerMorphModel* MorphModel = GetMorphModel();
-		const EMLDeformerMaskChannel MaskChannel = MorphModel->GetMaskChannel();
-		const bool bInvertMaskChannel = MorphModel->GetInvertMaskChannel();
+		if (MorphModel->GetMorphTargetDeltas().IsEmpty())
+		{
+			return;
+		}
+
+		// Zero out small deltas.
+		TArray<FVector3f> MorphTargetDeltas = Deltas;
+		ZeroDeltasByThreshold(MorphTargetDeltas, MorphModel->GetMorphDeltaZeroThreshold());
 
 		// Turn the delta buffer in a set of engine morph targets.
-		TArray<UMorphTarget*> MorphTargets;	// These will be garbage collected.
+		const int32 LOD = 0;
 		const bool bIncludeNormals = MorphModel->GetIncludeMorphTargetNormals();
-		CreateEngineMorphTargets(MorphTargets, Deltas, TEXT("MLDeformerMorph_"), LOD, MorphModel->GetMorphTargetDeltaThreshold(), bIncludeNormals, MaskChannel, bInvertMaskChannel);
+		const EMLDeformerMaskChannel MaskChannel = MorphModel->GetMaskChannel();
+		const bool bInvertMaskChannel = MorphModel->GetInvertMaskChannel();
+
+		// Create the engine morph targets.
+		TArray<UMorphTarget*> MorphTargets;
+		CreateEngineMorphTargets(
+			MorphTargets,
+			Deltas, 
+			TEXT("MLDeformerMorph_"), 
+			LOD,
+			MorphModel->GetMorphDeltaZeroThreshold(),
+			bIncludeNormals,
+			MaskChannel,
+			bInvertMaskChannel);
+
+		// Analyze the error values of the morph targets.
+		UpdateMorphErrorValues(MorphTargets);
 
 		// Now compress the morph targets to GPU friendly buffers.
 		check(MorphModel->GetMorphTargetSet().IsValid());
 		FMorphTargetVertexInfoBuffers& MorphBuffers = MorphModel->GetMorphTargetSet()->MorphBuffers;
-		CompressEngineMorphTargets(MorphBuffers, MorphTargets, LOD, MorphModel->GetMorphTargetErrorTolerance());
+		CompressEngineMorphTargets(MorphBuffers, MorphTargets, LOD, MorphModel->GetMorphCompressionLevel());
 
 		if (MorphBuffers.GetNumBatches() == 0 || MorphBuffers.GetNumMorphs() == 0)
 		{
 			MorphBuffers = FMorphTargetVertexInfoBuffers();
 		}
+
+		MorphModel->UpdateStatistics();
+
+		// Remove the morph targets again, as we don't need them anymore.
+		for (UMorphTarget* MorphTarget : MorphTargets)
+		{
+			MorphTarget->ConditionalBeginDestroy();
+		}
 	}
 
 	void FMLDeformerMorphModelEditorModel::Render(const FSceneView* View, FViewport* Viewport, FPrimitiveDrawInterface* PDI)
@@ -130,11 +267,12 @@ namespace UE::MLDeformer
 		FMLDeformerGeomCacheEditorModel::Render(View, Viewport, PDI);
 
 		// Debug draw the selected morph target.
+		UMLDeformerMorphModel* MorphModel = GetMorphModel();
 		const UMLDeformerMorphModelVizSettings* VizSettings = GetMorphModelVizSettings();
 		if (VizSettings->GetDrawMorphTargets() && VizSettings->GetVisualizationMode() == EMLDeformerVizMode::TestData)
 		{
 			const FVector DrawOffset = -VizSettings->GetMeshSpacingOffsetVector();
-			DrawMorphTarget(PDI, GetMorphModel()->GetMorphTargetDeltas(), VizSettings->GetMorphTargetDeltaThreshold(), VizSettings->GetMorphTargetNumber(), DrawOffset);
+			DrawMorphTarget(PDI, GetMorphModel()->GetMorphTargetDeltas(), MorphModel->GetMorphDeltaZeroThreshold(), VizSettings->GetMorphTargetNumber(), DrawOffset);
 		}
 	}
 }	// namespace UE::MLDeformer
diff --git a/Engine/Plugins/Animation/MLDeformer/MLDeformerFramework/Source/MLDeformerFrameworkEditor/Private/MLDeformerMorphModelQualityLevelCustomization.cpp b/Engine/Plugins/Animation/MLDeformer/MLDeformerFramework/Source/MLDeformerFrameworkEditor/Private/MLDeformerMorphModelQualityLevelCustomization.cpp
new file mode 100644
--- /dev/null
+++ b/Engine/Plugins/Animation/MLDeformer/MLDeformerFramework/Source/MLDeformerFrameworkEditor/Private/MLDeformerMorphModelQualityLevelCustomization.cpp
@@ -0,0 +1,67 @@
+// Copyright Epic Games, Inc. All Rights Reserved.
+
+#include "MLDeformerMorphModelQualityLevelCustomization.h"
+#include "MLDeformerMorphModelQualityLevel.h"
+
+#include "DetailLayoutBuilder.h"
+#include "Widgets/SBoxPanel.h"
+#include "Widgets/Text/STextBlock.h"
+
+#define LOCTEXT_NAMESPACE "MLDeformerMorphModelQualityLevelCustomization"
+
+namespace UE::MLDeformer
+{
+	TSharedRef<IPropertyTypeCustomization> FMLDeformerMorphModelQualityLevelCustomization::MakeInstance()
+	{
+		return MakeShareable(new FMLDeformerMorphModelQualityLevelCustomization());
+	}
+
+	void FMLDeformerMorphModelQualityLevelCustomization::CustomizeHeader(TSharedRef<IPropertyHandle> StructPropertyHandle, class FDetailWidgetRow& HeaderRow, IPropertyTypeCustomizationUtils& StructCustomizationUtils)
+	{
+		uint32 NumChildren;
+		StructPropertyHandle->GetNumChildren(NumChildren);
+		TMap<FName, TSharedPtr<IPropertyHandle>> PropertyHandles;
+		PropertyHandles.Reserve(NumChildren);
+		for (uint32 ChildIndex = 0; ChildIndex < NumChildren; ++ChildIndex)
+		{
+			TSharedRef<IPropertyHandle> ChildHandle = StructPropertyHandle->GetChildHandle(ChildIndex).ToSharedRef();
+			const FName PropertyName = ChildHandle->GetProperty()->GetFName();
+			PropertyHandles.Add(PropertyName, ChildHandle);
+		}
+
+		TSharedPtr<IPropertyHandle> MaxMorphsHandle = PropertyHandles.FindChecked(GET_MEMBER_NAME_CHECKED(FMLDeformerMorphModelQualityLevel, MaxActiveMorphs));
+		check(MaxMorphsHandle.IsValid());
+
+		const int32 QualityLevel = StructPropertyHandle->GetIndexInArray();
+		check(QualityLevel != INDEX_NONE);
+
+		HeaderRow
+		.NameContent()
+		[
+			SNew(STextBlock)
+			.Margin(FMargin(5, 0, 0, 0))
+			.Font(IDetailLayoutBuilder::GetDetailFont())
+			.Text(FText::Format(LOCTEXT("QualityLevelFormat", "Level {0}"), QualityLevel))
+		]
+		.ValueContent()
+		[
+			SNew(SHorizontalBox)
+			+SHorizontalBox::Slot()
+			[
+				MaxMorphsHandle->CreatePropertyValueWidget()
+			]
+			+SHorizontalBox::Slot()
+			.AutoWidth()			
+			.HAlign(EHorizontalAlignment::HAlign_Left)
+			.VAlign(EVerticalAlignment::VAlign_Center)
+			[
+				SNew(STextBlock)
+				.Margin(FMargin(5, 0, 0, 0))
+				.Font(IDetailLayoutBuilder::GetDetailFont())
+				.Text(LOCTEXT("ActiveMorphs", "Morphs"))
+			]
+		];
+	}
+}	// namespace UE::MLDeformer
+
+#undef LOCTEXT_NAMESPACE
diff --git a/Engine/Plugins/Animation/MLDeformer/MLDeformerFramework/Source/MLDeformerFrameworkEditor/Private/MLDeformerMorphModelVizSettingsDetails.cpp b/Engine/Plugins/Animation/MLDeformer/MLDeformerFramework/Source/MLDeformerFrameworkEditor/Private/MLDeformerMorphModelVizSettingsDetails.cpp
--- a/Engine/Plugins/Animation/MLDeformer/MLDeformerFramework/Source/MLDeformerFrameworkEditor/Private/MLDeformerMorphModelVizSettingsDetails.cpp
+++ b/Engine/Plugins/Animation/MLDeformer/MLDeformerFramework/Source/MLDeformerFrameworkEditor/Private/MLDeformerMorphModelVizSettingsDetails.cpp
@@ -4,9 +4,14 @@
 #include "MLDeformerMorphModel.h"
 #include "MLDeformerMorphModelEditorModel.h"
 #include "MLDeformerModule.h"
+#include "MLDeformerComponent.h"
+#include "MLDeformerMorphModelInstance.h"
 #include "MLDeformerEditorModule.h"
+#include "MLDeformerEditorActor.h"
 #include "MLDeformerGeomCacheHelpers.h"
 #include "MLDeformerMorphModelVizSettings.h"
+#include "Components/ExternalMorphSet.h"
+#include "Components/SkeletalMeshComponent.h"
 #include "DetailLayoutBuilder.h"
 #include "DetailCategoryBuilder.h"
 #include "DetailWidgetRow.h"
@@ -14,6 +19,7 @@
 #include "IDetailGroup.h"
 #include "Modules/ModuleManager.h"
 #include "Widgets/Layout/SBox.h"
+#include "Widgets/Text/STextBlock.h"
 #include "PropertyCustomizationHelpers.h"
 #include "DetailLayoutBuilder.h"
 #include "SWarningOrErrorBox.h"
@@ -39,6 +45,27 @@ namespace UE::MLDeformer
 		return MorphModelVizSettings->GetDrawMorphTargets() && !MorphModel->GetMorphTargetDeltas().IsEmpty();
 	}
 
+	int32 GetNumActiveMorphs(FMLDeformerMorphModelEditorModel* MorphEditorModel)
+	{
+		const UMLDeformerComponent* MLDeformerComponent = MorphEditorModel->FindMLDeformerComponent(ActorID_Test_MLDeformed);
+		if (MLDeformerComponent)
+		{
+			const UMLDeformerMorphModelInstance* MorphInstance = Cast<UMLDeformerMorphModelInstance>(MLDeformerComponent->GetModelInstance());
+			if (MorphInstance)
+			{
+				const int32 LOD = 0;
+				FExternalMorphSetWeights* WeightData = MorphInstance->FindWeightData(LOD);
+				if (WeightData)
+				{
+					WeightData->UpdateNumActiveMorphTargets();
+					return WeightData->NumActiveMorphTargets;
+				}
+			}
+		}
+
+		return 0;
+	}
+
 	void FMLDeformerMorphModelVizSettingsDetails::AddAdditionalSettings()
 	{
 		IDetailGroup& MorphsGroup = LiveSettingsCategory->AddGroup("Morph Targets", LOCTEXT("MorphTargetsLabel", "Morph Targets"), false, true);
@@ -47,9 +74,60 @@ namespace UE::MLDeformer
 
 		MorphsGroup.AddPropertyRow(DetailLayoutBuilder->GetProperty(UMLDeformerMorphModelVizSettings::GetMorphTargetNumberPropertyName(), UMLDeformerMorphModelVizSettings::StaticClass()))
 			.EditCondition(TAttribute<bool>(this, &FMLDeformerMorphModelVizSettingsDetails::IsMorphTargetsEnabled), nullptr);
+	
+		FMLDeformerMorphModelEditorModel* MorphEditorModel = static_cast<FMLDeformerMorphModelEditorModel*>(EditorModel);
+		MorphsGroup.AddWidgetRow()
+			.NameContent()
+			[			
+				SNew(STextBlock)
+				.Font(IDetailLayoutBuilder::GetDetailFont())
+				.Text(LOCTEXT("ActiveMorphs", "Num Active Morphs"))
+			]
+			.ValueContent()
+			[
+				SNew(STextBlock)
+				.Font(IDetailLayoutBuilder::GetDetailFont())
+				.Text_Lambda
+				(
+					[MorphEditorModel]()
+					{
+						const UMLDeformerMorphModel* MorphModelPtr = Cast<UMLDeformerMorphModel>(MorphEditorModel->GetModel());					
+						const int32 NumActiveMorphs = GetNumActiveMorphs(MorphEditorModel);
+						const int32 NumMorphs = MorphModelPtr->GetNumMorphTargets();
+						const float Percentage = (NumMorphs > 0) ? (NumActiveMorphs / static_cast<float>(NumMorphs)) * 100.0f : 100.0f;
 
-		MorphsGroup.AddPropertyRow(DetailLayoutBuilder->GetProperty(UMLDeformerMorphModelVizSettings::GetMorphTargetDeltaThresholdPropertyName(), UMLDeformerMorphModelVizSettings::StaticClass()))
-			.EditCondition(TAttribute<bool>(this, &FMLDeformerMorphModelVizSettingsDetails::IsMorphTargetsEnabled), nullptr);
+						FNumberFormattingOptions TwoDigitsFormat;
+						TwoDigitsFormat.SetMaximumFractionalDigits(0);
+						return FText::Format(LOCTEXT("ActiveMorphsFormat", "{0} / {1} ({2} %)"), NumActiveMorphs, NumMorphs, FText::AsNumber(Percentage, &TwoDigitsFormat));
+					}
+				)
+			];
+	}
+
+	void FMLDeformerMorphModelVizSettingsDetails::AddStatistics()
+	{
+		FMLDeformerGeomCacheVizSettingsDetails::AddStatistics();
+
+		StatsGPUMemUsageGroup->AddWidgetRow()
+			.NameContent()
+			[			
+				SNew(STextBlock)
+				.Text(LOCTEXT("MorphTargetsLabel", "Morph Targets"))
+				.Font(IDetailLayoutBuilder::GetDetailFont())
+			]
+			.ValueContent()
+			[
+				SNew(STextBlock)
+				.Text_Lambda
+				(
+					[this]()
+					{
+						const float CompressedMb = MorphModel->GetCompressedMorphDataSizeInBytes() / static_cast<float>(1024 * 1024);
+						return FText::Format(LOCTEXT("MorphTargetsValue", "{0} mb"), FText::AsNumber(CompressedMb, &MemUsageMetricFormat));
+					}
+				)
+				.Font(IDetailLayoutBuilder::GetDetailFont())
+			];
 	}
 }	//namespace UE::MLDeformer
 
diff --git a/Engine/Plugins/Animation/MLDeformer/MLDeformerFramework/Source/MLDeformerFrameworkEditor/Private/MLDeformerVizSettingsDetails.cpp b/Engine/Plugins/Animation/MLDeformer/MLDeformerFramework/Source/MLDeformerFrameworkEditor/Private/MLDeformerVizSettingsDetails.cpp
--- a/Engine/Plugins/Animation/MLDeformer/MLDeformerFramework/Source/MLDeformerFrameworkEditor/Private/MLDeformerVizSettingsDetails.cpp
+++ b/Engine/Plugins/Animation/MLDeformer/MLDeformerFramework/Source/MLDeformerFrameworkEditor/Private/MLDeformerVizSettingsDetails.cpp
@@ -7,6 +7,10 @@
 #include "MLDeformerEditorModel.h"
 #include "MLDeformerEditorModule.h"
 #include "MLDeformerVizSettings.h"
+#include "MLDeformerModelInstance.h"
+#include "MLDeformerComponent.h"
+#include "MLDeformerPerfCounter.h"
+#include "MLDeformerEditorStyle.h"
 #include "DetailLayoutBuilder.h"
 #include "DetailCategoryBuilder.h"
 #include "DetailWidgetRow.h"
@@ -27,6 +31,7 @@
 #include "Animation/AnimSequence.h"
 #include "Animation/MeshDeformer.h"
 #include "Engine/SkeletalMesh.h"
+#include "HAL/PlatformTime.h"
 
 #define LOCTEXT_NAMESPACE "MLDeformerVizSettingsDetails"
 
@@ -55,12 +60,17 @@ namespace UE::MLDeformer
 		TestAssetsCategory = &DetailLayoutBuilder->EditCategory("Test Assets", FText::GetEmpty(), ECategoryPriority::Important);
 		LiveSettingsCategory = &DetailLayoutBuilder->EditCategory("Live Settings", FText::GetEmpty(), ECategoryPriority::Important);
 		TrainingMeshesCategoryBuilder = &DetailLayoutBuilder->EditCategory("Training Meshes", FText::GetEmpty(), ECategoryPriority::Important);
+		StatsCategoryBuilder = &DetailLayoutBuilder->EditCategory("Statistics", FText::GetEmpty(), ECategoryPriority::Default);
 	}
 
 	void FMLDeformerVizSettingsDetails::CustomizeDetails(IDetailLayoutBuilder& DetailBuilder)
 	{
 		DetailLayoutBuilder = &DetailBuilder;
 
+		PerformanceMetricFormat.SetMaximumFractionalDigits(0);
+		PerformanceMetricFormat.SetMinimumFractionalDigits(0);
+		MemUsageMetricFormat.SetMaximumFractionalDigits(2);
+
 		// Try update the member model, editormodel and viz settings pointers.
 		TArray<TWeakObjectPtr<UObject>> Objects;
 		DetailLayoutBuilder->GetObjectsBeingCustomized(Objects);
@@ -153,6 +163,8 @@ namespace UE::MLDeformer
 		LiveSettingsCategory->AddProperty(UMLDeformerVizSettings::GetWeightPropertyName(), UMLDeformerVizSettings::StaticClass());
 		LiveSettingsCategory->AddProperty(UMLDeformerVizSettings::GetAnimPlaySpeedPropertyName(), UMLDeformerVizSettings::StaticClass());
 		LiveSettingsCategory->AddProperty(UMLDeformerVizSettings::GetTestingFrameNumberPropertyName(), UMLDeformerVizSettings::StaticClass());
+		LiveSettingsCategory->AddProperty(UMLDeformerVizSettings::GetQualityLevelPropertyName(), UMLDeformerVizSettings::StaticClass())
+			.Visibility(Model->DoesSupportQualityLevels() ? EVisibility::Visible : EVisibility::Collapsed);
 
 		IDetailGroup& HeatMapGroup = LiveSettingsCategory->AddGroup("HeatMap", LOCTEXT("HeatMap", "Heat Map"), false, true);
 		HeatMapGroup.AddPropertyRow(DetailBuilder.GetProperty(UMLDeformerVizSettings::GetShowHeatMapPropertyName(), UMLDeformerVizSettings::StaticClass()));
@@ -166,12 +178,14 @@ namespace UE::MLDeformer
 		VisGroup.AddPropertyRow(DetailBuilder.GetProperty(UMLDeformerVizSettings::GetDrawLinearSkinnedActorPropertyName(), UMLDeformerVizSettings::StaticClass()));
 		VisGroup.AddPropertyRow(DetailBuilder.GetProperty(UMLDeformerVizSettings::GetDrawMLDeformedActorPropertyName(), UMLDeformerVizSettings::StaticClass()));
 		VisGroup.AddPropertyRow(DetailBuilder.GetProperty(UMLDeformerVizSettings::GetDrawGroundTruthActorPropertyName(), UMLDeformerVizSettings::StaticClass()))
-			.EditCondition(VizSettings->HasTestGroundTruth(), nullptr);
+			.Visibility(VizSettings->HasTestGroundTruth() ? EVisibility::Visible : EVisibility::Collapsed);
 
 		TrainingMeshesCategoryBuilder->SetCategoryVisibility(bShowTrainingData);
 		TrainingMeshesCategoryBuilder->AddProperty(UMLDeformerVizSettings::GetTrainingFrameNumberPropertyName(), UMLDeformerVizSettings::StaticClass());
 		TrainingMeshesCategoryBuilder->AddProperty(UMLDeformerVizSettings::GetDrawVertexDeltasPropertyName(), UMLDeformerVizSettings::StaticClass());
 		TrainingMeshesCategoryBuilder->AddProperty(UMLDeformerVizSettings::GetXRayDeltasPropertyName(), UMLDeformerVizSettings::StaticClass());
+
+		AddStatistics();
 	}
 
 	bool FMLDeformerVizSettingsDetails::FilterAnimSequences(const FAssetData& AssetData, USkeleton* Skeleton)
@@ -214,6 +228,151 @@ namespace UE::MLDeformer
 
 		return false;
 	}
+
+	void FMLDeformerVizSettingsDetails::AddStatsPerfRow(IDetailGroup& Group, const FText& Label, const FMLDeformerEditorModel* InEditorModel, const FNumberFormattingOptions& Format, bool bHighlight, TFunctionRef<int32(const FMLDeformerEditorModel*)> GetCyclesFunction)
+	{
+		const FLinearColor HighlightColor = FMLDeformerEditorStyle::Get().GetColor("MLDeformer.Statistics.Performance");
+
+		Group.AddWidgetRow()
+			.NameContent()
+			[			
+				SNew(STextBlock)
+				.Font(IDetailLayoutBuilder::GetDetailFont())
+				.Text(Label)
+			]
+			.ValueContent()
+			[
+				SNew(STextBlock)
+				.Font(IDetailLayoutBuilder::GetDetailFont())
+				.ColorAndOpacity(bHighlight ? HighlightColor : FSlateColor::UseForeground())
+				.Text_Lambda
+				(
+					[EditorModel=InEditorModel, Format=Format, GetCyclesFunction]()
+					{
+						const int32 Cycles = GetCyclesFunction(EditorModel);
+						const float MilliSeconds = FPlatformTime::ToMilliseconds(Cycles);
+						return FText::Format(LOCTEXT("PerfTimeValue", "{0} \u00B5s"), FText::AsNumber(MilliSeconds * 1000.0f, &Format));
+					}
+				)
+			];
+	}
+
+	void FMLDeformerVizSettingsDetails::AddStatistics()
+	{
+		const FLinearColor MainColor = FMLDeformerEditorStyle::Get().GetColor("MLDeformer.Statistics.Performance");
+
+		// Create the groups
+		StatsPerformanceGroup = &StatsCategoryBuilder->AddGroup("CPU Performance", LOCTEXT("CPUPerformanceLabel", "CPU Performance"), false, true);
+		StatsMemUsageGroup = &StatsCategoryBuilder->AddGroup("Approximated Memory Usage", LOCTEXT("MemoryUsageLabel", "Memory Usage (Approximated)"), false, true);
+		StatsMainMemUsageGroup = &StatsMemUsageGroup->AddGroup("Main Memory Usage", FText(), false);
+		StatsMainMemUsageGroup->HeaderRow()
+			.NameContent()
+			[
+				SNew(STextBlock)
+				.Font(IDetailLayoutBuilder::GetDetailFont())
+				.Text(LOCTEXT("MainMemoryTotalMemUsageLabel", "Main Memory"))
+			]			
+			.ValueContent()
+			[
+				SNew(STextBlock)
+				.Font(IDetailLayoutBuilder::GetDetailFont())
+				.ColorAndOpacity(MainColor)
+				.Text_Lambda
+				(
+					[this]
+					{
+						const uint64 TotalBytes = Model->GetMemUsageInBytes(UE::MLDeformer::EMemUsageRequestFlags::Cooked);
+						return FText::Format(LOCTEXT("TotalMemUsageValue", "{0} mb"), FText::AsNumber(TotalBytes / static_cast<float>(1024*1024), &MemUsageMetricFormat));
+					}
+				)
+			];
+
+		StatsGPUMemUsageGroup = &StatsMemUsageGroup->AddGroup("GPU Memory Usage", FText(), false);
+		StatsGPUMemUsageGroup->HeaderRow()
+			.NameContent()
+			[
+				SNew(STextBlock)
+				.Font(IDetailLayoutBuilder::GetDetailFont())
+				.Text(LOCTEXT("GPUMemoryUsageLabel", "GPU Memory"))
+			]			
+			.ValueContent()
+			[
+				SNew(STextBlock)
+				.Font(IDetailLayoutBuilder::GetDetailFont())
+				.ColorAndOpacity(MainColor)
+				.Text_Lambda
+				(
+					[this]
+					{
+						const uint64 TotalBytes = Model->GetGPUMemUsageInBytes();
+						return FText::Format(LOCTEXT("TotalGPUMemUsageValue", "{0} mb"), FText::AsNumber(TotalBytes / static_cast<float>(1024*1024), &MemUsageMetricFormat));
+					}
+				)
+			];
+
+		AddStatsPerfRow(*StatsPerformanceGroup, LOCTEXT("AvgTickTimeLabel", "Avg Time"), EditorModel, PerformanceMetricFormat, true,
+			[](const FMLDeformerEditorModel* EditorModelPtr)
+			{
+				UMLDeformerComponent* MLDeformerComponent = EditorModelPtr->FindMLDeformerComponent(ActorID_Test_MLDeformed);
+				return MLDeformerComponent ? MLDeformerComponent->GetTickPerfCounter().GetCyclesAverage() : 0;
+			});
+
+		AddStatsPerfRow(*StatsPerformanceGroup, LOCTEXT("MinTickTimeLabel", "Min Time"), EditorModel, PerformanceMetricFormat, false,
+			[](const FMLDeformerEditorModel* EditorModelPtr)
+			{
+				UMLDeformerComponent* MLDeformerComponent = EditorModelPtr->FindMLDeformerComponent(ActorID_Test_MLDeformed);
+				return MLDeformerComponent ? MLDeformerComponent->GetTickPerfCounter().GetCyclesMin() : 0;
+			});
+
+		AddStatsPerfRow(*StatsPerformanceGroup, LOCTEXT("MaxTickTimeLabel", "Max Time"), EditorModel, PerformanceMetricFormat, false,
+			[](const FMLDeformerEditorModel* EditorModelPtr)
+			{
+				UMLDeformerComponent* MLDeformerComponent = EditorModelPtr->FindMLDeformerComponent(ActorID_Test_MLDeformed);
+				return MLDeformerComponent ? MLDeformerComponent->GetTickPerfCounter().GetCyclesMax() : 0;
+			});
+
+		StatsMainMemUsageGroup->AddWidgetRow()
+			.NameContent()
+			[			
+				SNew(STextBlock)
+				.Font(IDetailLayoutBuilder::GetDetailFont())
+				.Text(LOCTEXT("ModelCookedMemUsageLabel", "Model - Cooked"))
+			]
+			.ValueContent()
+			[
+				SNew(STextBlock)
+				.Font(IDetailLayoutBuilder::GetDetailFont())
+				.Text_Lambda
+				(
+					[this]
+					{
+						const uint64 TotalBytes = Model->GetMemUsageInBytes(UE::MLDeformer::EMemUsageRequestFlags::Cooked);
+						return FText::Format(LOCTEXT("TotalCookedModelMemUsageValue", "{0} mb"), FText::AsNumber(TotalBytes / static_cast<float>(1024*1024), &MemUsageMetricFormat));
+					}
+				)
+			];
+
+		StatsMainMemUsageGroup->AddWidgetRow()
+			.NameContent()
+			[			
+				SNew(STextBlock)
+				.Font(IDetailLayoutBuilder::GetDetailFont())
+				.Text(LOCTEXT("ModelInEditorMemUsageLabel", "Model - In Editor"))
+			]
+			.ValueContent()
+			[
+				SNew(STextBlock)
+				.Font(IDetailLayoutBuilder::GetDetailFont())
+				.Text_Lambda
+				(
+					[this]
+					{
+						const uint64 TotalBytes = Model->GetMemUsageInBytes(UE::MLDeformer::EMemUsageRequestFlags::Uncooked);
+						return FText::Format(LOCTEXT("TotalInEditorModelMemUsageValue", "{0} mb"), FText::AsNumber(TotalBytes / static_cast<float>(1024*1024), &MemUsageMetricFormat));
+					}
+				)
+			];
+	}
 }	// namespace UE::MLDeformer
 
 #undef LOCTEXT_NAMESPACE
diff --git a/Engine/Plugins/Animation/MLDeformer/MLDeformerFramework/Source/MLDeformerFrameworkEditor/Public/MLDeformerEditorModel.h b/Engine/Plugins/Animation/MLDeformer/MLDeformerFramework/Source/MLDeformerFrameworkEditor/Public/MLDeformerEditorModel.h
--- a/Engine/Plugins/Animation/MLDeformer/MLDeformerFramework/Source/MLDeformerFrameworkEditor/Public/MLDeformerEditorModel.h
+++ b/Engine/Plugins/Animation/MLDeformer/MLDeformerFramework/Source/MLDeformerFrameworkEditor/Public/MLDeformerEditorModel.h
@@ -573,6 +573,9 @@ namespace UE::MLDeformer
 		 */
 		virtual void GenerateNormalsForMorphTarget(int32 LOD, USkeletalMesh* SkelMesh, int32 MorphTargetIndex, TArrayView<const FVector3f> Deltas, TArrayView<const FVector3f> BaseVertexPositions, TArrayView<FVector3f> BaseNormals, TArray<FVector3f>& OutDeltaNormals);
 
+		/** Invalidate the memory usage, so it gets updated in the UI again. */
+		void UpdateMemoryUsage();	
+
 		/** Get the current view range. */
 		TRange<double> GetViewRange() const;
 
@@ -773,6 +776,13 @@ namespace UE::MLDeformer
 		 */
 		void DrawMorphTarget(FPrimitiveDrawInterface* PDI, const TArray<FVector3f>& MorphDeltas, float DeltaThreshold, int32 MorphTargetIndex, const FVector& DrawOffset);
 
+		/**
+		 * Find the ML Deformer component of a given editor actor.
+		 * @param ActorID The actor ID to get the component for.
+		 * @return A pointer to the ML Deformer component, if it exists on the actor, otherwise nullptr is returned.
+		 */
+		UMLDeformerComponent* FindMLDeformerComponent(int32 ActorID = ActorID_Test_MLDeformed) const;
+
 	protected:
 		/**
 		 * Executed when a property changes by a change in the UI.
diff --git a/Engine/Plugins/Animation/MLDeformer/MLDeformerFramework/Source/MLDeformerFrameworkEditor/Public/MLDeformerMorphModelEditorModel.h b/Engine/Plugins/Animation/MLDeformer/MLDeformerFramework/Source/MLDeformerFrameworkEditor/Public/MLDeformerMorphModelEditorModel.h
--- a/Engine/Plugins/Animation/MLDeformer/MLDeformerFramework/Source/MLDeformerFrameworkEditor/Public/MLDeformerMorphModelEditorModel.h
+++ b/Engine/Plugins/Animation/MLDeformer/MLDeformerFramework/Source/MLDeformerFrameworkEditor/Public/MLDeformerMorphModelEditorModel.h
@@ -54,6 +54,8 @@ namespace UE::MLDeformer
 		 */
 		void ClampMorphTargetNumber();
 
+		void UpdateMorphErrorValues(TArrayView<UMorphTarget*> MorphTargets);
+
 	protected:
 		/**
 		 * The entire set of morph target deltas, 3 per vertex, for each morph target, as one flattened buffer.
diff --git a/Engine/Plugins/Animation/MLDeformer/MLDeformerFramework/Source/MLDeformerFrameworkEditor/Public/MLDeformerMorphModelQualityLevelCustomization.h b/Engine/Plugins/Animation/MLDeformer/MLDeformerFramework/Source/MLDeformerFrameworkEditor/Public/MLDeformerMorphModelQualityLevelCustomization.h
new file mode 100644
--- /dev/null
+++ b/Engine/Plugins/Animation/MLDeformer/MLDeformerFramework/Source/MLDeformerFrameworkEditor/Public/MLDeformerMorphModelQualityLevelCustomization.h
@@ -0,0 +1,28 @@
+// Copyright Epic Games, Inc. All Rights Reserved.
+
+#pragma once
+
+#include "CoreTypes.h"
+#include "IPropertyTypeCustomization.h"
+#include "PropertyHandle.h"
+#include "DetailWidgetRow.h"
+
+
+namespace UE::MLDeformer
+{
+	/**
+	 * The ML Deformer morph model's quality level property detail customization.
+	 * @see FMLDeformerMorphModelQualityLevel
+	 */
+	class MLDEFORMERFRAMEWORKEDITOR_API FMLDeformerMorphModelQualityLevelCustomization
+		: public IPropertyTypeCustomization
+	{
+	public:
+		static TSharedRef<IPropertyTypeCustomization> MakeInstance();
+
+		// IPropertyTypeCustomization overrides.
+		virtual void CustomizeHeader(TSharedRef<IPropertyHandle> StructPropertyHandle, class FDetailWidgetRow& HeaderRow, IPropertyTypeCustomizationUtils& StructCustomizationUtils) override;
+		virtual void CustomizeChildren(TSharedRef<IPropertyHandle> StructPropertyHandle, class IDetailChildrenBuilder& ChildBuilder, IPropertyTypeCustomizationUtils& StructCustomizationUtils) override {}
+		// ~END IPropertyTypeCustomization overrides.
+	};
+}	// namespace UE::MLDeformer
diff --git a/Engine/Plugins/Animation/MLDeformer/MLDeformerFramework/Source/MLDeformerFrameworkEditor/Public/MLDeformerMorphModelVizSettingsDetails.h b/Engine/Plugins/Animation/MLDeformer/MLDeformerFramework/Source/MLDeformerFrameworkEditor/Public/MLDeformerMorphModelVizSettingsDetails.h
--- a/Engine/Plugins/Animation/MLDeformer/MLDeformerFramework/Source/MLDeformerFrameworkEditor/Public/MLDeformerMorphModelVizSettingsDetails.h
+++ b/Engine/Plugins/Animation/MLDeformer/MLDeformerFramework/Source/MLDeformerFrameworkEditor/Public/MLDeformerMorphModelVizSettingsDetails.h
@@ -20,6 +20,7 @@ namespace UE::MLDeformer
 		// FMLDeformerVizSettingsDetails overrides.
 		virtual bool UpdateMemberPointers(const TArray<TWeakObjectPtr<UObject>>& Objects) override;
 		virtual void AddAdditionalSettings() override;
+		virtual void AddStatistics() override;
 		// ~END FMLDeformerVizSettingsDetails overrides.
 
 	protected:
diff --git a/Engine/Plugins/Animation/MLDeformer/MLDeformerFramework/Source/MLDeformerFrameworkEditor/Public/MLDeformerVizSettingsDetails.h b/Engine/Plugins/Animation/MLDeformer/MLDeformerFramework/Source/MLDeformerFrameworkEditor/Public/MLDeformerVizSettingsDetails.h
--- a/Engine/Plugins/Animation/MLDeformer/MLDeformerFramework/Source/MLDeformerFrameworkEditor/Public/MLDeformerVizSettingsDetails.h
+++ b/Engine/Plugins/Animation/MLDeformer/MLDeformerFramework/Source/MLDeformerFrameworkEditor/Public/MLDeformerVizSettingsDetails.h
@@ -9,6 +9,7 @@
 
 class IDetailCategoryBuilder;
 class IDetailLayoutBuilder;
+class IDetailGroup;
 class USkeleton;
 class UMLDeformerModel;
 class UMLDeformerVizSettings;
@@ -63,6 +64,17 @@ namespace UE::MLDeformer
 		/** Add any additional settings. */
 		virtual void AddAdditionalSettings() {}
 
+		/** Add any statistics. */
+		virtual void AddStatistics();
+
+		void AddStatsPerfRow(
+			IDetailGroup& Group,
+			const FText& Label,
+			const FMLDeformerEditorModel* InEditorModel,
+			const FNumberFormattingOptions& Format,
+			bool bHighlight,
+			TFunctionRef<int32(const FMLDeformerEditorModel*)> GetCyclesFunction);
+
 	protected:
 		/** Associated detail layout builder. */
 		IDetailLayoutBuilder* DetailLayoutBuilder = nullptr;
@@ -87,5 +99,26 @@ namespace UE::MLDeformer
 
 		/** A pointer to the training mesh category. This is only visible in training mode. */
 		IDetailCategoryBuilder* TrainingMeshesCategoryBuilder = nullptr;
+
+		/** A pointer to the statistics category. */
+		IDetailCategoryBuilder* StatsCategoryBuilder = nullptr;
+
+		/** The performance group inside the statistics category. */
+		IDetailGroup* StatsPerformanceGroup = nullptr;
+
+		/** The memory usage group inside the statistics category. */
+		IDetailGroup* StatsMemUsageGroup = nullptr;
+
+		/** Main memory usage subgroup. */
+		IDetailGroup* StatsMainMemUsageGroup = nullptr;
+
+		/** GPU memory usage subgroup. */
+		IDetailGroup* StatsGPUMemUsageGroup = nullptr;
+
+		/** Number format used for performance numbers. */
+		FNumberFormattingOptions PerformanceMetricFormat;
+
+		/** Number format used for memory usage numbers. */
+		FNumberFormattingOptions MemUsageMetricFormat;
 	};
 }	// namespace UE::MLDeformer
diff --git a/Engine/Plugins/Animation/MLDeformer/MLDeformerFramework/Source/MLDeformerFrameworkEditor/Public/SMLDeformerCurvePickerWidget.h b/Engine/Plugins/Animation/MLDeformer/MLDeformerFramework/Source/MLDeformerFrameworkEditor/Public/SMLDeformerCurvePickerWidget.h
--- a/Engine/Plugins/Animation/MLDeformer/MLDeformerFramework/Source/MLDeformerFrameworkEditor/Public/SMLDeformerCurvePickerWidget.h
+++ b/Engine/Plugins/Animation/MLDeformer/MLDeformerFramework/Source/MLDeformerFrameworkEditor/Public/SMLDeformerCurvePickerWidget.h
@@ -7,7 +7,6 @@
 #include "SlateFwd.h"
 #include "Widgets/SCompoundWidget.h"
 #include "Widgets/DeclarativeSyntaxSupport.h"
-#include "Widgets/SCompoundWidget.h"
 #include "Widgets/Views/STableRow.h"
 
 class SComboButton;
diff --git a/Engine/Plugins/Animation/MLDeformer/NearestNeighborModel/Source/NearestNeighborModel/Public/NearestNeighborModel.h b/Engine/Plugins/Animation/MLDeformer/NearestNeighborModel/Source/NearestNeighborModel/Public/NearestNeighborModel.h
--- a/Engine/Plugins/Animation/MLDeformer/NearestNeighborModel/Source/NearestNeighborModel/Public/NearestNeighborModel.h
+++ b/Engine/Plugins/Animation/MLDeformer/NearestNeighborModel/Source/NearestNeighborModel/Public/NearestNeighborModel.h
@@ -157,6 +157,7 @@ public:
 	virtual void PostLoad() override;
 
 	// UMLDeformerModel overrides.
+	virtual bool DoesSupportQualityLevels() const override { return false; }
 	virtual UMLDeformerModelInstance* CreateModelInstance(UMLDeformerComponent* Component) override;
 	virtual UMLDeformerInputInfo* CreateInputInfo() override;
 	virtual FString GetDisplayName() const override { return "Nearest Neighbor Model"; }
diff --git a/Engine/Plugins/Animation/MLDeformer/NearestNeighborModel/Source/NearestNeighborModelEditor/Private/NearestNeighborEditorModel.cpp b/Engine/Plugins/Animation/MLDeformer/NearestNeighborModel/Source/NearestNeighborModelEditor/Private/NearestNeighborEditorModel.cpp
--- a/Engine/Plugins/Animation/MLDeformer/NearestNeighborModel/Source/NearestNeighborModelEditor/Private/NearestNeighborEditorModel.cpp
+++ b/Engine/Plugins/Animation/MLDeformer/NearestNeighborModel/Source/NearestNeighborModelEditor/Private/NearestNeighborEditorModel.cpp
@@ -74,7 +74,7 @@ namespace UE::NearestNeighborModel
 			GetEditor()->GetModelDetailsView()->ForceRefresh();
 		}
 
-		if (Property->GetFName() == UNearestNeighborModel::GetMorphTargetDeltaThresholdPropertyName() || Property->GetFName() == UNearestNeighborModel::GetMorphTargetErrorTolerancePropertyName())
+		if (Property->GetFName() == UNearestNeighborModel::GetMorphCompressionLevelPropertyName() || Property->GetFName() == UNearestNeighborModel::GetMorphDeltaZeroThresholdPropertyName())
 		{
 			GetNearestNeighborModel()->InvalidateMorphTargetData();
 			GetEditor()->GetModelDetailsView()->ForceRefresh();
@@ -751,14 +751,14 @@ namespace UE::NearestNeighborModel
 			Deltas, 
 			FString("NNMorphTarget_"),
 			LOD,
-			NearestNeighborModel->GetMorphTargetDeltaThreshold(),
+			NearestNeighborModel->GetMorphDeltaZeroThreshold(),
 			NearestNeighborModel->GetIncludeMorphTargetNormals(),
 			NearestNeighborModel->GetMaskChannel(),
 			NearestNeighborModel->GetInvertMaskChannel());
 
 		check(NearestNeighborModel->GetMorphTargetSet().IsValid());
 		FMorphTargetVertexInfoBuffers& MorphBuffers = NearestNeighborModel->GetMorphTargetSet()->MorphBuffers;
-		CompressEngineMorphTargets(MorphBuffers, MorphTargets, LOD, NearestNeighborModel->GetMorphTargetErrorTolerance());
+		CompressEngineMorphTargets(MorphBuffers, MorphTargets, LOD, NearestNeighborModel->GetMorphCompressionLevel());
 
 		if (MorphBuffers.GetNumBatches() <= 0)
 		{
@@ -767,6 +767,12 @@ namespace UE::NearestNeighborModel
 			NearestNeighborModel->ResetMorphBuffers();
 		}
 
+		// Remove the morph targets again, as we don't need them anymore.
+		for (UMorphTarget* MorphTarget : MorphTargets)
+		{
+			MorphTarget->ConditionalBeginDestroy();
+		}
+
 		NearestNeighborModel->SetMorphTargetDeltas(Deltas);
 		return Result;
 	}
diff --git a/Engine/Plugins/Animation/MLDeformer/NearestNeighborModel/Source/NearestNeighborModelEditor/Private/NearestNeighborModelDetails.cpp b/Engine/Plugins/Animation/MLDeformer/NearestNeighborModel/Source/NearestNeighborModelEditor/Private/NearestNeighborModelDetails.cpp
--- a/Engine/Plugins/Animation/MLDeformer/NearestNeighborModel/Source/NearestNeighborModelEditor/Private/NearestNeighborModelDetails.cpp
+++ b/Engine/Plugins/Animation/MLDeformer/NearestNeighborModel/Source/NearestNeighborModelEditor/Private/NearestNeighborModelDetails.cpp
@@ -226,9 +226,9 @@ namespace UE::NearestNeighborModel
 					return FReply::Handled();
 				})
 			];
-		
-		MorphTargetCategoryBuilder->AddProperty(UMLDeformerMorphModel::GetMorphTargetDeltaThresholdPropertyName(), UMLDeformerMorphModel::StaticClass());
-		MorphTargetCategoryBuilder->AddProperty(UMLDeformerMorphModel::GetMorphTargetErrorTolerancePropertyName(), UMLDeformerMorphModel::StaticClass());
+				
+		MorphTargetCategoryBuilder->AddProperty(UMLDeformerMorphModel::GetMorphDeltaZeroThresholdPropertyName(), UMLDeformerMorphModel::StaticClass());
+		MorphTargetCategoryBuilder->AddProperty(UMLDeformerMorphModel::GetMorphCompressionLevelPropertyName(), UMLDeformerMorphModel::StaticClass());
 		FText ButtonText = NearestNeighborModel->IsMorphTargetDataValid() ? LOCTEXT("Update", "Update") : LOCTEXT("Update *", "Update *");
 		MorphTargetCategoryBuilder->AddProperty(UNearestNeighborModel::GetMorphDataSizePropertyName());
 		MorphTargetCategoryBuilder->AddCustomRow(FText::FromString(""))
diff --git a/Engine/Plugins/Animation/MLDeformer/NeuralMorphModel/Source/NeuralMorphModel/Private/NeuralMorphModel.cpp b/Engine/Plugins/Animation/MLDeformer/NeuralMorphModel/Source/NeuralMorphModel/Private/NeuralMorphModel.cpp
--- a/Engine/Plugins/Animation/MLDeformer/NeuralMorphModel/Source/NeuralMorphModel/Private/NeuralMorphModel.cpp
+++ b/Engine/Plugins/Animation/MLDeformer/NeuralMorphModel/Source/NeuralMorphModel/Private/NeuralMorphModel.cpp
@@ -57,7 +57,7 @@ void UNeuralMorphModel::Serialize(FArchive& Archive)
 		// Show a warning when we are not using custom inference yet on this model.
 		if (NeuralMorphNetwork == nullptr && NNINeuralNetwork != nullptr)
 		{
-			UE_LOG(LogNeuralMorphModel, Warning, TEXT("Neural Morph Model in MLD asset '%s' should be retrained to get higher performance by taking advantage of custom inference."), *GetDeformerAsset()->GetName());
+			UE_LOG(LogNeuralMorphModel, Display, TEXT("Neural Morph Model in MLD asset '%s' should be retrained to get higher performance by taking advantage of custom inference."), *GetDeformerAsset()->GetName());
 		}
 
 		// If we have a custom inference network, make sure we don't save out the NNI network.
@@ -79,7 +79,7 @@ void UNeuralMorphModel::PostLoad()
 	// Show a warning when we are not using custom inference yet on this model.
 	if (NeuralMorphNetwork == nullptr && GetNeuralNetwork() != nullptr)
 	{
-		UE_LOG(LogNeuralMorphModel, Warning, TEXT("Neural Morph Model in MLD asset '%s' should be retrained to get higher performance by taking advantage of custom inference."), *GetDeformerAsset()->GetName());
+		UE_LOG(LogNeuralMorphModel, Display, TEXT("Neural Morph Model in MLD asset '%s' should be retrained to get higher performance by taking advantage of custom inference."), *GetDeformerAsset()->GetName());
 	}
 #endif
 }
diff --git a/Engine/Plugins/Animation/MLDeformer/NeuralMorphModel/Source/NeuralMorphModel/Private/NeuralMorphModelInstance.cpp b/Engine/Plugins/Animation/MLDeformer/NeuralMorphModel/Source/NeuralMorphModel/Private/NeuralMorphModelInstance.cpp
--- a/Engine/Plugins/Animation/MLDeformer/NeuralMorphModel/Source/NeuralMorphModel/Private/NeuralMorphModelInstance.cpp
+++ b/Engine/Plugins/Animation/MLDeformer/NeuralMorphModel/Source/NeuralMorphModel/Private/NeuralMorphModelInstance.cpp
@@ -3,6 +3,7 @@
 #include "NeuralMorphModelInstance.h"
 #include "NeuralMorphModel.h"
 #include "NeuralMorphNetwork.h"
+#include "MLDeformerComponent.h"
 #include "MLDeformerAsset.h"
 #include "Components/ExternalMorphSet.h"
 #include "Components/SkeletalMeshComponent.h"
@@ -102,13 +103,35 @@ void UNeuralMorphModelInstance::Execute(float ModelWeight)
 		if (NumMorphTargets == NumNetworkWeights + 1)
 		{
 			// Set the first morph target, which represents the means, to a weight of 1.0, as it always needs to be fully active.
-			WeightData->Weights[0] = 1.0f * ModelWeight;
+			WeightData->Weights[0] = ModelWeight;
 
 			// Update all generated morph target weights with the values calculated by our neural network.
-			for (int32 MorphIndex = 0; MorphIndex < NumNetworkWeights; ++MorphIndex)
+			const TArrayView<const float> ErrorValues = MorphModel->GetMorphTargetErrorValues();
+			const TArrayView<const int32> ErrorOrder = MorphModel->GetMorphTargetErrorOrder();
+			if (!ErrorValues.IsEmpty())
 			{
-				const float OutputTensorValue = NetworkOutputs[MorphIndex];
-				WeightData->Weights[MorphIndex + 1] = OutputTensorValue * ModelWeight;
+				const int32 QualityLevel = GetMLDeformerComponent()->GetQualityLevel();
+				const int32 NumActiveMorphs = MorphModel->GetNumActiveMorphs(QualityLevel);
+				for (int32 Index = 0; Index < NumActiveMorphs; ++Index)
+				{
+					const int32 MorphIndex = ErrorOrder[Index];
+					const float TargetWeight = NetworkOutputs[MorphIndex] * ModelWeight;
+					WeightData->Weights[MorphIndex + 1] = FMath::Lerp<float, float>(StartMorphWeights[MorphIndex + 1], TargetWeight, MorphLerpAlpha);
+				}
+
+				// Disable all inactive morphs.
+				for (int32 Index = NumActiveMorphs; Index < NumNetworkWeights; ++Index)
+				{
+					const int32 MorphIndex = ErrorOrder[Index];
+					WeightData->Weights[MorphIndex + 1] = FMath::Lerp<float, float>(StartMorphWeights[MorphIndex + 1], 0.0f, MorphLerpAlpha);
+				}
+			}
+			else
+			{
+				for (int32 MorphIndex = 0; MorphIndex < NumNetworkWeights; ++MorphIndex)
+				{
+					WeightData->Weights[MorphIndex + 1] = NetworkOutputs[MorphIndex] * ModelWeight;
+				}
 			}
 			return;
 		}
diff --git a/Engine/Plugins/Animation/MLDeformer/VertexDeltaModel/Source/VertexDeltaModel/Public/VertexDeltaModel.h b/Engine/Plugins/Animation/MLDeformer/VertexDeltaModel/Source/VertexDeltaModel/Public/VertexDeltaModel.h
--- a/Engine/Plugins/Animation/MLDeformer/VertexDeltaModel/Source/VertexDeltaModel/Public/VertexDeltaModel.h
+++ b/Engine/Plugins/Animation/MLDeformer/VertexDeltaModel/Source/VertexDeltaModel/Public/VertexDeltaModel.h
@@ -28,7 +28,7 @@ public:
 	// UMLDeformerModel overrides.
 	virtual FString GetDisplayName() const override			{ return "Vertex Delta Model"; }
 	virtual bool IsNeuralNetworkOnGPU() const override		{ return true; }	// GPU neural network.
-	FString GetDefaultDeformerGraphAssetPath() const override;
+	virtual FString GetDefaultDeformerGraphAssetPath() const override;
 	// ~END UMLDeformerModel overrides.
 
 #if WITH_EDITORONLY_DATA
diff --git a/Engine/Source/Runtime/Engine/Public/Rendering/MorphTargetVertexInfoBuffers.h b/Engine/Source/Runtime/Engine/Public/Rendering/MorphTargetVertexInfoBuffers.h
--- a/Engine/Source/Runtime/Engine/Public/Rendering/MorphTargetVertexInfoBuffers.h
+++ b/Engine/Source/Runtime/Engine/Public/Rendering/MorphTargetVertexInfoBuffers.h
@@ -63,6 +63,11 @@ public:
 		return PositionPrecision;
 	}
 
+	uint64 GetMorphDataSizeInBytes() const 
+	{ 
+		return MorphData.Num() * sizeof(uint32);
+	}
+
 	static const float CalculatePositionPrecision(float TargetPositionErrorTolerance);
 
 	const float GetTangentZPrecision() const
