MLDeformer: Add MLDeformer.ForceWeight CVar to force the applied weight of all MLDeformer components.
Useful for debugging or making demo videos.

diff --git a/Engine/Plugins/Animation/MLDeformer/MLDeformerFramework/Source/MLDeformerFramework/Private/MLDeformerComponent.cpp b/Engine/Plugins/Animation/MLDeformer/MLDeformerFramework/Source/MLDeformerFramework/Private/MLDeformerComponent.cpp
--- a/Engine/Plugins/Animation/MLDeformer/MLDeformerFramework/Source/MLDeformerFramework/Private/MLDeformerComponent.cpp
+++ b/Engine/Plugins/Animation/MLDeformer/MLDeformerFramework/Source/MLDeformerFramework/Private/MLDeformerComponent.cpp
@@ -12,6 +12,16 @@
 
 #include UE_INLINE_GENERATED_CPP_BY_NAME(MLDeformerComponent)
 
+static float GMLDeformerOverrideWeight = -1;
+static FAutoConsoleVariableRef CVarMLDeformerOverrideWeight(
+	TEXT("MLDeformer.ForceWeight"),
+	GMLDeformerOverrideWeight,
+	TEXT("Force the Weight for MLDeformer components.")
+	TEXT("1 will force completely on, 0 will force completely off.")
+	TEXT("Negative values (the default) will use default weights."),
+	ECVF_Default
+);
+
 UMLDeformerComponent::UMLDeformerComponent(const FObjectInitializer& ObjectInitializer)
 	: Super(ObjectInitializer)
 {
@@ -230,7 +240,12 @@ void UMLDeformerComponent::TickComponent(float DeltaTime, enum ELevelTick TickTy
 			SkelMeshComponent->GetPredictedLODLevel() == 0)
 		{
 			TRACE_CPUPROFILER_EVENT_SCOPE(UMLDeformerComponent::TickComponent)
-			ModelInstance->Tick(DeltaTime, Weight);
+			float ApplyWeight = Weight;
+			if (GMLDeformerOverrideWeight >= 0.f)
+			{
+				ApplyWeight = FMath::Min(GMLDeformerOverrideWeight, 1.f);
+			}
+			ModelInstance->Tick(DeltaTime, ApplyWeight);
 
 			#if WITH_EDITOR
 				// Update our memory usage if desired.
