[MLDeformer]
- Fixed some bug where changing input bones or curves wouldn't regenerate the input data when training. This could lead to incorrect deformation.
- Fixed a bug where if it couldn't find a specific skelmesh submesh for a given geom cache track, it wouldn't allow you to train.
- Fixed an issue where you could run out of memory when the training data set was too large.

#trivial
#rnx

diff --git a/Engine/Plugins/Animation/MLDeformer/MLDeformerFramework/Source/MLDeformerFramework/Private/MLDeformerGeomCacheHelpers.cpp b/Engine/Plugins/Animation/MLDeformer/MLDeformerFramework/Source/MLDeformerFramework/Private/MLDeformerGeomCacheHelpers.cpp
--- a/Engine/Plugins/Animation/MLDeformer/MLDeformerFramework/Source/MLDeformerFramework/Private/MLDeformerGeomCacheHelpers.cpp
+++ b/Engine/Plugins/Animation/MLDeformer/MLDeformerFramework/Source/MLDeformerFramework/Private/MLDeformerGeomCacheHelpers.cpp
@@ -217,6 +217,7 @@ namespace UE::MLDeformer
 
 					// We found a match, no need to iterate over more Tracks.
 					bFoundMatch = true;
+					UE_LOG(LogMLDeformer, Display, TEXT("Mapped geom cache track '%s' to SkelMesh submesh '%s'"), *Track->GetName(), *SkelMeshName);
 					break;
 				} // If the track name matches the skeletal meshes internal mesh name.
 			} // For all meshes in the Skeletal Mesh.
diff --git a/Engine/Plugins/Animation/MLDeformer/MLDeformerFramework/Source/MLDeformerFrameworkEditor/Private/MLDeformerGeomCacheEditorModel.cpp b/Engine/Plugins/Animation/MLDeformer/MLDeformerFramework/Source/MLDeformerFrameworkEditor/Private/MLDeformerGeomCacheEditorModel.cpp
--- a/Engine/Plugins/Animation/MLDeformer/MLDeformerFramework/Source/MLDeformerFrameworkEditor/Private/MLDeformerGeomCacheEditorModel.cpp
+++ b/Engine/Plugins/Animation/MLDeformer/MLDeformerFramework/Source/MLDeformerFrameworkEditor/Private/MLDeformerGeomCacheEditorModel.cpp
@@ -149,15 +149,6 @@ namespace UE::MLDeformer
 			return;
 		}
 
-		// Allow the special case where there is just one mesh and track.
-		if (NumGeomCacheTracks != 1 || NumSkelMeshes != 1)
-		{
-			if (!GeomCacheSampler->GetFailedImportedMeshNames().IsEmpty())
-			{
-				return;
-			}
-		}
-
 		bIsReadyForTraining = true;
 	}
 
diff --git a/Engine/Plugins/Animation/MLDeformer/MLDeformerFramework/Source/MLDeformerFrameworkEditor/Private/MLDeformerModelDetails.cpp b/Engine/Plugins/Animation/MLDeformer/MLDeformerFramework/Source/MLDeformerFrameworkEditor/Private/MLDeformerModelDetails.cpp
--- a/Engine/Plugins/Animation/MLDeformer/MLDeformerFramework/Source/MLDeformerFrameworkEditor/Private/MLDeformerModelDetails.cpp
+++ b/Engine/Plugins/Animation/MLDeformer/MLDeformerFramework/Source/MLDeformerFrameworkEditor/Private/MLDeformerModelDetails.cpp
@@ -381,6 +381,7 @@ namespace UE::MLDeformer
 	FReply FMLDeformerModelDetails::OnFilterAnimatedBonesOnly() const
 	{
 		EditorModel->InitBoneIncludeListToAnimatedBonesOnly();
+		EditorModel->SetResamplingInputOutputsNeeded(true);
 		DetailLayoutBuilder->ForceRefreshDetails();
 		return FReply::Handled();
 	}
@@ -388,6 +389,7 @@ namespace UE::MLDeformer
 	FReply FMLDeformerModelDetails::OnFilterAnimatedCurvesOnly() const
 	{
 		EditorModel->InitCurveIncludeListToAnimatedCurvesOnly();
+		EditorModel->SetResamplingInputOutputsNeeded(true);
 		DetailLayoutBuilder->ForceRefreshDetails();
 		return FReply::Handled();
 	}
diff --git a/Engine/Plugins/Animation/MLDeformer/MLDeformerFramework/Source/MLDeformerFrameworkEditor/Private/MLDeformerMorphModelEditorModel.cpp b/Engine/Plugins/Animation/MLDeformer/MLDeformerFramework/Source/MLDeformerFrameworkEditor/Private/MLDeformerMorphModelEditorModel.cpp
--- a/Engine/Plugins/Animation/MLDeformer/MLDeformerFramework/Source/MLDeformerFrameworkEditor/Private/MLDeformerMorphModelEditorModel.cpp
+++ b/Engine/Plugins/Animation/MLDeformer/MLDeformerFramework/Source/MLDeformerFrameworkEditor/Private/MLDeformerMorphModelEditorModel.cpp
@@ -151,16 +151,17 @@ namespace UE::MLDeformer
 
 	void FMLDeformerMorphModelEditorModel::UpdateMorphErrorValues(TArrayView<UMorphTarget*> MorphTargets)
 	{
+		// If we don't have morph targets yet, there is nothing to do.
 		UMLDeformerMorphModel* MorphModel = GetMorphModel();
-
-		// Check if we have max morph weight information.
-		// If we do not have this we can't really update the morph error values.
-		TArrayView<const float> MaxMorphWeights = MorphModel->GetMorphTargetMaxWeights();
-		if (MaxMorphWeights.IsEmpty())
+		if (MorphModel->GetNumMorphTargets() == 0)
 		{
 			return;
 		}
 
+		// Check if we have max morph weight information.
+		// If we do not have this yet, we have to initialize the weights to 1.
+		TArrayView<const float> MaxMorphWeights = MorphModel->GetMorphTargetMaxWeights();
+
 		// Preallocate space for the standard deviation of each morph target.
 		TArray<float> ErrorValues;
 		ErrorValues.SetNumZeroed(MorphTargets.Num() - 1);
@@ -170,7 +171,7 @@ namespace UE::MLDeformer
 		for (int32 MorphIndex = 0; MorphIndex < MorphTargets.Num() - 1; ++MorphIndex)	// We have one extra morph for the means, skip that one.
 		{
 			const UMorphTarget* MorphTarget = MorphTargets[MorphIndex + 1];
-			const float MaxWeight = MaxMorphWeights[MorphIndex];
+			const float MaxWeight = !MaxMorphWeights.IsEmpty() ? MaxMorphWeights[MorphIndex] : 1.0f;
 
 			// Get the array of deltas.
 			int32 NumDeltas = 0;
