[UE] Changed behaviour of MLDeformerComponent::FindSkeletalMeshComponent
- Now returns nullptr instead of the first availible SkelMesh in the event that no matching SkelMesh is found

diff --git a/Engine/Plugins/Animation/MLDeformer/MLDeformerFramework/Source/MLDeformerFramework/Private/MLDeformerComponent.cpp b/Engine/Plugins/Animation/MLDeformer/MLDeformerFramework/Source/MLDeformerFramework/Private/MLDeformerComponent.cpp
--- a/Engine/Plugins/Animation/MLDeformer/MLDeformerFramework/Source/MLDeformerFramework/Private/MLDeformerComponent.cpp
+++ b/Engine/Plugins/Animation/MLDeformer/MLDeformerFramework/Source/MLDeformerFramework/Private/MLDeformerComponent.cpp
@@ -154,7 +154,7 @@ void UMLDeformerComponent::BeginDestroy()
 	Super::BeginDestroy();
 }
 
-USkeletalMeshComponent* UMLDeformerComponent::FindSkeletalMeshComponent(UMLDeformerAsset* Asset)
+USkeletalMeshComponent* UMLDeformerComponent::FindSkeletalMeshComponent(const UMLDeformerAsset* const Asset) const
 {
 	USkeletalMeshComponent* ResultingComponent = nullptr;
 	if (Asset != nullptr)
@@ -186,16 +186,6 @@ USkeletalMeshComponent* UMLDeformerComponent::FindSkeletalMeshComponent(UMLDefor
 		}
 	}
 
-	if (ResultingComponent == nullptr)
-	{
-		// Fall back to the first skeletal mesh component.
-		const AActor* Actor = Cast<AActor>(GetOuter());
-		if (Actor)
-		{
-			ResultingComponent = Actor->FindComponentByClass<USkeletalMeshComponent>();
-		}
-	}
-
 	return ResultingComponent;
 }
 
diff --git a/Engine/Plugins/Animation/MLDeformer/MLDeformerFramework/Source/MLDeformerFramework/Public/MLDeformerComponent.h b/Engine/Plugins/Animation/MLDeformer/MLDeformerFramework/Source/MLDeformerFramework/Public/MLDeformerComponent.h
--- a/Engine/Plugins/Animation/MLDeformer/MLDeformerFramework/Source/MLDeformerFramework/Public/MLDeformerComponent.h
+++ b/Engine/Plugins/Animation/MLDeformer/MLDeformerFramework/Source/MLDeformerFramework/Public/MLDeformerComponent.h
@@ -105,13 +105,12 @@ public:
 	/**
 	 * Find the skeletal mesh component to apply the deformer on.
 	 * This will return the skeletal mesh component (on this actor) which uses the same skeletal mesh as the passed in ML Deformer asset was trained on.
-	 * If there is no such skeletal mesh component then it will return the first skeletal mesh component it finds.
-	 * In case there is no skeletal mesh component on the actor, a nullptr is returned.
+	 * If there is no such skeletal mesh component then it will return a nullptr.
 	 * @param Asset The ML Deformer asset to search a component for.
 	 * @return The skeletal mesh component that would be the best.
 	 */
 	UFUNCTION(BlueprintCallable, Category = "MLDeformer")
-	USkeletalMeshComponent* FindSkeletalMeshComponent(UMLDeformerAsset* Asset);
+	USkeletalMeshComponent* FindSkeletalMeshComponent(const UMLDeformerAsset* const Asset) const;
 
 	/**
 	 * Get the ML Deformer model instance that this component currently uses.
