Slate: Initial phase of FVector2D deprecation

  - Expanded FDeprecateSlateVector2D to provide structs for parameters, return types and member variables that use FVector2f, while still allowing conversion to/from FVector2d with optional per-module deprecation mechanisms.
  - Many of the high-traffic SlateCore types like FSlateBrush, FGeometry and FSlateLayoutTransform have been converted to use these deprecation mechanisms.
  - Some legacy FGeometry::ToPaintGeometry and MakeChild overloads have been explicitly deprecated since they cause ambiguous overloads with FSlateLayoutTransform if it were to support implicit construction.
  - Deprecated ULocalPlayer::GetPixelBoundingBox and GetPixelPoint to prefer FVector2f for OptionalAllotedSize parameter since this parameter is expected to come from FGeometry
  - Exposed FVector2f members to blueprints

diff --git a/Engine/Plugins/2D/Paper2D/Source/Paper2DEditor/Private/SPaperEditorViewport.cpp b/Engine/Plugins/2D/Paper2D/Source/Paper2DEditor/Private/SPaperEditorViewport.cpp
--- a/Engine/Plugins/2D/Paper2D/Source/Paper2DEditor/Private/SPaperEditorViewport.cpp
+++ b/Engine/Plugins/2D/Paper2D/Source/Paper2DEditor/Private/SPaperEditorViewport.cpp
@@ -518,7 +518,7 @@ void SPaperEditorViewport::PaintSoftwareCursor(const FGeometry& AllottedGeometry
 		FSlateDrawElement::MakeBox(
 			OutDrawElements,
 			DrawLayerId,
-			AllottedGeometry.ToPaintGeometry( GraphCoordToPanelCoord(SoftwareCursorPosition) / AllottedGeometry.Scale - ( CursorSize / 2 ), CursorSize ),
+			AllottedGeometry.ToPaintGeometry( CursorSize, FSlateLayoutTransform(GraphCoordToPanelCoord(SoftwareCursorPosition) / AllottedGeometry.Scale - ( CursorSize * .5f )) ),
 			Brush);
 	}
 }
diff --git a/Engine/Plugins/AI/EnvironmentQueryEditor/Source/EnvironmentQueryEditor/Private/SEnvQueryLoadGraph.cpp b/Engine/Plugins/AI/EnvironmentQueryEditor/Source/EnvironmentQueryEditor/Private/SEnvQueryLoadGraph.cpp
--- a/Engine/Plugins/AI/EnvironmentQueryEditor/Source/EnvironmentQueryEditor/Private/SEnvQueryLoadGraph.cpp
+++ b/Engine/Plugins/AI/EnvironmentQueryEditor/Source/EnvironmentQueryEditor/Private/SEnvQueryLoadGraph.cpp
@@ -26,7 +26,7 @@ int32 SEnvQueryLoadGraph::OnPaint(const FPaintArgs& Args, const FGeometry& Allot
 	FSlateDrawElement::MakeBox(
 		OutDrawElements,
 		LayerId,
-		AllottedGeometry.ToPaintGeometry(FVector2D(0, 0), FVector2D(AllottedGeometry.GetLocalSize().X, AllottedGeometry.GetLocalSize().Y)),
+		AllottedGeometry.ToPaintGeometry(FVector2D(AllottedGeometry.GetLocalSize().X, AllottedGeometry.GetLocalSize().Y), FSlateLayoutTransform()),
 		TimelineAreaBrush,
 		DrawEffects,
 		TimelineAreaBrush->GetTint(InWidgetStyle) * InWidgetStyle.GetColorAndOpacityTint()
diff --git a/Engine/Plugins/AI/EnvironmentQueryEditor/Source/EnvironmentQueryEditor/Private/STestFunctionWidget.cpp b/Engine/Plugins/AI/EnvironmentQueryEditor/Source/EnvironmentQueryEditor/Private/STestFunctionWidget.cpp
--- a/Engine/Plugins/AI/EnvironmentQueryEditor/Source/EnvironmentQueryEditor/Private/STestFunctionWidget.cpp
+++ b/Engine/Plugins/AI/EnvironmentQueryEditor/Source/EnvironmentQueryEditor/Private/STestFunctionWidget.cpp
@@ -27,7 +27,7 @@ int32 STestFunctionWidget::OnPaint(const FPaintArgs& Args, const FGeometry& Allo
 	FSlateDrawElement::MakeBox(
 		OutDrawElements,
 		LayerId,
-		AllottedGeometry.ToPaintGeometry(FVector2D(0, 0), FVector2D(AllottedGeometry.Size.X, AllottedGeometry.Size.Y)),
+		AllottedGeometry.ToPaintGeometry(FVector2D(AllottedGeometry.Size.X, AllottedGeometry.Size.Y), FSlateLayoutTransform()),
 		TimelineAreaBrush,
 		DrawEffects,
 		TimelineAreaBrush->GetTint(InWidgetStyle) * InWidgetStyle.GetColorAndOpacityTint()
@@ -49,7 +49,7 @@ int32 STestFunctionWidget::OnPaint(const FPaintArgs& Args, const FGeometry& Allo
 		FSlateDrawElement::MakeBox(
 			OutDrawElements,
 			LayerId,
-			AllottedGeometry.ToPaintGeometry(FVector2D(0, 0), FVector2D(FMath::TruncToInt(PreviewData.FilterLow * AllottedGeometry.Size.X), AllottedGeometry.Size.Y)),
+			AllottedGeometry.ToPaintGeometry(FVector2D(FMath::TruncToInt(PreviewData.FilterLow * AllottedGeometry.Size.X), AllottedGeometry.Size.Y), FSlateLayoutTransform(FVector2D(0, 0))),
 			WhiteBrush,
 			DrawEffects,
 			WhiteBrush->GetTint(InWidgetStyle) * FLinearColor(1.0f, 0.0f, 0.0f, 0.4f)
@@ -63,7 +63,7 @@ int32 STestFunctionWidget::OnPaint(const FPaintArgs& Args, const FGeometry& Allo
 		FSlateDrawElement::MakeBox(
 			OutDrawElements,
 			LayerId,
-			AllottedGeometry.ToPaintGeometry(FVector2D(FMath::TruncToInt(PreviewData.FilterHigh * AllottedGeometry.Size.X), 0), FVector2D(AllottedGeometry.Size.X, AllottedGeometry.Size.Y)),
+			AllottedGeometry.ToPaintGeometry(FVector2D(AllottedGeometry.Size.X, AllottedGeometry.Size.Y), FSlateLayoutTransform(FVector2D(FMath::TruncToInt(PreviewData.FilterHigh * AllottedGeometry.Size.X), 0))),
 			WhiteBrush,
 			DrawEffects,
 			WhiteBrush->GetTint(InWidgetStyle) * FLinearColor(1.0f, 0.0f, 0.0f, 0.5f)
diff --git a/Engine/Plugins/Animation/GameplayInsights/Source/GameplayInsights/Private/SCurveTimelineView.cpp b/Engine/Plugins/Animation/GameplayInsights/Source/GameplayInsights/Private/SCurveTimelineView.cpp
--- a/Engine/Plugins/Animation/GameplayInsights/Source/GameplayInsights/Private/SCurveTimelineView.cpp
+++ b/Engine/Plugins/Animation/GameplayInsights/Source/GameplayInsights/Private/SCurveTimelineView.cpp
@@ -190,7 +190,7 @@ int32 SCurveTimelineView::PaintCurve(const FGeometry& AllottedGeometry, const FS
 					{
 						float X1 = RangeToScreen.InputToLocalX(PreviousTime);
 						float X2 = RangeToScreen.InputToLocalX(CurrentTime);
-						FPaintGeometry PaintGeo = AllottedGeometry.ToPaintGeometry(FVector2D(X1, Y2) , FVector2D(X2-X1, Y1-Y2), 1.0);
+						FPaintGeometry PaintGeo = AllottedGeometry.ToPaintGeometry(FVector2f(X2-X1, Y1-Y2), FSlateLayoutTransform(FVector2f(X1, Y2)));
 						FSlateDrawElement::MakeBox(OutDrawElements, LayerId, PaintGeo, FillBrush, ESlateDrawEffect::None, FillColor.Get());
 					}
 				}
diff --git a/Engine/Plugins/Animation/GameplayInsights/Source/GameplayInsights/Private/SEventTimelineView.cpp b/Engine/Plugins/Animation/GameplayInsights/Source/GameplayInsights/Private/SEventTimelineView.cpp
--- a/Engine/Plugins/Animation/GameplayInsights/Source/GameplayInsights/Private/SEventTimelineView.cpp
+++ b/Engine/Plugins/Animation/GameplayInsights/Source/GameplayInsights/Private/SEventTimelineView.cpp
@@ -76,10 +76,10 @@ int32 SEventTimelineView::PaintEvents(const FGeometry& AllottedGeometry, const F
 				PrevPointTime = Point.Time;
 				
 				FSlateDrawElement::MakeBox(OutDrawElements, LayerId,
-					AllottedGeometry.ToPaintGeometry(FVector2D(X, Y) , EventSize, 1.0), EventBrush, ESlateDrawEffect::None, Point.Color);
+					AllottedGeometry.ToPaintGeometry(EventSize, FSlateLayoutTransform(FVector2f(X, Y))), EventBrush, ESlateDrawEffect::None, Point.Color);
 				
 				FSlateDrawElement::MakeBox(OutDrawElements, LayerId,
-					AllottedGeometry.ToPaintGeometry(FVector2D(X, Y) , EventSize, 1.0), EventBorderBrush, ESlateDrawEffect::None, FLinearColor::Black); 
+					AllottedGeometry.ToPaintGeometry(EventSize, FSlateLayoutTransform(FVector2f(X, Y))), EventBorderBrush, ESlateDrawEffect::None, FLinearColor::Black); 
 			}
 		}
 		
@@ -107,21 +107,21 @@ int32 SEventTimelineView::PaintEvents(const FGeometry& AllottedGeometry, const F
 
 				// window bar
 				FSlateDrawElement::MakeBox(OutDrawElements, LayerId++,
-					AllottedGeometry.ToPaintGeometry(FVector2D(XStart,Y + 1),FVector2D(XEnd-XStart, EventSize.Y-2), 1), Brush, ESlateDrawEffect::None, Window.Color);
+					AllottedGeometry.ToPaintGeometry(FVector2D(XEnd-XStart, EventSize.Y-2), FSlateLayoutTransform(FVector2D(XStart,Y + 1))), Brush, ESlateDrawEffect::None, Window.Color);
 
 				// key diamond at start
 				FSlateDrawElement::MakeBox(OutDrawElements, LayerId,
-					AllottedGeometry.ToPaintGeometry(FVector2D(XStartDiamond, Y) , EventSize, 1.0), EventBrush, ESlateDrawEffect::None, Window.Color);
+					AllottedGeometry.ToPaintGeometry(EventSize, FSlateLayoutTransform(FVector2D(XStartDiamond, Y))), EventBrush, ESlateDrawEffect::None, Window.Color);
 				
 				FSlateDrawElement::MakeBox(OutDrawElements, LayerId,
-					AllottedGeometry.ToPaintGeometry(FVector2D(XStartDiamond, Y) , EventSize, 1.0), EventBorderBrush, ESlateDrawEffect::None, FLinearColor::Black); 
+					AllottedGeometry.ToPaintGeometry(EventSize, FSlateLayoutTransform(FVector2D(XStartDiamond, Y))), EventBorderBrush, ESlateDrawEffect::None, FLinearColor::Black); 
 
 				// key diamond at end
 				FSlateDrawElement::MakeBox(OutDrawElements, LayerId,
-					AllottedGeometry.ToPaintGeometry(FVector2D(XEndDiamond, Y) , EventSize, 1.0), EventBrush, ESlateDrawEffect::None, Window.Color);
+					AllottedGeometry.ToPaintGeometry(EventSize, FSlateLayoutTransform(FVector2D(XEndDiamond, Y))), EventBrush, ESlateDrawEffect::None, Window.Color);
 				
 				FSlateDrawElement::MakeBox(OutDrawElements, LayerId,
-					AllottedGeometry.ToPaintGeometry(FVector2D(XEndDiamond, Y) , EventSize, 1.0), EventBorderBrush, ESlateDrawEffect::None, FLinearColor::Black);
+					AllottedGeometry.ToPaintGeometry(EventSize, FSlateLayoutTransform(FVector2D(XEndDiamond, Y))), EventBorderBrush, ESlateDrawEffect::None, FLinearColor::Black);
 			}
 		}
 		
diff --git a/Engine/Plugins/Animation/MLDeformer/MLDeformerFramework/Source/MLDeformerFrameworkEditor/Private/SMLDeformerTimeline.cpp b/Engine/Plugins/Animation/MLDeformer/MLDeformerFramework/Source/MLDeformerFrameworkEditor/Private/SMLDeformerTimeline.cpp
--- a/Engine/Plugins/Animation/MLDeformer/MLDeformerFramework/Source/MLDeformerFrameworkEditor/Private/SMLDeformerTimeline.cpp
+++ b/Engine/Plugins/Animation/MLDeformer/MLDeformerFramework/Source/MLDeformerFrameworkEditor/Private/SMLDeformerTimeline.cpp
@@ -427,7 +427,7 @@ namespace UE::MLDeformer
 				FSlateDrawElement::MakeText(
 					OutDrawElements,
 					InArgs.StartLayer + 1,
-					InArgs.AllottedGeometry.ToPaintGeometry(TextOffset, InArgs.AllottedGeometry.Size),
+					InArgs.AllottedGeometry.ToPaintGeometry(InArgs.AllottedGeometry.Size, FSlateLayoutTransform(TextOffset)),
 					FrameString,
 					SmallLayoutFont,
 					InArgs.DrawEffects,
@@ -503,7 +503,7 @@ namespace UE::MLDeformer
 			const float HandleEnd = HandleStart + 13.0f;
 
 			const int32 ArrowLayer = LayerId + 2;
-			FPaintGeometry MyGeometry = AllottedGeometry.ToPaintGeometry(FVector2D(HandleStart, 0), FVector2D(HandleEnd - HandleStart, AllottedGeometry.Size.Y));
+			FPaintGeometry MyGeometry = AllottedGeometry.ToPaintGeometry(FVector2f(HandleEnd - HandleStart, AllottedGeometry.Size.Y), FSlateLayoutTransform(FVector2f(HandleStart, 0)));
 			FLinearColor ScrubColor = InWidgetStyle.GetColorAndOpacityTint();
 			{
 				ScrubColor.A = ScrubColor.A * 0.75f;
@@ -546,7 +546,7 @@ namespace UE::MLDeformer
 				FSlateDrawElement::MakeText(
 					OutDrawElements,
 					Args.StartLayer + 1,
-					Args.AllottedGeometry.ToPaintGeometry(TextOffset, TextSize),
+					Args.AllottedGeometry.ToPaintGeometry(TextSize, FSlateLayoutTransform(TextOffset)),
 					FrameString,
 					SmallLayoutFont,
 					Args.DrawEffects,
@@ -574,7 +574,7 @@ namespace UE::MLDeformer
 				FSlateDrawElement::MakeBox(
 					OutDrawElements,
 					LayerId + 1,
-					AllottedGeometry.ToPaintGeometry(FVector2D(SelectionRangeL, 0.f), FVector2D(SelectionRangeR - SelectionRangeL, AllottedGeometry.Size.Y)),
+					AllottedGeometry.ToPaintGeometry(FVector2f(SelectionRangeR - SelectionRangeL, AllottedGeometry.Size.Y), FSlateLayoutTransform(FVector2f(SelectionRangeL, 0.f))),
 					FAppStyle::GetBrush("WhiteBrush"),
 					ESlateDrawEffect::None,
 					DrawColor.CopyWithNewOpacity(Args.SolidFillOpacity));
@@ -583,7 +583,7 @@ namespace UE::MLDeformer
 			FSlateDrawElement::MakeBox(
 				OutDrawElements,
 				LayerId + 1,
-				AllottedGeometry.ToPaintGeometry(FVector2D(SelectionRangeL, 0.f), FVector2D(Args.BrushWidth, AllottedGeometry.Size.Y)),
+				AllottedGeometry.ToPaintGeometry(FVector2f(Args.BrushWidth, AllottedGeometry.Size.Y), FSlateLayoutTransform(FVector2f(SelectionRangeL, 0.f))),
 				Args.StartBrush,
 				ESlateDrawEffect::None,
 				DrawColor);
@@ -591,7 +591,7 @@ namespace UE::MLDeformer
 			FSlateDrawElement::MakeBox(
 				OutDrawElements,
 				LayerId + 1,
-				AllottedGeometry.ToPaintGeometry(FVector2D(SelectionRangeR - Args.BrushWidth, 0.f), FVector2D(Args.BrushWidth, AllottedGeometry.Size.Y)),
+				AllottedGeometry.ToPaintGeometry(FVector2f(Args.BrushWidth, AllottedGeometry.Size.Y), FSlateLayoutTransform(FVector2f(SelectionRangeR - Args.BrushWidth, 0.f))),
 				Args.EndBrush,
 				ESlateDrawEffect::None,
 				DrawColor);
@@ -616,7 +616,7 @@ namespace UE::MLDeformer
 		FSlateDrawElement::MakeBox(
 			OutDrawElements,
 			LayerId + 1,
-			AllottedGeometry.ToPaintGeometry(FVector2D(PlaybackRangeL, 0.0f), FVector2D(Args.BrushWidth, AllottedGeometry.Size.Y)),
+			AllottedGeometry.ToPaintGeometry(FVector2f(Args.BrushWidth, AllottedGeometry.Size.Y), FSlateLayoutTransform(FVector2f(PlaybackRangeL, 0.0f))),
 			Args.StartBrush,
 			ESlateDrawEffect::None,
 			FColor(32, 128, 32, OpacityBlend)	// 120, 75, 50 (HSV)
@@ -625,7 +625,7 @@ namespace UE::MLDeformer
 		FSlateDrawElement::MakeBox(
 			OutDrawElements,
 			LayerId + 1,
-			AllottedGeometry.ToPaintGeometry(FVector2D(PlaybackRangeR - Args.BrushWidth, 0.0f), FVector2D(Args.BrushWidth, AllottedGeometry.Size.Y)),
+			AllottedGeometry.ToPaintGeometry(FVector2f(Args.BrushWidth, AllottedGeometry.Size.Y), FSlateLayoutTransform(FVector2f(PlaybackRangeR - Args.BrushWidth, 0.0f))),
 			Args.EndBrush,
 			ESlateDrawEffect::None,
 			FColor(128, 32, 32, OpacityBlend)	// 0, 75, 50 (HSV)
@@ -635,7 +635,7 @@ namespace UE::MLDeformer
 		FSlateDrawElement::MakeBox(
 			OutDrawElements,
 			LayerId + 1,
-			AllottedGeometry.ToPaintGeometry(FVector2D(0.f, 0.f), FVector2D(PlaybackRangeL, AllottedGeometry.Size.Y)),
+			AllottedGeometry.ToPaintGeometry(FVector2f(PlaybackRangeL, AllottedGeometry.Size.Y), FSlateLayoutTransform(FVector2f(0.f, 0.f))),
 			FAppStyle::GetBrush("WhiteBrush"),
 			ESlateDrawEffect::None,
 			FLinearColor::Black.CopyWithNewOpacity(0.3f * OpacityBlend / 255.0f)
@@ -644,7 +644,7 @@ namespace UE::MLDeformer
 		FSlateDrawElement::MakeBox(
 			OutDrawElements,
 			LayerId + 1,
-			AllottedGeometry.ToPaintGeometry(FVector2D(PlaybackRangeR, 0.0f), FVector2D(AllottedGeometry.Size.X - PlaybackRangeR, AllottedGeometry.Size.Y)),
+			AllottedGeometry.ToPaintGeometry(FVector2f(AllottedGeometry.Size.X - PlaybackRangeR, AllottedGeometry.Size.Y), FSlateLayoutTransform(FVector2f(PlaybackRangeR, 0.0f))),
 			FAppStyle::GetBrush("WhiteBrush"),
 			ESlateDrawEffect::None,
 			FLinearColor::Black.CopyWithNewOpacity(0.3f * OpacityBlend / 255.0f)
@@ -953,7 +953,7 @@ namespace UE::MLDeformer
 			FSlateDrawElement::MakeLines(
 				OutDrawElements,
 				LayerId + 1,
-				AllottedGeometry.ToPaintGeometry(FVector2D(LinePos, 0.0f), FVector2D(1.0f, 1.0f)),
+				AllottedGeometry.ToPaintGeometry(FVector2f(1.0f, 1.0f), FSlateLayoutTransform(FVector2f(LinePos, 0.0f))),
 				LinePoints,
 				DrawEffects,
 				FLinearColor(1.0f, 1.0f, 1.0f, 0.5f),
diff --git a/Engine/Plugins/Editor/WaveformEditor/Source/WaveformEditorWidgets/Private/SWaveformEditorTimeRuler.cpp b/Engine/Plugins/Editor/WaveformEditor/Source/WaveformEditorWidgets/Private/SWaveformEditorTimeRuler.cpp
--- a/Engine/Plugins/Editor/WaveformEditor/Source/WaveformEditorWidgets/Private/SWaveformEditorTimeRuler.cpp
+++ b/Engine/Plugins/Editor/WaveformEditor/Source/WaveformEditorWidgets/Private/SWaveformEditorTimeRuler.cpp
@@ -129,7 +129,7 @@ void SWaveformEditorTimeRuler::DrawTickTimeString(float TickTimeSeconds, const d
 	FSlateDrawElement::MakeText(
 		OutDrawElements,
 		++LayerId,
-		AllottedGeometry.ToPaintGeometry(TextOffset, AllottedGeometry.Size),
+		AllottedGeometry.ToPaintGeometry(AllottedGeometry.Size, FSlateLayoutTransform(TextOffset)),
 		TimeString,
 		TicksTextFont,
 		ESlateDrawEffect::None,
@@ -240,7 +240,7 @@ void SWaveformEditorTimeRuler::DrawPlayheadHandle(const FGeometry& AllottedGeome
 	const float WindowWidth = AllottedGeometry.Size.X;
 	const float	HandleStart = TransportCoordinator->GetPlayheadPosition() * WindowWidth - HandleWidth / 2;
 	const float	HandleEnd = HandleStart + HandleWidth;
- 	FPaintGeometry HandleGeometry = AllottedGeometry.ToPaintGeometry(FVector2D(HandleStart, 0), FVector2D(HandleEnd - HandleStart, AllottedGeometry.Size.Y));
+ 	FPaintGeometry HandleGeometry = AllottedGeometry.ToPaintGeometry(FVector2f(HandleEnd - HandleStart, AllottedGeometry.Size.Y), FSlateLayoutTransform(FVector2f(HandleStart, 0)));
 
 	FSlateDrawElement::MakeBox(
 		OutDrawElements,
diff --git a/Engine/Plugins/Editor/WaveformEditor/Source/WaveformTransformationsWidgets/Private/WaveformTransformationDurationRenderer.h b/Engine/Plugins/Editor/WaveformEditor/Source/WaveformTransformationsWidgets/Private/WaveformTransformationDurationRenderer.h
--- a/Engine/Plugins/Editor/WaveformEditor/Source/WaveformTransformationsWidgets/Private/WaveformTransformationDurationRenderer.h
+++ b/Engine/Plugins/Editor/WaveformEditor/Source/WaveformTransformationsWidgets/Private/WaveformTransformationDurationRenderer.h
@@ -35,7 +35,7 @@ public:
 			FSlateDrawElement::MakeBox(
 				OutDrawElements,
 				++LayerId,
-				AllottedGeometry.ToPaintGeometry(FVector2D(0.f, 0.f), FVector2D(RightMarginX, AllottedGeometry.Size.Y)),
+				AllottedGeometry.ToPaintGeometry(FVector2D(RightMarginX, AllottedGeometry.Size.Y), FSlateLayoutTransform()),
 				FAppStyle::GetBrush(WaveformTransformationDurationHiglightParams::BackgroundBrushName),
 				ESlateDrawEffect::None,
 				WaveformTransformationDurationHiglightParams::BoxColor);
@@ -49,7 +49,7 @@ public:
 			FSlateDrawElement::MakeBox(
 				OutDrawElements,
 				++LayerId,
-				AllottedGeometry.ToPaintGeometry(FVector2D(LeftMarginX, 0.f), FVector2D(AllottedGeometry.Size.X, AllottedGeometry.Size.Y)),
+				AllottedGeometry.ToPaintGeometry(FVector2D(AllottedGeometry.Size.X, AllottedGeometry.Size.Y), FSlateLayoutTransform(FVector2D(LeftMarginX, 0.f))),
 				FAppStyle::GetBrush(WaveformTransformationDurationHiglightParams::BackgroundBrushName),
 				ESlateDrawEffect::None,
 				WaveformTransformationDurationHiglightParams::BoxColor);
diff --git a/Engine/Plugins/Enterprise/DataprepEditor/Source/DataprepEditor/Private/Widgets/DataprepGraph/SDataprepGraphTrackNode.cpp b/Engine/Plugins/Enterprise/DataprepEditor/Source/DataprepEditor/Private/Widgets/DataprepGraph/SDataprepGraphTrackNode.cpp
--- a/Engine/Plugins/Enterprise/DataprepEditor/Source/DataprepEditor/Private/Widgets/DataprepGraph/SDataprepGraphTrackNode.cpp
+++ b/Engine/Plugins/Enterprise/DataprepEditor/Source/DataprepEditor/Private/Widgets/DataprepGraph/SDataprepGraphTrackNode.cpp
@@ -130,7 +130,7 @@ public:
 			FSlateDrawElement::MakeBox(
 				OutDrawElements,
 				MaxLayerId,
-				AllottedGeometry.ToPaintGeometry( GetSoftwareCursorPosition() - ( Brush->ImageSize / 2 ) , Brush->ImageSize / AllottedGeometry.Scale ),
+				AllottedGeometry.ToPaintGeometry( Brush->ImageSize / AllottedGeometry.Scale, FSlateLayoutTransform(GetSoftwareCursorPosition() - ( Brush->ImageSize / 2 )) ),
 				Brush);
 #ifdef DEBUG_DRAG
 			FVector2D LocalCursorPosition = GetTickSpaceGeometry().AbsoluteToLocal(FSlateApplication::Get().GetCursorPos());
diff --git a/Engine/Plugins/Experimental/Animation/ContextualAnimation/Source/ContextualAnimationEditor/Private/ContextualAnimMovieSceneNotifyTrackEditor.cpp b/Engine/Plugins/Experimental/Animation/ContextualAnimation/Source/ContextualAnimationEditor/Private/ContextualAnimMovieSceneNotifyTrackEditor.cpp
--- a/Engine/Plugins/Experimental/Animation/ContextualAnimation/Source/ContextualAnimationEditor/Private/ContextualAnimMovieSceneNotifyTrackEditor.cpp
+++ b/Engine/Plugins/Experimental/Animation/ContextualAnimation/Source/ContextualAnimationEditor/Private/ContextualAnimMovieSceneNotifyTrackEditor.cpp
@@ -491,7 +491,7 @@ void FContextualAnimNotifySection::PaintNotifyName(FSequencerSectionPainter& Pai
 	FSlateDrawElement::MakeBox(
 		Painter.DrawElements,
 		LayerId + 1,
-		Painter.SectionGeometry.ToPaintGeometry(BoxOffset, BoxSize),
+		Painter.SectionGeometry.ToPaintGeometry(BoxSize, FSlateLayoutTransform(BoxOffset)),
 		FAppStyle::GetBrush("WhiteBrush"),
 		ESlateDrawEffect::None,
 		FLinearColor::Black.CopyWithNewOpacity(0.5f)
@@ -503,7 +503,7 @@ void FContextualAnimNotifySection::PaintNotifyName(FSequencerSectionPainter& Pai
 		FSlateDrawElement::MakeText(
 			Painter.DrawElements,
 			LayerId + 2,
-			Painter.SectionGeometry.ToPaintGeometry(BoxOffset + IconOffset, IconSize),
+			Painter.SectionGeometry.ToPaintGeometry(IconSize, FSlateLayoutTransform(BoxOffset + IconOffset)),
 			WarningString,
 			FontAwesomeFont,
 			Painter.bParentEnabled ? ESlateDrawEffect::None : ESlateDrawEffect::DisabledEffect,
@@ -514,7 +514,7 @@ void FContextualAnimNotifySection::PaintNotifyName(FSequencerSectionPainter& Pai
 	FSlateDrawElement::MakeText(
 		Painter.DrawElements,
 		LayerId + 2,
-		Painter.SectionGeometry.ToPaintGeometry(BoxOffset + TextOffset, TextSize),
+		Painter.SectionGeometry.ToPaintGeometry(TextSize, FSlateLayoutTransform(BoxOffset + TextOffset)),
 		InEventString,
 		SmallLayoutFont,
 		Painter.bParentEnabled ? ESlateDrawEffect::None : ESlateDrawEffect::DisabledEffect,
diff --git a/Engine/Plugins/Experimental/ChaosCaching/Source/ChaosCachingEditor/Private/Sequencer/ChaosCacheTrackEditor.cpp b/Engine/Plugins/Experimental/ChaosCaching/Source/ChaosCachingEditor/Private/Sequencer/ChaosCacheTrackEditor.cpp
--- a/Engine/Plugins/Experimental/ChaosCaching/Source/ChaosCachingEditor/Private/Sequencer/ChaosCacheTrackEditor.cpp
+++ b/Engine/Plugins/Experimental/ChaosCaching/Source/ChaosCachingEditor/Private/Sequencer/ChaosCacheTrackEditor.cpp
@@ -151,7 +151,7 @@ int32 FChaosCacheSection::OnPaintSection(FSequencerSectionPainter& Painter) cons
 			FSlateDrawElement::MakeBox(
 				Painter.DrawElements,
 				LayerId + 5,
-				Painter.SectionGeometry.ToPaintGeometry(TextOffset - BoxPadding, TextSize + 2.0f * BoxPadding),
+				Painter.SectionGeometry.ToPaintGeometry(TextSize + 2.0f * BoxPadding, FSlateLayoutTransform(TextOffset - BoxPadding)),
 				FAppStyle::GetBrush("WhiteBrush"),
 				ESlateDrawEffect::None,
 				FLinearColor::Black.CopyWithNewOpacity(0.5f)
@@ -160,7 +160,7 @@ int32 FChaosCacheSection::OnPaintSection(FSequencerSectionPainter& Painter) cons
 			FSlateDrawElement::MakeText(
 				Painter.DrawElements,
 				LayerId + 6,
-				Painter.SectionGeometry.ToPaintGeometry(TextOffset, TextSize),
+				Painter.SectionGeometry.ToPaintGeometry(TextSize, FSlateLayoutTransform(TextOffset)),
 				FrameString,
 				SmallLayoutFont,
 				DrawEffects,
diff --git a/Engine/Plugins/Experimental/GeometryCollectionPlugin/Source/GeometryCollectionSequencer/Private/GeometryCollectionTrackEditor.cpp b/Engine/Plugins/Experimental/GeometryCollectionPlugin/Source/GeometryCollectionSequencer/Private/GeometryCollectionTrackEditor.cpp
--- a/Engine/Plugins/Experimental/GeometryCollectionPlugin/Source/GeometryCollectionSequencer/Private/GeometryCollectionTrackEditor.cpp
+++ b/Engine/Plugins/Experimental/GeometryCollectionPlugin/Source/GeometryCollectionSequencer/Private/GeometryCollectionTrackEditor.cpp
@@ -178,7 +178,7 @@ int32 FGeometryCollectionTrackSection::OnPaintSection( FSequencerSectionPainter&
 			FSlateDrawElement::MakeBox(
 				Painter.DrawElements,
 				LayerId + 5,
-				Painter.SectionGeometry.ToPaintGeometry(TextOffset - BoxPadding, TextSize + 2.0f * BoxPadding),
+				Painter.SectionGeometry.ToPaintGeometry(TextSize + 2.0f * BoxPadding, FSlateLayoutTransform(TextOffset - BoxPadding)),
 				FAppStyle::GetBrush("WhiteBrush"),
 				ESlateDrawEffect::None,
 				FLinearColor::Black.CopyWithNewOpacity(0.5f)
@@ -187,7 +187,7 @@ int32 FGeometryCollectionTrackSection::OnPaintSection( FSequencerSectionPainter&
 			FSlateDrawElement::MakeText(
 				Painter.DrawElements,
 				LayerId + 6,
-				Painter.SectionGeometry.ToPaintGeometry(TextOffset, TextSize),
+				Painter.SectionGeometry.ToPaintGeometry(TextSize, FSlateLayoutTransform(TextOffset)),
 				FrameString,
 				SmallLayoutFont,
 				DrawEffects,
diff --git a/Engine/Plugins/Experimental/Mutable/Source/CustomizableObjectEditor/Private/MuCOE/PinViewer/SPinViewerDetailRowIndent.cpp b/Engine/Plugins/Experimental/Mutable/Source/CustomizableObjectEditor/Private/MuCOE/PinViewer/SPinViewerDetailRowIndent.cpp
--- a/Engine/Plugins/Experimental/Mutable/Source/CustomizableObjectEditor/Private/MuCOE/PinViewer/SPinViewerDetailRowIndent.cpp
+++ b/Engine/Plugins/Experimental/Mutable/Source/CustomizableObjectEditor/Private/MuCOE/PinViewer/SPinViewerDetailRowIndent.cpp
@@ -46,7 +46,7 @@ int32 SPinViewerDetailRowIndent::OnPaint(const FPaintArgs& Args, const FGeometry
 	FSlateDrawElement::MakeBox(
 		OutDrawElements,
 		LayerId,
-		AllottedGeometry.ToPaintGeometry(FVector2D(0, 0), FVector2D(16, AllottedGeometry.GetLocalSize().Y)),
+		AllottedGeometry.ToPaintGeometry(FVector2D(16, AllottedGeometry.GetLocalSize().Y), FSlateLayoutTransform()),
 		BackgroundBrush,
 		ESlateDrawEffect::None,
 		BackgroundColor.GetColor(InWidgetStyle)
@@ -55,7 +55,7 @@ int32 SPinViewerDetailRowIndent::OnPaint(const FPaintArgs& Args, const FGeometry
 	FSlateDrawElement::MakeBox(
 		OutDrawElements,
 		LayerId + 1,
-		AllottedGeometry.ToPaintGeometry(FVector2D(0, 0), FVector2D(16, AllottedGeometry.GetLocalSize().Y)),
+		AllottedGeometry.ToPaintGeometry(FVector2D(16, AllottedGeometry.GetLocalSize().Y), FSlateLayoutTransform()),
 		DropShadowBrush
 	);
 
diff --git a/Engine/Plugins/Experimental/WidgetEditorToolPalette/Source/WidgetEditorToolPalette/Private/DefaultTools/RectangleSelectTool.cpp b/Engine/Plugins/Experimental/WidgetEditorToolPalette/Source/WidgetEditorToolPalette/Private/DefaultTools/RectangleSelectTool.cpp
--- a/Engine/Plugins/Experimental/WidgetEditorToolPalette/Source/WidgetEditorToolPalette/Private/DefaultTools/RectangleSelectTool.cpp
+++ b/Engine/Plugins/Experimental/WidgetEditorToolPalette/Source/WidgetEditorToolPalette/Private/DefaultTools/RectangleSelectTool.cpp
@@ -177,7 +177,7 @@ int32 URectangleSelectTool::OnPaint(const FPaintArgs& Args, const FGeometry& All
 			FSlateDrawElement::MakeBox(
 				OutDrawElements,
 				LayerId,
-				AllottedGeometry.ToPaintGeometry(PinnedOwningWidget->GraphCoordToPanelCoord(Properties->Marquee.Rect.GetUpperLeft()), Properties->Marquee.Rect.GetSize() * PinnedOwningWidget->GetZoomAmount()),
+				AllottedGeometry.ToPaintGeometry(Properties->Marquee.Rect.GetSize() * PinnedOwningWidget->GetZoomAmount(), FSlateLayoutTransform(PinnedOwningWidget->GraphCoordToPanelCoord(Properties->Marquee.Rect.GetUpperLeft()))),
 				FAppStyle::GetBrush(TEXT("MarqueeSelection"))
 			);
 		}
diff --git a/Engine/Plugins/FX/Niagara/Source/NiagaraEditor/Private/Widgets/SNiagaraBakerTimelineWidget.cpp b/Engine/Plugins/FX/Niagara/Source/NiagaraEditor/Private/Widgets/SNiagaraBakerTimelineWidget.cpp
--- a/Engine/Plugins/FX/Niagara/Source/NiagaraEditor/Private/Widgets/SNiagaraBakerTimelineWidget.cpp
+++ b/Engine/Plugins/FX/Niagara/Source/NiagaraEditor/Private/Widgets/SNiagaraBakerTimelineWidget.cpp
@@ -59,7 +59,7 @@ int32 SNiagaraBakerTimelineWidget::OnPaint(const FPaintArgs& Args, const FGeomet
 				const FVector2D BoxLocation(UStep * float(i), 0.0f);
 				const FVector2D BoxSize(UStep, AllottedGeometry.Size.Y);
 				const FLinearColor& Tint = (i == OutputFrameIndices.FrameIndexA) ? CurrentTint : BoxTints[i & 1];
-				FSlateDrawElement::MakeBox(OutDrawElements, RetLayerId++, AllottedGeometry.ToPaintGeometry(BoxLocation, BoxSize), BoxBrush, ESlateDrawEffect::None, Tint);
+				FSlateDrawElement::MakeBox(OutDrawElements, RetLayerId++, AllottedGeometry.ToPaintGeometry(BoxSize, FSlateLayoutTransform(BoxLocation)), BoxBrush, ESlateDrawEffect::None, Tint);
 			}
 		}
 		else
diff --git a/Engine/Plugins/FX/Niagara/Source/NiagaraEditor/Private/Widgets/SNiagaraGraphNodeConvert.cpp b/Engine/Plugins/FX/Niagara/Source/NiagaraEditor/Private/Widgets/SNiagaraGraphNodeConvert.cpp
--- a/Engine/Plugins/FX/Niagara/Source/NiagaraEditor/Private/Widgets/SNiagaraGraphNodeConvert.cpp
+++ b/Engine/Plugins/FX/Niagara/Source/NiagaraEditor/Private/Widgets/SNiagaraGraphNodeConvert.cpp
@@ -262,7 +262,7 @@ int32 SNiagaraGraphNodeConvert::OnPaint(const FPaintArgs& Args, const FGeometry&
 		FVector2D LocalStart = AllottedGeometry.AbsoluteToLocal(AbsStart);
 		FVector2D StartDirection = FVector2D(DirectionOffset, 0);
 
-		FVector2D AbsEnd = ConvertNodeViewModel->GetDraggedSocketViewModel()->GetAbsoluteDragPosition() + Inverse(Args.GetWindowToDesktopTransform());
+		FVector2D AbsEnd = ConvertNodeViewModel->GetDraggedSocketViewModel()->GetAbsoluteDragPosition() + FVector2D(Inverse(Args.GetWindowToDesktopTransform()));
 		FVector2D LocalEnd = AllottedGeometry.AbsoluteToLocal(AbsEnd);
 		FVector2D EndDirection = FVector2D(DirectionOffset, 0);
 
diff --git a/Engine/Plugins/FX/Niagara/Source/NiagaraEditor/Private/Widgets/SNiagaraSimCacheViewTimeline.cpp b/Engine/Plugins/FX/Niagara/Source/NiagaraEditor/Private/Widgets/SNiagaraSimCacheViewTimeline.cpp
--- a/Engine/Plugins/FX/Niagara/Source/NiagaraEditor/Private/Widgets/SNiagaraSimCacheViewTimeline.cpp
+++ b/Engine/Plugins/FX/Niagara/Source/NiagaraEditor/Private/Widgets/SNiagaraSimCacheViewTimeline.cpp
@@ -44,7 +44,7 @@ int32 SNiagaraSimCacheViewTimeline::OnPaint(const FPaintArgs& Args, const FGeome
 				const FVector2D BoxLocation(UStep * float(i), 0.0f);
 				const FVector2D BoxSize(UStep, AllottedGeometry.Size.Y);
 				const FLinearColor& Tint = (i == CurrentFrame) ? CurrentTint : BoxTints[i & 1];
-				FSlateDrawElement::MakeBox(OutDrawElements, LayerId++, AllottedGeometry.ToPaintGeometry(BoxLocation, BoxSize), BoxBrush, ESlateDrawEffect::None, Tint);
+				FSlateDrawElement::MakeBox(OutDrawElements, LayerId++, AllottedGeometry.ToPaintGeometry(BoxSize, FSlateLayoutTransform(BoxLocation)), BoxBrush, ESlateDrawEffect::None, Tint);
 			}
 		}
 		else
diff --git a/Engine/Plugins/FX/Niagara/Source/NiagaraEditor/Private/Widgets/SVerticalResizeBox.cpp b/Engine/Plugins/FX/Niagara/Source/NiagaraEditor/Private/Widgets/SVerticalResizeBox.cpp
--- a/Engine/Plugins/FX/Niagara/Source/NiagaraEditor/Private/Widgets/SVerticalResizeBox.cpp
+++ b/Engine/Plugins/FX/Niagara/Source/NiagaraEditor/Private/Widgets/SVerticalResizeBox.cpp
@@ -98,7 +98,7 @@ int32 SVerticalResizeBox::OnPaint(const FPaintArgs& Args, const FGeometry& Allot
 	(
 		OutDrawElements,
 		HandleLayerId,
-		AllottedGeometry.ToPaintGeometry(HandleLocation, HandleSize),
+		AllottedGeometry.ToPaintGeometry(HandleSize, FSlateLayoutTransform(HandleLocation)),
 		&HandleBrush,
 		ESlateDrawEffect::None,
 		HandleBoxColor
diff --git a/Engine/Plugins/FX/Niagara/Source/NiagaraEditorWidgets/Private/Stack/SNiagaraStackIndent.cpp b/Engine/Plugins/FX/Niagara/Source/NiagaraEditorWidgets/Private/Stack/SNiagaraStackIndent.cpp
--- a/Engine/Plugins/FX/Niagara/Source/NiagaraEditorWidgets/Private/Stack/SNiagaraStackIndent.cpp
+++ b/Engine/Plugins/FX/Niagara/Source/NiagaraEditorWidgets/Private/Stack/SNiagaraStackIndent.cpp
@@ -41,7 +41,7 @@ int32 SNiagaraStackIndent::OnPaint(const FPaintArgs& Args, const FGeometry& Allo
 			FSlateDrawElement::MakeBox(
 				OutDrawElements,
 				LayerId,
-				AllottedGeometry.ToPaintGeometry(FVector2D((SingleIndentWidth * i) + PaintOffsetX, 0), FVector2D(SingleIndentWidth, AllottedGeometry.GetLocalSize().Y)),
+				AllottedGeometry.ToPaintGeometry(FVector2D(SingleIndentWidth, AllottedGeometry.GetLocalSize().Y), FSlateLayoutTransform(FVector2D((SingleIndentWidth * i) + PaintOffsetX, 0))),
 				BackgroundBrush,
 				ESlateDrawEffect::None,
 				BackgroundColor.GetColor(InWidgetStyle));
@@ -49,7 +49,7 @@ int32 SNiagaraStackIndent::OnPaint(const FPaintArgs& Args, const FGeometry& Allo
 			FSlateDrawElement::MakeBox(
 				OutDrawElements,
 				LayerId + 1,
-				AllottedGeometry.ToPaintGeometry(FVector2D((SingleIndentWidth * i) + PaintOffsetX, 0), FVector2D(SingleIndentWidth, AllottedGeometry.GetLocalSize().Y)),
+				AllottedGeometry.ToPaintGeometry(FVector2D(SingleIndentWidth, AllottedGeometry.GetLocalSize().Y), FSlateLayoutTransform(FVector2D((SingleIndentWidth * i) + PaintOffsetX, 0))),
 				DropShadowBrush);
 		}
 	}
@@ -58,7 +58,7 @@ int32 SNiagaraStackIndent::OnPaint(const FPaintArgs& Args, const FGeometry& Allo
 		FSlateDrawElement::MakeBox(
 			OutDrawElements,
 			LayerId,
-			AllottedGeometry.ToPaintGeometry(FVector2D((SingleIndentWidth * (IndentLevel - 1)), 0), FVector2D(SingleIndentWidth, AllottedGeometry.GetLocalSize().Y)),
+			AllottedGeometry.ToPaintGeometry(FVector2D(SingleIndentWidth, AllottedGeometry.GetLocalSize().Y), FSlateLayoutTransform(FVector2D((SingleIndentWidth * (IndentLevel - 1)), 0))),
 			DropShadowBrush);
 	}
 	return LayerId + 1;
diff --git a/Engine/Plugins/Media/MediaCompositing/Source/MediaCompositingEditor/Private/Sequencer/MediaThumbnailSection.cpp b/Engine/Plugins/Media/MediaCompositing/Source/MediaCompositingEditor/Private/Sequencer/MediaThumbnailSection.cpp
--- a/Engine/Plugins/Media/MediaCompositing/Source/MediaCompositingEditor/Private/Sequencer/MediaThumbnailSection.cpp
+++ b/Engine/Plugins/Media/MediaCompositing/Source/MediaCompositingEditor/Private/Sequencer/MediaThumbnailSection.cpp
@@ -306,7 +306,7 @@ void FMediaThumbnailSection::DrawLoopIndicators(FSequencerSectionPainter& InPain
 		FSlateDrawElement::MakeBox(
 			InPainter.DrawElements,
 			InPainter.LayerId++,
-			InPainter.SectionGeometry.ToPaintGeometry(FVector2D(DrawOffset, 0.0f), FVector2D(1.0f, SectionSize.Y)),
+			InPainter.SectionGeometry.ToPaintGeometry(FVector2D(1.0f, SectionSize.Y), FSlateLayoutTransform(FVector2D(DrawOffset, 0.0f))),
 			GenericBrush,
 			ESlateDrawEffect::None,
 			FLinearColor::Gray
@@ -348,7 +348,7 @@ void FMediaThumbnailSection::DrawSampleStates(FSequencerSectionPainter& InPainte
 			FSlateDrawElement::MakeBox(
 				InPainter.DrawElements,
 				InPainter.LayerId++,
-				InPainter.SectionGeometry.ToPaintGeometry(FVector2D(DrawOffset, SectionSize.Y - BarHeight - 1.0f), FVector2D(DrawSize, BarHeight)),
+				InPainter.SectionGeometry.ToPaintGeometry(FVector2D(DrawSize, BarHeight), FSlateLayoutTransform(FVector2D(DrawOffset, SectionSize.Y - BarHeight - 1.0f))),
 				GenericBrush,
 				ESlateDrawEffect::None,
 				Color
@@ -401,7 +401,7 @@ void FMediaThumbnailSection::DrawMediaInfo(FSequencerSectionPainter& InPainter,
 		FSlateDrawElement::MakeText(
 			InPainter.DrawElements,
 			InPainter.LayerId++,
-			InPainter.SectionGeometry.ToPaintGeometry(TextOffset, TextSize),
+			InPainter.SectionGeometry.ToPaintGeometry(TextSize, FSlateLayoutTransform(TextOffset)),
 			TileString,
 			SmallLayoutFont,
 			DrawEffects,
@@ -420,7 +420,7 @@ void FMediaThumbnailSection::DrawMediaInfo(FSequencerSectionPainter& InPainter,
 		FSlateDrawElement::MakeText(
 			InPainter.DrawElements,
 			InPainter.LayerId++,
-			InPainter.SectionGeometry.ToPaintGeometry(TextOffset, TextSize),
+			InPainter.SectionGeometry.ToPaintGeometry(TextSize, FSlateLayoutTransform(TextOffset)),
 			MipString,
 			SmallLayoutFont,
 			DrawEffects,
diff --git a/Engine/Plugins/Media/MediaPlayerEditor/Source/MediaPlayerEditor/Private/Widgets/SMediaPlayerEditorCache.cpp b/Engine/Plugins/Media/MediaPlayerEditor/Source/MediaPlayerEditor/Private/Widgets/SMediaPlayerEditorCache.cpp
--- a/Engine/Plugins/Media/MediaPlayerEditor/Source/MediaPlayerEditor/Private/Widgets/SMediaPlayerEditorCache.cpp
+++ b/Engine/Plugins/Media/MediaPlayerEditor/Source/MediaPlayerEditor/Private/Widgets/SMediaPlayerEditorCache.cpp
@@ -50,7 +50,7 @@ int32 SMediaPlayerEditorCache::OnPaint(const FPaintArgs& Args, const FGeometry&
 	FSlateDrawElement::MakeBox(
 		OutDrawElements,
 		LayerId + 1,
-		AllottedGeometry.ToPaintGeometry(FVector2D(0.0f, 0.0f), FVector2D(AllottedGeometry.Size.X, ProgressBarHeight.Get())),
+		AllottedGeometry.ToPaintGeometry(FVector2f(AllottedGeometry.Size.X, ProgressBarHeight.Get()), FSlateLayoutTransform()),
 		FCoreStyle::Get().GetBrush("GenericWhiteBox"),
 		ESlateDrawEffect::None,
 		InWidgetStyle.GetColorAndOpacityTint() * FLinearColor::Black
@@ -85,7 +85,7 @@ void SMediaPlayerEditorCache::DrawPlayerPosition(FTimespan Time, const FGeometry
 	FSlateDrawElement::MakeBox(
 		OutDrawElements,
 		LayerId + 1,
-		AllottedGeometry.ToPaintGeometry(FVector2D(DrawOffset, ProgressBarHeight.Get() + PositionMarkerMargin.Get()), FVector2D(PositionMarkerSize.Get(), MarkerSize)),
+		AllottedGeometry.ToPaintGeometry(FVector2f(PositionMarkerSize.Get(), MarkerSize), FSlateLayoutTransform(FVector2f(DrawOffset, ProgressBarHeight.Get() + PositionMarkerMargin.Get()))),
 		GenericBrush,
 		ESlateDrawEffect::None,
 		InWidgetStyle.GetColorAndOpacityTint() * Color
@@ -139,7 +139,7 @@ void SMediaPlayerEditorCache::DrawSampleStates(const TRangeSet<FTimespan>& Range
 		FSlateDrawElement::MakeBox(
 			OutDrawElements,
 			LayerId + 1,
-			AllottedGeometry.ToPaintGeometry(FVector2D(DrawOffset, YPos * BarHeight), FVector2D(DrawSize, YScale * BarHeight)),
+			AllottedGeometry.ToPaintGeometry(FVector2f(DrawSize, YScale * BarHeight), FSlateLayoutTransform(FVector2f(DrawOffset, YPos * BarHeight))),
 			GenericBrush,
 			ESlateDrawEffect::None,
 			InWidgetStyle.GetColorAndOpacityTint() * Color
diff --git a/Engine/Plugins/MovieScene/LevelSequenceEditor/Source/LevelSequenceEditor/Private/CinematicViewport/SCinematicTransportRange.cpp b/Engine/Plugins/MovieScene/LevelSequenceEditor/Source/LevelSequenceEditor/Private/CinematicViewport/SCinematicTransportRange.cpp
--- a/Engine/Plugins/MovieScene/LevelSequenceEditor/Source/LevelSequenceEditor/Private/CinematicViewport/SCinematicTransportRange.cpp
+++ b/Engine/Plugins/MovieScene/LevelSequenceEditor/Source/LevelSequenceEditor/Private/CinematicViewport/SCinematicTransportRange.cpp
@@ -208,7 +208,7 @@ int32 SCinematicTransportRange::OnPaint(const FPaintArgs& Args, const FGeometry&
 	FSlateDrawElement::MakeBox(
 		OutDrawElements,
 		LayerId,
-		AllottedGeometry.ToPaintGeometry( FVector2D(0.f, TrackOffsetY),  FVector2D(AllottedGeometry.GetLocalSize().X, TrackHeight)),
+		AllottedGeometry.ToPaintGeometry(  FVector2D(AllottedGeometry.GetLocalSize().X, TrackHeight), FSlateLayoutTransform(FVector2f(0.f, TrackOffsetY))),
 		FAppStyle::GetBrush("WhiteBrush"),
 		DrawEffects,
 		FLinearColor(DarkGray)
@@ -223,7 +223,7 @@ int32 SCinematicTransportRange::OnPaint(const FPaintArgs& Args, const FGeometry&
 	FSlateDrawElement::MakeBox(
 		OutDrawElements,
 		++LayerId,
-		AllottedGeometry.ToPaintGeometry( FVector2D(AllottedGeometry.GetLocalSize().X*PlaybackStartLerp, TrackOffsetY),  FVector2D(AllottedGeometry.GetLocalSize().X*(PlaybackEndLerp - PlaybackStartLerp), TrackHeight)),
+		AllottedGeometry.ToPaintGeometry(  FVector2f(AllottedGeometry.GetLocalSize().X*(PlaybackEndLerp - PlaybackStartLerp), TrackHeight), FSlateLayoutTransform(FVector2f(AllottedGeometry.GetLocalSize().X*PlaybackStartLerp, TrackOffsetY))),
 		FAppStyle::GetBrush("WhiteBrush"),
 		DrawEffects,
 		FLinearColor(MidGray)
@@ -239,7 +239,7 @@ int32 SCinematicTransportRange::OnPaint(const FPaintArgs& Args, const FGeometry&
 		FSlateDrawElement::MakeBox(
 			OutDrawElements,
 			++LayerId,
-			AllottedGeometry.ToPaintGeometry( FVector2D(AllottedGeometry.GetLocalSize().X*PlaybackStartLerp, TrackOffsetY), FVector2D(AllottedGeometry.GetLocalSize().X * (ClampedProgressLerp - PlaybackStartLerp), TrackHeight) ),
+			AllottedGeometry.ToPaintGeometry( FVector2f(AllottedGeometry.GetLocalSize().X * (ClampedProgressLerp - PlaybackStartLerp), TrackHeight) , FSlateLayoutTransform(FVector2f(AllottedGeometry.GetLocalSize().X*PlaybackStartLerp, TrackOffsetY))),
 			FAppStyle::GetBrush("WhiteBrush"),
 			DrawEffects,
 			FLinearColor(LightGray)
@@ -311,7 +311,7 @@ int32 SCinematicTransportRange::OnPaint(const FPaintArgs& Args, const FGeometry&
 		FSlateDrawElement::MakeBox(
 			OutDrawElements,
 			++LayerId,
-			AllottedGeometry.ToPaintGeometry(FVector2D(PositionX - FMath::CeilToFloat(BrushWidth/2), 0.f), FVector2D(BrushWidth, BrushHeight)),
+			AllottedGeometry.ToPaintGeometry(FVector2f(BrushWidth, BrushHeight), FSlateLayoutTransform(FVector2f(PositionX - FMath::CeilToFloat(BrushWidth/2), 0.f))),
 			FLevelSequenceEditorStyle::Get()->GetBrush("LevelSequenceEditor.CinematicViewportPlayMarker"),
 			DrawEffects,
 			bPlayMarkerOnKey ? KeyframeColor : LightGray
@@ -342,7 +342,7 @@ int32 SCinematicTransportRange::OnPaint(const FPaintArgs& Args, const FGeometry&
 		FSlateDrawElement::MakeBox(
 			OutDrawElements,
 			LayerId+1,
-			AllottedGeometry.ToPaintGeometry(FVector2D(AllottedGeometry.GetLocalSize().X*PlaybackStartLerp, TrackOffsetY), FVector2D(BrushWidth, TrackHeight)),
+			AllottedGeometry.ToPaintGeometry(FVector2f(BrushWidth, TrackHeight), FSlateLayoutTransform(FVector2f(AllottedGeometry.GetLocalSize().X*PlaybackStartLerp, TrackOffsetY))),
 			FLevelSequenceEditorStyle::Get()->GetBrush("LevelSequenceEditor.CinematicViewportRangeStart"),
 			DrawEffects,
 			FColor(32, 128, 32)	// 120, 75, 50 (HSV)
@@ -351,7 +351,7 @@ int32 SCinematicTransportRange::OnPaint(const FPaintArgs& Args, const FGeometry&
 		FSlateDrawElement::MakeBox(
 			OutDrawElements,
 			LayerId+1,
-			AllottedGeometry.ToPaintGeometry(FVector2D(AllottedGeometry.GetLocalSize().X*PlaybackEndLerp - BrushWidth, TrackOffsetY), FVector2D(BrushWidth, TrackHeight)),
+			AllottedGeometry.ToPaintGeometry(FVector2f(BrushWidth, TrackHeight), FSlateLayoutTransform(FVector2f(AllottedGeometry.GetLocalSize().X*PlaybackEndLerp - BrushWidth, TrackOffsetY))),
 			FLevelSequenceEditorStyle::Get()->GetBrush("LevelSequenceEditor.CinematicViewportRangeEnd"),
 			DrawEffects,
 			FColor(128, 32, 32)	// 0, 75, 50 (HSV)
diff --git a/Engine/Plugins/Runtime/AudioWidgets/Source/AudioWidgets/Private/SAudioMeter.cpp b/Engine/Plugins/Runtime/AudioWidgets/Source/AudioWidgets/Private/SAudioMeter.cpp
--- a/Engine/Plugins/Runtime/AudioWidgets/Source/AudioWidgets/Private/SAudioMeter.cpp
+++ b/Engine/Plugins/Runtime/AudioWidgets/Source/AudioWidgets/Private/SAudioMeter.cpp
@@ -71,7 +71,7 @@ int32 SAudioMeter::OnPaint(const FPaintArgs& Args, const FGeometry& AllottedGeom
 		FSlateDrawElement::MakeBox(
 			OutDrawElements,
 			LayerId,
-			MeterGeometry.ToPaintGeometry(BackgroundImageTopLeft, BackgroundImageSize),
+			MeterGeometry.ToPaintGeometry(BackgroundImageSize, FSlateLayoutTransform(BackgroundImageTopLeft)),
 			&Style->MeterBackgroundImage,
 			ESlateDrawEffect::None,
 			Style->BackgroundImage.GetTint(InWidgetStyle) * BackgroundColor.Get().GetColor(InWidgetStyle) * InWidgetStyle.GetColorAndOpacityTint()
@@ -98,7 +98,7 @@ int32 SAudioMeter::OnPaint(const FPaintArgs& Args, const FGeometry& AllottedGeom
  		FSlateDrawElement::MakeBox(
 			OutDrawElements,
 			LayerId,
-			MeterGeometry.ToPaintGeometry(MeterBackgroundTopLeft, MeterBackgroundSize),
+			MeterGeometry.ToPaintGeometry(MeterBackgroundSize, FSlateLayoutTransform(MeterBackgroundTopLeft)),
 			&Style->MeterBackgroundImage,
 			ESlateDrawEffect::None,
 			Style->MeterBackgroundImage.GetTint(InWidgetStyle) * MeterBackgroundColor.Get().GetColor(InWidgetStyle) * InWidgetStyle.GetColorAndOpacityTint()
@@ -122,7 +122,7 @@ int32 SAudioMeter::OnPaint(const FPaintArgs& Args, const FGeometry& AllottedGeom
 		FSlateDrawElement::MakeBox(
 			OutDrawElements,
 			LayerId,
-			MeterGeometry.ToPaintGeometry(MeterValueTopLeft, MeterValueSize),
+			MeterGeometry.ToPaintGeometry(MeterValueSize, FSlateLayoutTransform(MeterValueTopLeft)),
 			&Style->MeterValueImage,
 			ESlateDrawEffect::None,
 			Style->MeterValueImage.GetTint(InWidgetStyle) * MeterValueColor.Get().GetColor(InWidgetStyle) * InWidgetStyle.GetColorAndOpacityTint()
@@ -148,7 +148,7 @@ int32 SAudioMeter::OnPaint(const FPaintArgs& Args, const FGeometry& AllottedGeom
 		FSlateDrawElement::MakeBox(
 			OutDrawElements,
 			LayerId,
-			MeterGeometry.ToPaintGeometry(MeterPeakValueTopLeft, MeterPeakValueSize),
+			MeterGeometry.ToPaintGeometry(MeterPeakValueSize, FSlateLayoutTransform(MeterPeakValueTopLeft)),
 			&Style->MeterPeakImage,
 			ESlateDrawEffect::None,
 			Style->MeterPeakImage.GetTint(InWidgetStyle) * MeterPeakColor.Get().GetColor(InWidgetStyle) * InWidgetStyle.GetColorAndOpacityTint()
@@ -200,7 +200,7 @@ int32 SAudioMeter::OnPaint(const FPaintArgs& Args, const FGeometry& AllottedGeom
 			FSlateDrawElement::MakeBox(
 				OutDrawElements,
 				LayerId,
-				MeterGeometry.ToPaintGeometry(HashTopLeft, HashSize),
+				MeterGeometry.ToPaintGeometry(HashSize, FSlateLayoutTransform(HashTopLeft)),
 				&Style->MeterPeakImage,
 				ESlateDrawEffect::None,
 				MeterScaleColor.Get().GetColor(InWidgetStyle) * InWidgetStyle.GetColorAndOpacityTint()
diff --git a/Engine/Plugins/Runtime/Oculus/OculusVR/Source/OculusHMD/Private/OculusHMD.cpp b/Engine/Plugins/Runtime/Oculus/OculusVR/Source/OculusHMD/Private/OculusHMD.cpp
--- a/Engine/Plugins/Runtime/Oculus/OculusVR/Source/OculusHMD/Private/OculusHMD.cpp
+++ b/Engine/Plugins/Runtime/Oculus/OculusVR/Source/OculusHMD/Private/OculusHMD.cpp
@@ -1663,7 +1663,7 @@ namespace OculusHMD
 
 		if (bUseSeparateRenderTarget && Frame.IsValid())
 		{
-			CachedWindowSize = (Window.IsValid()) ? Window->GetSizeInScreen() : Viewport.GetSizeXY();
+			CachedWindowSize = (Window.IsValid()) ? FVector2D(Window->GetSizeInScreen()) : FVector2D(Viewport.GetSizeXY());
 		}
 	}
 
diff --git a/Engine/Plugins/Runtime/Synthesis/Source/Synthesis/Private/UI/SSynth2DSlider.cpp b/Engine/Plugins/Runtime/Synthesis/Source/Synthesis/Private/UI/SSynth2DSlider.cpp
--- a/Engine/Plugins/Runtime/Synthesis/Source/Synthesis/Private/UI/SSynth2DSlider.cpp
+++ b/Engine/Plugins/Runtime/Synthesis/Source/Synthesis/Private/UI/SSynth2DSlider.cpp
@@ -63,7 +63,7 @@ int32 SSynth2DSlider::OnPaint(const FPaintArgs& Args, const FGeometry& AllottedG
 	FSlateDrawElement::MakeBox(
 		OutDrawElements,
 		LayerId++,
-		SliderGeometry.ToPaintGeometry(HandleTopLeftPoint, Style->NormalThumbImage.ImageSize),
+		SliderGeometry.ToPaintGeometry(Style->NormalThumbImage.ImageSize, FSlateLayoutTransform(HandleTopLeftPoint)),
 		LockedAttribute.Get() ? &Style->DisabledThumbImage : &Style->NormalThumbImage,
 		DrawEffects,
 		SliderHandleColor.Get().GetColor(InWidgetStyle) * InWidgetStyle.GetColorAndOpacityTint()
diff --git a/Engine/Plugins/Runtime/nDisplay/Source/DisplayClusterColorGrading/Private/DetailView/SDisplayClusterColorGradingDetailTreeRow.cpp b/Engine/Plugins/Runtime/nDisplay/Source/DisplayClusterColorGrading/Private/DetailView/SDisplayClusterColorGradingDetailTreeRow.cpp
--- a/Engine/Plugins/Runtime/nDisplay/Source/DisplayClusterColorGrading/Private/DetailView/SDisplayClusterColorGradingDetailTreeRow.cpp
+++ b/Engine/Plugins/Runtime/nDisplay/Source/DisplayClusterColorGrading/Private/DetailView/SDisplayClusterColorGradingDetailTreeRow.cpp
@@ -67,7 +67,7 @@ namespace UE::DisplayClusterColorGradingDetailTreeRow
 				FSlateDrawElement::MakeBox(
 					OutDrawElements,
 					LayerId,
-					AllottedGeometry.ToPaintGeometry(FVector2D(16 * IndentIndex, 0), FVector2D(16, AllottedGeometry.GetLocalSize().Y)),
+					AllottedGeometry.ToPaintGeometry(FVector2D(16, AllottedGeometry.GetLocalSize().Y), FSlateLayoutTransform(FVector2D(16 * IndentIndex, 0))),
 					BackgroundBrush,
 					ESlateDrawEffect::None,
 					BackgroundColor.GetColor(InWidgetStyle)
@@ -76,7 +76,7 @@ namespace UE::DisplayClusterColorGradingDetailTreeRow
 				FSlateDrawElement::MakeBox(
 					OutDrawElements,
 					LayerId + 1,
-					AllottedGeometry.ToPaintGeometry(FVector2D(16 * IndentIndex, 0), FVector2D(16, AllottedGeometry.GetLocalSize().Y)),
+					AllottedGeometry.ToPaintGeometry(FVector2D(16, AllottedGeometry.GetLocalSize().Y), FSlateLayoutTransform(FVector2D(16 * IndentIndex, 0))),
 					DropShadowBrush
 				);
 			}
diff --git a/Engine/Plugins/Runtime/nDisplay/Source/DisplayClusterConfigurator/Private/Views/OutputMapping/Widgets/SDisplayClusterConfiguratorRuler.cpp b/Engine/Plugins/Runtime/nDisplay/Source/DisplayClusterConfigurator/Private/Views/OutputMapping/Widgets/SDisplayClusterConfiguratorRuler.cpp
--- a/Engine/Plugins/Runtime/nDisplay/Source/DisplayClusterConfigurator/Private/Views/OutputMapping/Widgets/SDisplayClusterConfiguratorRuler.cpp
+++ b/Engine/Plugins/Runtime/nDisplay/Source/DisplayClusterConfigurator/Private/Views/OutputMapping/Widgets/SDisplayClusterConfiguratorRuler.cpp
@@ -178,7 +178,7 @@ int32 SDisplayClusterConfiguratorRuler::DrawTicks( FSlateWindowElementList& OutD
 			FSlateDrawElement::MakeLines(
 				OutDrawElements,
 				InArgs.StartLayer,
-				InArgs.AllottedGeometry.ToPaintGeometry( Offset, TickSize ),
+				InArgs.AllottedGeometry.ToPaintGeometry( TickSize, FSlateLayoutTransform(Offset) ),
 				LinePoints,
 				InArgs.DrawEffects,
 				InArgs.TickColor,
@@ -211,7 +211,7 @@ int32 SDisplayClusterConfiguratorRuler::DrawTicks( FSlateWindowElementList& OutD
 				FSlateDrawElement::MakeText(
 					OutDrawElements,
 					InArgs.StartLayer, 
-					InArgs.AllottedGeometry.ToPaintGeometry( TextOffset, TextSize ), 
+					InArgs.AllottedGeometry.ToPaintGeometry( TextSize, FSlateLayoutTransform(TextOffset) ), 
 					FrameString, 
 					SmallLayoutFont,
 					InArgs.DrawEffects,
@@ -234,7 +234,7 @@ int32 SDisplayClusterConfiguratorRuler::DrawTicks( FSlateWindowElementList& OutD
 			FSlateDrawElement::MakeLines(
 				OutDrawElements,
 				InArgs.StartLayer,
-				InArgs.AllottedGeometry.ToPaintGeometry( Offset, TickSize ),
+				InArgs.AllottedGeometry.ToPaintGeometry( TickSize, FSlateLayoutTransform(Offset) ),
 				LinePoints,
 				InArgs.DrawEffects,
 				InArgs.TickColor,
diff --git a/Engine/Plugins/Slate/SlateInsights/Source/SlateInsights/Private/SSlateTraceFlags.cpp b/Engine/Plugins/Slate/SlateInsights/Source/SlateInsights/Private/SSlateTraceFlags.cpp
--- a/Engine/Plugins/Slate/SlateInsights/Source/SlateInsights/Private/SSlateTraceFlags.cpp
+++ b/Engine/Plugins/Slate/SlateInsights/Source/SlateInsights/Private/SSlateTraceFlags.cpp
@@ -41,7 +41,7 @@ namespace Private
 				(
 					OutDrawElements,
 					LayerId,
-					AllottedGeometry.ToPaintGeometry(BoxSizeX0 * Index, BoxSize),
+					AllottedGeometry.ToPaintGeometry(BoxSize, FSlateLayoutTransform(BoxSizeX0 * Index)),
 					BackgroundBrush,
 					ESlateDrawEffect::None,
 					SelectionColor.GetColor(InWidgetStyle)
@@ -53,7 +53,7 @@ namespace Private
 				(
 					OutDrawElements,
 					LayerId,
-					AllottedGeometry.ToPaintGeometry(BoxSizeX0 * Index, BoxSize),
+					AllottedGeometry.ToPaintGeometry(BoxSize, FSlateLayoutTransform(BoxSizeX0 * Index)),
 					BackgroundBrush,
 					ESlateDrawEffect::None,
 					InvertedForeground.GetColor(InWidgetStyle)
diff --git a/Engine/Plugins/Tests/AutomationDriverTests/Source/AutomationDriverTests/Private/AutomationDriver.spec.cpp b/Engine/Plugins/Tests/AutomationDriverTests/Source/AutomationDriverTests/Private/AutomationDriver.spec.cpp
--- a/Engine/Plugins/Tests/AutomationDriverTests/Source/AutomationDriverTests/Private/AutomationDriver.spec.cpp
+++ b/Engine/Plugins/Tests/AutomationDriverTests/Source/AutomationDriverTests/Private/AutomationDriver.spec.cpp
@@ -969,7 +969,7 @@ void FAutomationDriverSpec::Define()
 							FWidgetPath WidgetPath;
 							FSlateApplication::Get().FindPathToWidget(SuiteWidget->GetKeyWidget(EPianoKey::F).ToSharedRef(), WidgetPath);
 
-							TEST_EQUAL(ElementPosition, WidgetPath.Widgets.Last().Geometry.LocalToAbsolute(FVector2D::ZeroVector));
+							TEST_EQUAL(ElementPosition, FVector2D(WidgetPath.Widgets.Last().Geometry.LocalToAbsolute(FVector2f::ZeroVector)));
 						}
 						Done.Execute();
 					}
@@ -1001,7 +1001,7 @@ void FAutomationDriverSpec::Define()
 							FWidgetPath WidgetPath;
 							FSlateApplication::Get().FindPathToWidget(SuiteWidget->GetKeyWidget(EPianoKey::F).ToSharedRef(), WidgetPath);
 
-							TEST_EQUAL(ElementSize, WidgetPath.Widgets.Last().Geometry.GetLocalSize());
+							TEST_EQUAL(ElementSize, FVector2D(WidgetPath.Widgets.Last().Geometry.GetLocalSize()));
 						}
 						Done.Execute();
 					}
diff --git a/Engine/Plugins/VirtualProduction/CameraCalibration/Source/CameraCalibrationEditor/Private/AssetEditor/CameraCalibrationTimeSliderController.cpp b/Engine/Plugins/VirtualProduction/CameraCalibration/Source/CameraCalibrationEditor/Private/AssetEditor/CameraCalibrationTimeSliderController.cpp
--- a/Engine/Plugins/VirtualProduction/CameraCalibration/Source/CameraCalibrationEditor/Private/AssetEditor/CameraCalibrationTimeSliderController.cpp
+++ b/Engine/Plugins/VirtualProduction/CameraCalibration/Source/CameraCalibrationEditor/Private/AssetEditor/CameraCalibrationTimeSliderController.cpp
@@ -80,7 +80,7 @@ int32 FCameraCalibrationTimeSliderController::OnPaintTimeSlider( bool bMirrorLab
 		const float      HandleEnd          = HandleStart + 13.0f;
 
 		const int32 ArrowLayer = LayerId + 2;
-		FPaintGeometry MyGeometry =	AllottedGeometry.ToPaintGeometry( FVector2D( HandleStart, 0 ), FVector2D( HandleEnd - HandleStart, AllottedGeometry.Size.Y ) );
+		FPaintGeometry MyGeometry =	AllottedGeometry.ToPaintGeometry( FVector2D( HandleEnd - HandleStart, AllottedGeometry.Size.Y ), FSlateLayoutTransform(FVector2D( HandleStart, 0 )) );
 		FLinearColor ScrubColor = InWidgetStyle.GetColorAndOpacityTint();
 		{
 			ScrubColor.A = ScrubColor.A * 0.75f;
@@ -121,7 +121,7 @@ int32 FCameraCalibrationTimeSliderController::OnPaintTimeSlider( bool bMirrorLab
 			FSlateDrawElement::MakeText(
 				OutDrawElements,
 				LayerId + 1, 
-				AllottedGeometry.ToPaintGeometry( TextOffset, TextSize ), 
+				AllottedGeometry.ToPaintGeometry( TextSize, FSlateLayoutTransform(TextOffset) ), 
 				FormattedCategoryString, 
 				SmallLayoutFont,
 				DrawEffects,
@@ -161,7 +161,7 @@ int32 FCameraCalibrationTimeSliderController::OnPaintViewArea( const FGeometry&
 		FSlateDrawElement::MakeLines(
 			OutDrawElements,
 			LayerId+1,
-			AllottedGeometry.ToPaintGeometry( FVector2D(LinePos, 0.0f ), FVector2D(1.0f,1.0f) ),
+			AllottedGeometry.ToPaintGeometry( FVector2D(1.0f,1.0f), FSlateLayoutTransform(FVector2D(LinePos, 0.0f )) ),
 			LinePoints,
 			DrawEffects,
 			FLinearColor(1.f, 1.f, 1.f, .5f),
diff --git a/Engine/Plugins/VirtualProduction/DMX/DMXEngine/Source/DMXControlConsoleEditor/Private/Widgets/SDMXControlConsoleEditorSpinBoxVertical.h b/Engine/Plugins/VirtualProduction/DMX/DMXEngine/Source/DMXControlConsoleEditor/Private/Widgets/SDMXControlConsoleEditorSpinBoxVertical.h
--- a/Engine/Plugins/VirtualProduction/DMX/DMXEngine/Source/DMXControlConsoleEditor/Private/Widgets/SDMXControlConsoleEditorSpinBoxVertical.h
+++ b/Engine/Plugins/VirtualProduction/DMX/DMXEngine/Source/DMXControlConsoleEditor/Private/Widgets/SDMXControlConsoleEditorSpinBoxVertical.h
@@ -353,7 +353,7 @@ public:
 				FSlateDrawElement::MakeBox(
 					OutDrawElements,
 					FilledLayer,
-					AllottedGeometry.ToPaintGeometry(AllottedGeometry.GetLocalSize() - FillSize, FillSize),
+					AllottedGeometry.ToPaintGeometry(FillSize, FSlateLayoutTransform(AllottedGeometry.GetLocalSize() - FillSize)),
 					FillImage,
 					DrawEffects,
 					FillImage->GetTint(InWidgetStyle) * InWidgetStyle.GetColorAndOpacityTint()
diff --git a/Engine/Plugins/VirtualProduction/DMX/DMXPixelMapping/Source/DMXPixelMappingEditor/Private/Views/SDMXPixelMappingDesignerView.cpp b/Engine/Plugins/VirtualProduction/DMX/DMXPixelMapping/Source/DMXPixelMappingEditor/Private/Views/SDMXPixelMappingDesignerView.cpp
--- a/Engine/Plugins/VirtualProduction/DMX/DMXPixelMapping/Source/DMXPixelMappingEditor/Private/Views/SDMXPixelMappingDesignerView.cpp
+++ b/Engine/Plugins/VirtualProduction/DMX/DMXPixelMapping/Source/DMXPixelMappingEditor/Private/Views/SDMXPixelMappingDesignerView.cpp
@@ -793,39 +793,39 @@ FVector2D SDMXPixelMappingDesignerView::GetExtensionPosition(TSharedPtr<SDMXPixe
 
 		if (GetComponentGeometry(SelectedComponent, SelectedComponentGeometry))
 		{
-			const FVector2D WidgetPosition = [Handle, &SelectedComponentGeometry]()
+			const FVector2f WidgetPosition = [Handle, &SelectedComponentGeometry]()
 			{
 				// Get the initial offset based on the location around the selected object.
 				switch (Handle->GetTransformDirection())
 				{
 				case EDMXPixelMappingTransformDirection::CenterRight:
-					return FVector2D(SelectedComponentGeometry.GetLocalSize().X, SelectedComponentGeometry.GetLocalSize().Y * 0.5f);
+					return FVector2f(SelectedComponentGeometry.GetLocalSize().X, SelectedComponentGeometry.GetLocalSize().Y * 0.5f);
 				case EDMXPixelMappingTransformDirection::BottomLeft:
-					return FVector2D(0, SelectedComponentGeometry.GetLocalSize().Y);
+					return FVector2f(0, SelectedComponentGeometry.GetLocalSize().Y);
 				case EDMXPixelMappingTransformDirection::BottomCenter:
-					return FVector2D(SelectedComponentGeometry.GetLocalSize().X * 0.5f, SelectedComponentGeometry.GetLocalSize().Y);
+					return FVector2f(SelectedComponentGeometry.GetLocalSize().X * 0.5f, SelectedComponentGeometry.GetLocalSize().Y);
 				case EDMXPixelMappingTransformDirection::BottomRight:
-					return SelectedComponentGeometry.GetLocalSize();
+					return UE::Slate::CastToVector2f(SelectedComponentGeometry.GetLocalSize());
 				default:
 					checkNoEntry(); // Unhandled enum value
 				}
-				return FVector2D::ZeroVector;
+				return FVector2f::ZeroVector;
 			}();
 
-			const FVector2D SelectedWidgetScale = FVector2D(SelectedComponentGeometry.GetAccumulatedRenderTransform().GetMatrix().GetScale().GetVector());
+			const FVector2f SelectedWidgetScale = SelectedComponentGeometry.GetAccumulatedRenderTransform().GetMatrix().GetScale().GetVector();
 
-			const FVector2D ApplicationScaledOffset = Handle->GetOffset() * GetDesignerGeometry().Scale;
+			const FVector2f ApplicationScaledOffset = UE::Slate::CastToVector2f(Handle->GetOffset() * GetDesignerGeometry().Scale);
 
-			const FVector2D LocalOffsetFull = ApplicationScaledOffset / SelectedWidgetScale;
-			const FVector2D PositionFullOffset = GetDesignerGeometry().AbsoluteToLocal(SelectedComponentGeometry.LocalToAbsolute(WidgetPosition + LocalOffsetFull));
-			const FVector2D LocalOffsetHalf = (ApplicationScaledOffset / 2.0f) / SelectedWidgetScale;
-			const FVector2D PositionHalfOffset = GetDesignerGeometry().AbsoluteToLocal(SelectedComponentGeometry.LocalToAbsolute(WidgetPosition + LocalOffsetHalf));
+			const FVector2f LocalOffsetFull = ApplicationScaledOffset / SelectedWidgetScale;
+			const FVector2f PositionFullOffset = GetDesignerGeometry().AbsoluteToLocal(SelectedComponentGeometry.LocalToAbsolute(WidgetPosition + LocalOffsetFull));
+			const FVector2f LocalOffsetHalf = (ApplicationScaledOffset / 2.0f) / SelectedWidgetScale;
+			const FVector2f PositionHalfOffset = GetDesignerGeometry().AbsoluteToLocal(SelectedComponentGeometry.LocalToAbsolute(WidgetPosition + LocalOffsetHalf));
 
-			FVector2D PivotCorrection = PositionHalfOffset - (PositionFullOffset + FVector2D(5.0f, 5.0f));
+			FVector2f PivotCorrection = PositionHalfOffset - (PositionFullOffset + FVector2f(5.0f, 5.0f));
 
-			const FVector2D FinalPosition = PositionFullOffset + PivotCorrection;
+			const FVector2f FinalPosition = PositionFullOffset + PivotCorrection;
 
-			return FinalPosition;
+			return FVector2D(FinalPosition);
 		}
 	}
 
diff --git a/Engine/Plugins/VirtualProduction/DMX/DMXPixelMapping/Source/DMXPixelMappingEditor/Private/Widgets/SDMXPixelMappingRuler.cpp b/Engine/Plugins/VirtualProduction/DMX/DMXPixelMapping/Source/DMXPixelMappingEditor/Private/Widgets/SDMXPixelMappingRuler.cpp
--- a/Engine/Plugins/VirtualProduction/DMX/DMXPixelMapping/Source/DMXPixelMappingEditor/Private/Widgets/SDMXPixelMappingRuler.cpp
+++ b/Engine/Plugins/VirtualProduction/DMX/DMXPixelMapping/Source/DMXPixelMappingEditor/Private/Widgets/SDMXPixelMappingRuler.cpp
@@ -178,7 +178,7 @@ int32 SDMXPixelMappingRuler::DrawTicks( FSlateWindowElementList& OutDrawElements
 			FSlateDrawElement::MakeLines(
 				OutDrawElements,
 				InArgs.StartLayer,
-				InArgs.AllottedGeometry.ToPaintGeometry( Offset, TickSize ),
+				InArgs.AllottedGeometry.ToPaintGeometry( TickSize, FSlateLayoutTransform(Offset) ),
 				LinePoints,
 				InArgs.DrawEffects,
 				InArgs.TickColor,
@@ -211,7 +211,7 @@ int32 SDMXPixelMappingRuler::DrawTicks( FSlateWindowElementList& OutDrawElements
 				FSlateDrawElement::MakeText(
 					OutDrawElements,
 					InArgs.StartLayer, 
-					InArgs.AllottedGeometry.ToPaintGeometry( TextOffset, TextSize ), 
+					InArgs.AllottedGeometry.ToPaintGeometry( TextSize, FSlateLayoutTransform(TextOffset) ), 
 					FrameString, 
 					SmallLayoutFont,
 					InArgs.DrawEffects,
@@ -234,7 +234,7 @@ int32 SDMXPixelMappingRuler::DrawTicks( FSlateWindowElementList& OutDrawElements
 			FSlateDrawElement::MakeLines(
 				OutDrawElements,
 				InArgs.StartLayer,
-				InArgs.AllottedGeometry.ToPaintGeometry( Offset, TickSize ),
+				InArgs.AllottedGeometry.ToPaintGeometry( TickSize, FSlateLayoutTransform(Offset) ),
 				LinePoints,
 				InArgs.DrawEffects,
 				InArgs.TickColor,
diff --git a/Engine/Plugins/VirtualProduction/DMX/DMXPixelMapping/Source/DMXPixelMappingEditorWidgets/Private/SDMXPixelMappingComponentBox.cpp b/Engine/Plugins/VirtualProduction/DMX/DMXPixelMapping/Source/DMXPixelMappingEditorWidgets/Private/SDMXPixelMappingComponentBox.cpp
--- a/Engine/Plugins/VirtualProduction/DMX/DMXPixelMapping/Source/DMXPixelMappingEditorWidgets/Private/SDMXPixelMappingComponentBox.cpp
+++ b/Engine/Plugins/VirtualProduction/DMX/DMXPixelMapping/Source/DMXPixelMappingEditorWidgets/Private/SDMXPixelMappingComponentBox.cpp
@@ -74,7 +74,7 @@ void SDMXPixelMappingComponentBox::SetIDText(const FText& NewIDText)
 	IDTextBlock->SetText(NewIDText);
 }
 
-const FVector2D& SDMXPixelMappingComponentBox::GetLocalSize() const
+FVector2f SDMXPixelMappingComponentBox::GetLocalSize() const
 {
 	return GetCachedGeometry().GetLocalSize();
 }
diff --git a/Engine/Plugins/VirtualProduction/DMX/DMXPixelMapping/Source/DMXPixelMappingEditorWidgets/Private/SDMXPixelMappingComponentLabel.cpp b/Engine/Plugins/VirtualProduction/DMX/DMXPixelMapping/Source/DMXPixelMappingEditorWidgets/Private/SDMXPixelMappingComponentLabel.cpp
--- a/Engine/Plugins/VirtualProduction/DMX/DMXPixelMapping/Source/DMXPixelMappingEditorWidgets/Private/SDMXPixelMappingComponentLabel.cpp
+++ b/Engine/Plugins/VirtualProduction/DMX/DMXPixelMapping/Source/DMXPixelMappingEditorWidgets/Private/SDMXPixelMappingComponentLabel.cpp
@@ -76,7 +76,7 @@ void SDMXPixelMappingComponentLabel::SetText(const FText& Text)
 	LabelTextBlock->SetText(Text);
 }
 
-const FVector2D& SDMXPixelMappingComponentLabel::GetLocalSize() const
+FVector2f SDMXPixelMappingComponentLabel::GetLocalSize() const
 {
 	return GetCachedGeometry().GetLocalSize();
 }
diff --git a/Engine/Plugins/VirtualProduction/DMX/DMXPixelMapping/Source/DMXPixelMappingEditorWidgets/Private/SDMXPixelMappingScreenComponentBox.cpp b/Engine/Plugins/VirtualProduction/DMX/DMXPixelMapping/Source/DMXPixelMappingEditorWidgets/Private/SDMXPixelMappingScreenComponentBox.cpp
--- a/Engine/Plugins/VirtualProduction/DMX/DMXPixelMapping/Source/DMXPixelMappingEditorWidgets/Private/SDMXPixelMappingScreenComponentBox.cpp
+++ b/Engine/Plugins/VirtualProduction/DMX/DMXPixelMapping/Source/DMXPixelMappingEditorWidgets/Private/SDMXPixelMappingScreenComponentBox.cpp
@@ -61,7 +61,7 @@ void SDMXPixelMappingScreenComponentBox::SetLocalSize(const FVector2D& NewLocalS
 	ComponentBox->SetHeightOverride(NewLocalSize.Y);
 }
 
-const FVector2D& SDMXPixelMappingScreenComponentBox::GetLocalSize() const
+FVector2f SDMXPixelMappingScreenComponentBox::GetLocalSize() const
 {
 	return GetCachedGeometry().GetLocalSize();
 }
diff --git a/Engine/Plugins/VirtualProduction/DMX/DMXPixelMapping/Source/DMXPixelMappingEditorWidgets/Public/SDMXPixelMappingComponentBox.h b/Engine/Plugins/VirtualProduction/DMX/DMXPixelMapping/Source/DMXPixelMappingEditorWidgets/Public/SDMXPixelMappingComponentBox.h
--- a/Engine/Plugins/VirtualProduction/DMX/DMXPixelMapping/Source/DMXPixelMappingEditorWidgets/Public/SDMXPixelMappingComponentBox.h
+++ b/Engine/Plugins/VirtualProduction/DMX/DMXPixelMapping/Source/DMXPixelMappingEditorWidgets/Public/SDMXPixelMappingComponentBox.h
@@ -38,7 +38,7 @@ public:
 	virtual void SetIDText(const FText& NewIDText);
 
 	/** Returns the local size */
-	virtual const FVector2D& GetLocalSize() const;
+	virtual FVector2f GetLocalSize() const;
 
 	/** Sets the color of the border */
 	virtual void SetBorderColor(const FLinearColor& Color);
diff --git a/Engine/Plugins/VirtualProduction/DMX/DMXPixelMapping/Source/DMXPixelMappingEditorWidgets/Public/SDMXPixelMappingComponentLabel.h b/Engine/Plugins/VirtualProduction/DMX/DMXPixelMapping/Source/DMXPixelMappingEditorWidgets/Public/SDMXPixelMappingComponentLabel.h
--- a/Engine/Plugins/VirtualProduction/DMX/DMXPixelMapping/Source/DMXPixelMappingEditorWidgets/Public/SDMXPixelMappingComponentLabel.h
+++ b/Engine/Plugins/VirtualProduction/DMX/DMXPixelMapping/Source/DMXPixelMappingEditorWidgets/Public/SDMXPixelMappingComponentLabel.h
@@ -40,7 +40,7 @@ public:
 	virtual void SetText(const FText& Text);
 
 	/** Returns the local size */
-	virtual const FVector2D& GetLocalSize() const;
+	virtual FVector2f GetLocalSize() const;
 
 protected:
 	/** The box that defines the size */
diff --git a/Engine/Plugins/VirtualProduction/DMX/DMXPixelMapping/Source/DMXPixelMappingEditorWidgets/Public/SDMXPixelMappingScreenComponentBox.h b/Engine/Plugins/VirtualProduction/DMX/DMXPixelMapping/Source/DMXPixelMappingEditorWidgets/Public/SDMXPixelMappingScreenComponentBox.h
--- a/Engine/Plugins/VirtualProduction/DMX/DMXPixelMapping/Source/DMXPixelMappingEditorWidgets/Public/SDMXPixelMappingScreenComponentBox.h
+++ b/Engine/Plugins/VirtualProduction/DMX/DMXPixelMapping/Source/DMXPixelMappingEditorWidgets/Public/SDMXPixelMappingScreenComponentBox.h
@@ -74,7 +74,7 @@ public:
 	virtual void SetLocalSize(const FVector2D& NewSize) override;
 
 	/** Returns the local size */
-	virtual const FVector2D& GetLocalSize() const override;
+	virtual FVector2f GetLocalSize() const override;
 
 	/** Rebuilds the grid */
 	void RebuildGrid(const FDMXPixelMappingScreenComponentGridParams& GridParams);
diff --git a/Engine/Plugins/VirtualProduction/TimedDataMonitor/Source/TimedDataMonitorEditor/Private/SFrameTime.h b/Engine/Plugins/VirtualProduction/TimedDataMonitor/Source/TimedDataMonitorEditor/Private/SFrameTime.h
--- a/Engine/Plugins/VirtualProduction/TimedDataMonitor/Source/TimedDataMonitorEditor/Private/SFrameTime.h
+++ b/Engine/Plugins/VirtualProduction/TimedDataMonitor/Source/TimedDataMonitorEditor/Private/SFrameTime.h
@@ -97,8 +97,9 @@ public:
 		{
 			FPaintGeometry PaintRect;
 			PaintRect = AllottedGeometry.ToPaintGeometry(
-				FVector2D(0.0f, 0.0f),
-				FVector2D(AllottedGeometry.GetLocalSize().X * ClampedFraction, AllottedGeometry.GetLocalSize().Y));
+				FVector2D(AllottedGeometry.GetLocalSize().X * ClampedFraction, AllottedGeometry.GetLocalSize().Y),
+				FSlateLayoutTransform()
+				);
 
 			// Draw fill
 			FSlateDrawElement::MakeBox(
diff --git a/Engine/Plugins/VirtualProduction/TimedDataMonitor/Source/TimedDataMonitorEditor/Private/STimingDiagramWidget.cpp b/Engine/Plugins/VirtualProduction/TimedDataMonitor/Source/TimedDataMonitorEditor/Private/STimingDiagramWidget.cpp
--- a/Engine/Plugins/VirtualProduction/TimedDataMonitor/Source/TimedDataMonitorEditor/Private/STimingDiagramWidget.cpp
+++ b/Engine/Plugins/VirtualProduction/TimedDataMonitor/Source/TimedDataMonitorEditor/Private/STimingDiagramWidget.cpp
@@ -92,7 +92,7 @@ public:
 			if (bShowSnapshot && SnapshotLocationMinX < SizeOfFurthur)
 			{
 				FSlateDrawElement::MakeText(OutDrawElements, LayerId,
-					AllottedGeometry.ToPaintGeometry(FVector2D(0, LocationOfSnapshotBoxY), FVector2D(SizeOfFurthur, SizeOfBoxY)),
+					AllottedGeometry.ToPaintGeometry(FVector2D(SizeOfFurthur, SizeOfBoxY), FSlateLayoutTransform(FVector2D(0, LocationOfSnapshotBoxY))),
 					FEditorFontGlyphs::Angle_Double_Left,
 					FontInfo, ESlateDrawEffect::None,
 					FLinearColor::White);
@@ -101,7 +101,7 @@ public:
 			if (bShowMean && MeanLocationMinX < SizeOfFurthur)
 			{
 				FSlateDrawElement::MakeText(OutDrawElements, LayerId,
-					AllottedGeometry.ToPaintGeometry(FVector2D(0, LocationOfMeanBoxY), FVector2D(SizeOfFurthur, SizeOfBoxY)),
+					AllottedGeometry.ToPaintGeometry(FVector2D(SizeOfFurthur, SizeOfBoxY), FSlateLayoutTransform(FVector2D(0, LocationOfMeanBoxY))),
 					FEditorFontGlyphs::Angle_Double_Left,
 					FontInfo, ESlateDrawEffect::None,
 					FLinearColor::White);
@@ -116,7 +116,7 @@ public:
 			{
 				const float DrawSizeX = FMath::Clamp(SnapshotLocationMaxX - DrawLocationX, 1.f, SizeX - DrawLocationX - SizeOfFurthur);
 				FSlateDrawElement::MakeBox(OutDrawElements, LayerId,
-					AllottedGeometry.ToPaintGeometry(FVector2D(DrawLocationX, LocationOfSnapshotBoxY), FVector2D(DrawSizeX, SizeOfBoxY)),
+					AllottedGeometry.ToPaintGeometry(FVector2D(DrawSizeX, SizeOfBoxY), FSlateLayoutTransform(FVector2D(DrawLocationX, LocationOfSnapshotBoxY))),
 					DarkBrush, ESlateDrawEffect::None,
 					BrightBrush->GetTint(InWidgetStyle) * FLinearColor(0.5f, 0.5f, 0.5f));
 			}
@@ -131,7 +131,7 @@ public:
 				if (bShowMean)
 				{
 					FSlateDrawElement::MakeBox(OutDrawElements, LayerId,
-						AllottedGeometry.ToPaintGeometry(FVector2D(DrawLocationX, LocationOfMeanBoxY), FVector2D(DrawSizeX, SizeOfBoxY)),
+						AllottedGeometry.ToPaintGeometry(FVector2D(DrawSizeX, SizeOfBoxY), FSlateLayoutTransform(FVector2D(DrawLocationX, LocationOfMeanBoxY))),
 						DarkBrush, ESlateDrawEffect::None,
 						BrightBrush->GetTint(InWidgetStyle) * FLinearColor(0.2f, 0.2f, 0.2f));
 				}
@@ -151,7 +151,7 @@ public:
 					const float DrawLocationX = FMath::Clamp(SigmaLocationMinX, SizeOfFurthur, SizeX - SizeOfFurthur);
 					const float DrawSizeX = FMath::Clamp(SigmaLocationMaxX - DrawLocationX, 0.f, SizeX - SizeOfFurthur);
 					FSlateDrawElement::MakeBox(OutDrawElements, LayerId,
-						AllottedGeometry.ToPaintGeometry(FVector2D(DrawLocationX, LocationOfSigmaY), FVector2D(DrawSizeX, SizeOfSigmaY)),
+						AllottedGeometry.ToPaintGeometry(FVector2D(DrawSizeX, SizeOfSigmaY), FSlateLayoutTransform(FVector2D(DrawLocationX, LocationOfSigmaY))),
 						DarkBrush, ESlateDrawEffect::None,
 						BrightBrush->GetTint(InWidgetStyle) * FLinearColor::White);
 				}
@@ -164,7 +164,7 @@ public:
 					const float DrawLocationX = FMath::Clamp(SigmaLocationMinX, SizeOfFurthur, SizeX - SizeOfFurthur);
 					const float DrawSizeX = FMath::Clamp(SigmaLocationMaxX - DrawLocationX, 0.f, SizeX - SizeOfFurthur);
 					FSlateDrawElement::MakeBox(OutDrawElements, LayerId,
-						AllottedGeometry.ToPaintGeometry(FVector2D(DrawLocationX, LocationOfSigmaY), FVector2D(DrawSizeX, SizeOfSigmaY)),
+						AllottedGeometry.ToPaintGeometry(FVector2D(DrawSizeX, SizeOfSigmaY), FSlateLayoutTransform(FVector2D(DrawLocationX, LocationOfSigmaY))),
 						DarkBrush, ESlateDrawEffect::None,
 						BrightBrush->GetTint(InWidgetStyle) * FLinearColor::White);
 				}
@@ -177,7 +177,7 @@ public:
 			if (bShowSnapshot && SnapshotLocationMaxX > SizeX - SizeOfFurthur)
 			{
 				FSlateDrawElement::MakeText(OutDrawElements, LayerId,
-					AllottedGeometry.ToPaintGeometry(FVector2D(SizeX - SizeOfFurthur + 4, LocationOfSnapshotBoxY), FVector2D(SizeOfFurthur, SizeOfBoxY)),
+					AllottedGeometry.ToPaintGeometry(FVector2D(SizeOfFurthur, SizeOfBoxY), FSlateLayoutTransform(FVector2D(SizeX - SizeOfFurthur + 4, LocationOfSnapshotBoxY))),
 					FEditorFontGlyphs::Angle_Double_Right,
 					FontInfo, ESlateDrawEffect::None,
 					FLinearColor::White);
@@ -186,7 +186,7 @@ public:
 			if (bShowMean && MeanLocationMaxX > SizeX - SizeOfFurthur)
 			{
 				FSlateDrawElement::MakeText(OutDrawElements, LayerId,
-					AllottedGeometry.ToPaintGeometry(FVector2D(SizeX - SizeOfFurthur + 4, LocationOfSnapshotBoxY), FVector2D(SizeOfFurthur, SizeOfBoxY)),
+					AllottedGeometry.ToPaintGeometry(FVector2D(SizeOfFurthur, SizeOfBoxY), FSlateLayoutTransform(FVector2D(SizeX - SizeOfFurthur + 4, LocationOfSnapshotBoxY))),
 					FEditorFontGlyphs::Angle_Double_Right,
 					FontInfo, ESlateDrawEffect::None,
 					FLinearColor::White);
@@ -198,7 +198,7 @@ public:
 			const bool bIsInRange = (FMath::IsNearlyEqual(EvaluationTime, MinSampleTime) || EvaluationTime >= MinSampleTime) && (FMath::IsNearlyEqual(EvaluationTime, MaxSampleTime) || EvaluationTime <= MaxSampleTime);
 			FLinearColor EvaluationColor = bIsInRange ? FLinearColor::Green : FLinearColor::Red;
 			FSlateDrawElement::MakeBox(OutDrawElements, LayerId,
-				AllottedGeometry.ToPaintGeometry(FVector2D(LocationOfCenter, 0), FVector2D(1, SizeY)),
+				AllottedGeometry.ToPaintGeometry(FVector2D(1, SizeY), FSlateLayoutTransform(FVector2D(LocationOfCenter, 0))),
 				DarkBrush, ESlateDrawEffect::None,
 				DarkBrush->GetTint(InWidgetStyle) * EvaluationColor);
 		}
diff --git a/Engine/Source/Developer/LogVisualizer/Private/SVisualLoggerTimelineBar.cpp b/Engine/Source/Developer/LogVisualizer/Private/SVisualLoggerTimelineBar.cpp
--- a/Engine/Source/Developer/LogVisualizer/Private/SVisualLoggerTimelineBar.cpp
+++ b/Engine/Source/Developer/LogVisualizer/Private/SVisualLoggerTimelineBar.cpp
@@ -237,13 +237,13 @@ int32 SVisualLoggerTimelineBar::OnPaint(const FPaintArgs& Args, const FGeometry&
 
 		if (EndPos - StartPos > 0)
 		{
-			const double BarWidth = (EndPos - StartPos);
+			const float BarWidth = (EndPos - StartPos);
 			FSlateDrawElement::MakeBox(
 				OutDrawElements,
 				RetLayerId,
 				AllottedGeometry.ToPaintGeometry(
-				FVector2D(StartPos, 0.),
-				FVector2D(BarWidth, AllottedGeometry.Size.Y)),
+				FVector2f(StartPos, 0.0f),
+				FSlateLayoutTransform(FVector2f(BarWidth, AllottedGeometry.Size.Y))),
 				FillImage,
 				DrawEffects,
 				CurrentTimeColor
@@ -261,8 +261,8 @@ int32 SVisualLoggerTimelineBar::OnPaint(const FPaintArgs& Args, const FGeometry&
 			OutDrawElements,
 			RetLayerId,
 			AllottedGeometry.ToPaintGeometry(
-			FVector2D(LinePos - 3., 0.),
-			FVector2D(6., AllottedGeometry.GetLocalSize().Y)),
+			FVector2f(6.f, AllottedGeometry.GetLocalSize().Y),
+			FSlateLayoutTransform(FVector2f(LinePos - 3.f, 0.0f))),
 			FillImage,
 			DrawEffects,
 			WarningTimeColor
@@ -278,8 +278,8 @@ int32 SVisualLoggerTimelineBar::OnPaint(const FPaintArgs& Args, const FGeometry&
 			OutDrawElements,
 			RetLayerId,
 			AllottedGeometry.ToPaintGeometry(
-			FVector2D(LinePos - 3., 0.),
-			FVector2D(6, AllottedGeometry.GetLocalSize().Y)),
+			FVector2f(6, AllottedGeometry.GetLocalSize().Y),
+			FSlateLayoutTransform(FVector2f(LinePos - 3, 0.0f))),
 			FillImage,
 			DrawEffects,
 			ErrorTimeColor
@@ -298,8 +298,8 @@ int32 SVisualLoggerTimelineBar::OnPaint(const FPaintArgs& Args, const FGeometry&
 			OutDrawElements,
 			++RetLayerId,
 			AllottedGeometry.ToPaintGeometry(
-			FVector2D(LinePos - 2., 0.),
-			FVector2D(4., AllottedGeometry.GetLocalSize().Y)),
+			FVector2f(4, AllottedGeometry.GetLocalSize().Y),
+			FSlateLayoutTransform(FVector2f(LinePos - 2, 0.0f))),
 			SelectedFillImage,
 			ESlateDrawEffect::None,
 			SelectedBarColor
diff --git a/Engine/Source/Developer/LogVisualizer/Private/VisualLoggerTimeSliderController.cpp b/Engine/Source/Developer/LogVisualizer/Private/VisualLoggerTimeSliderController.cpp
--- a/Engine/Source/Developer/LogVisualizer/Private/VisualLoggerTimeSliderController.cpp
+++ b/Engine/Source/Developer/LogVisualizer/Private/VisualLoggerTimeSliderController.cpp
@@ -170,7 +170,7 @@ void FVisualLoggerTimeSliderController::DrawTicks( FSlateWindowElementList& OutD
 			FSlateDrawElement::MakeLines(
 				OutDrawElements,
 				InArgs.StartLayer,
-				InArgs.AllottedGeometry.ToPaintGeometry( Offset, TickSize ),
+				InArgs.AllottedGeometry.ToPaintGeometry( TickSize, FSlateLayoutTransform(Offset) ),
 				LinePoints,
 				InArgs.DrawEffects,
 				InArgs.TickColor,
@@ -189,7 +189,7 @@ void FVisualLoggerTimeSliderController::DrawTicks( FSlateWindowElementList& OutD
 				FSlateDrawElement::MakeText(
 					OutDrawElements,
 					InArgs.StartLayer+1, 
-					InArgs.AllottedGeometry.ToPaintGeometry( TextOffset, TextSize ), 
+					InArgs.AllottedGeometry.ToPaintGeometry( TextSize, FSlateLayoutTransform(TextOffset) ), 
 					FrameString, 
 					SmallLayoutFont, 
 					InArgs.DrawEffects,
@@ -213,7 +213,7 @@ void FVisualLoggerTimeSliderController::DrawTicks( FSlateWindowElementList& OutD
 			FSlateDrawElement::MakeLines(
 				OutDrawElements,
 				InArgs.StartLayer,
-				InArgs.AllottedGeometry.ToPaintGeometry( Offset, TickSize ),
+				InArgs.AllottedGeometry.ToPaintGeometry( TickSize, FSlateLayoutTransform(Offset) ),
 				LinePoints,
 				InArgs.DrawEffects,
 				InArgs.TickColor,
@@ -266,7 +266,7 @@ int32 FVisualLoggerTimeSliderController::OnPaintTimeSlider( bool bMirrorLabels,
 		const float CursorHalfSize = TimeSliderArgs.CursorSize.Get() * 0.5f;
 		const int32 CursorLayer = LayerId + 2;
 		const double CursorHalfLength = AllottedGeometry.GetLocalSize().X * CursorHalfSize;
-		FPaintGeometry CursorGeometry = AllottedGeometry.ToPaintGeometry(FVector2D(XPos - CursorHalfLength, 0), FVector2D(2 * CursorHalfLength, AllottedGeometry.GetLocalSize().Y));
+		FPaintGeometry CursorGeometry = AllottedGeometry.ToPaintGeometry(FVector2f(2 * CursorHalfLength, AllottedGeometry.GetLocalSize().Y), FSlateLayoutTransform(FVector2f(XPos - CursorHalfLength, 0)));
 
 		FLinearColor CursorColor = InWidgetStyle.GetColorAndOpacityTint();
 		CursorColor.A = CursorColor.A*0.08f;
@@ -283,7 +283,7 @@ int32 FVisualLoggerTimeSliderController::OnPaintTimeSlider( bool bMirrorLabels,
 
 		// Should draw above the text
 		const int32 ArrowLayer = LayerId + 3;
-		FPaintGeometry MyGeometry =	AllottedGeometry.ToPaintGeometry( FVector2D( XPos-HalfSize, 0 ), FVector2D( HandleSize, AllottedGeometry.GetLocalSize().Y ) );
+		FPaintGeometry MyGeometry =	AllottedGeometry.ToPaintGeometry( FVector2f( HandleSize, AllottedGeometry.GetLocalSize().Y ) , FSlateLayoutTransform(FVector2f( XPos-HalfSize, 0 )));
 		FLinearColor ScrubColor = InWidgetStyle.GetColorAndOpacityTint();
 
 		// @todo Sequencer this color should be specified in the style
@@ -648,8 +648,8 @@ int32 FVisualLoggerTimeSliderController::OnPaintSectionView( const FGeometry& Al
 	{
 		// Draw cursor size
 		const float CursorHalfSize = TimeSliderArgs.CursorSize.Get() * 0.5f;
-		const double CursorHalfLength = AllottedGeometry.GetLocalSize().X * CursorHalfSize;
-		FPaintGeometry CursorGeometry = AllottedGeometry.ToPaintGeometry(FVector2D(LinePos - CursorHalfLength, 0.), FVector2D(2 * CursorHalfLength, AllottedGeometry.GetLocalSize().Y));
+		const float CursorHalfLength = AllottedGeometry.GetLocalSize().X * CursorHalfSize;
+		FPaintGeometry CursorGeometry = AllottedGeometry.ToPaintGeometry(FVector2f(2 * CursorHalfLength, AllottedGeometry.GetLocalSize().Y), FSlateLayoutTransform(FVector2f(LinePos - CursorHalfLength, 0)));
 
 		FSlateDrawElement::MakeBox(
 			OutDrawElements,
@@ -669,7 +669,7 @@ int32 FVisualLoggerTimeSliderController::OnPaintSectionView( const FGeometry& Al
 		FSlateDrawElement::MakeLines(
 			OutDrawElements,
 			++LayerId,
-			AllottedGeometry.ToPaintGeometry( FVector2D(LinePos, 0. ), FVector2D(1.,1.) ),
+			AllottedGeometry.ToPaintGeometry( FVector2f(1.0f,1.0f), FSlateLayoutTransform(FVector2f(LinePos, 0.0f )) ),
 			LinePoints,
 			DrawEffects,
 			FLinearColor::White.CopyWithNewOpacity(0.39f),
diff --git a/Engine/Source/Developer/ProfileVisualizer/Private/SGraphBar.cpp b/Engine/Source/Developer/ProfileVisualizer/Private/SGraphBar.cpp
--- a/Engine/Source/Developer/ProfileVisualizer/Private/SGraphBar.cpp
+++ b/Engine/Source/Developer/ProfileVisualizer/Private/SGraphBar.cpp
@@ -88,8 +88,8 @@ int32 SGraphBar::OnPaint( const FPaintArgs& Args, const FGeometry& AllottedGeome
 				OutDrawElements,
 				RetLayerId++,
 				AllottedGeometry.ToPaintGeometry(
-					FVector2D( StartX, 0.0f ),
-					FVector2D( EndX - StartX, AllottedGeometry.GetLocalSize().Y )),
+					FVector2f( EndX - StartX, AllottedGeometry.GetLocalSize().Y ),
+					FSlateLayoutTransform(FVector2f( StartX, 0.0f ))),
 				Event->IsSelected ? SelectedImage : FillImage,
 				DrawEffects,
 				Event->IsSelected ? SelectedBarColor : ColorPalette[Event->ColorIndex % (sizeof(ColorPalette) / sizeof(ColorPalette[0]))]			
diff --git a/Engine/Source/Developer/Profiler/Private/Widgets/SDataGraph.cpp b/Engine/Source/Developer/Profiler/Private/Widgets/SDataGraph.cpp
--- a/Engine/Source/Developer/Profiler/Private/Widgets/SDataGraph.cpp
+++ b/Engine/Source/Developer/Profiler/Private/Widgets/SDataGraph.cpp
@@ -313,7 +313,7 @@ int32 SDataGraph::OnPaint( const FPaintArgs& Args, const FGeometry& AllottedGeom
 	(
 		OutDrawElements,
 		LayerId,
-		AllottedGeometry.ToPaintGeometry( FVector2D(0,0), FVector2D(AreaX1,AllottedGeometry.GetLocalSize().Y) ),
+		AllottedGeometry.ToPaintGeometry( FVector2D(AreaX1,AllottedGeometry.GetLocalSize().Y), FSlateLayoutTransform() ),
 		TimelineAreaBrush,
 		DrawEffects,
 		TimelineAreaBrush->GetTint( InWidgetStyle ) * InWidgetStyle.GetColorAndOpacityTint()
@@ -838,7 +838,7 @@ int32 SDataGraph::OnPaint( const FPaintArgs& Args, const FGeometry& AllottedGeom
 				(
 					OutDrawElements,
 					LayerId,
-					AllottedGeometry.ToPaintGeometry( FVector2D(LocalGraphSelectionX[Nx]-HalfGraphMarkerWidth,0.0f), FVector2D(GraphMarkerWidth, AllottedGeometry.GetLocalSize().Y) ),
+					AllottedGeometry.ToPaintGeometry( FVector2D(GraphMarkerWidth, AllottedGeometry.GetLocalSize().Y), FSlateLayoutTransform(FVector2D(LocalGraphSelectionX[Nx]-HalfGraphMarkerWidth,0.0f)) ),
 					FAppStyle::GetBrush("Brushes.Foldout"),
 					DrawEffects,
 					FColor(64,64,255,128)
@@ -860,7 +860,7 @@ int32 SDataGraph::OnPaint( const FPaintArgs& Args, const FGeometry& AllottedGeom
 				(
 					OutDrawElements,
 					LayerId,
-					AllottedGeometry.ToPaintGeometry( FVector2D(GraphSelectionX0/*+HalfGraphMarkerWidth*/,0.0f), FVector2D(GraphSelectionW/*-GraphMarkerWidth*/, AllottedGeometry.GetLocalSize().Y) ),
+					AllottedGeometry.ToPaintGeometry( FVector2D(GraphSelectionW/*-GraphMarkerWidth*/, AllottedGeometry.GetLocalSize().Y), FSlateLayoutTransform(FVector2D(GraphSelectionX0/*+HalfGraphMarkerWidth*/,0.0f)) ),
 					FAppStyle::GetBrush("Brushes.Foldout"),
 					DrawEffects,
 					FColor(64,64,255,32)
@@ -880,7 +880,7 @@ int32 SDataGraph::OnPaint( const FPaintArgs& Args, const FGeometry& AllottedGeom
 		(
 			OutDrawElements,
 			LayerId,
-			AllottedGeometry.ToPaintGeometry( FVector2D(LocalPositionGraphX-HalfGraphMarkerWidth,0.0f), FVector2D(GraphMarkerWidth, AllottedGeometry.GetLocalSize().Y) ),
+			AllottedGeometry.ToPaintGeometry( FVector2D(GraphMarkerWidth, AllottedGeometry.GetLocalSize().Y), FSlateLayoutTransform(FVector2D(LocalPositionGraphX-HalfGraphMarkerWidth,0.0f)) ),
 			FAppStyle::Get().GetBrush("Brushes.Recessed"),
 			DrawEffects,
 			FColor(255,128,128,128)
diff --git a/Engine/Source/Developer/Profiler/Private/Widgets/SHistogram.cpp b/Engine/Source/Developer/Profiler/Private/Widgets/SHistogram.cpp
--- a/Engine/Source/Developer/Profiler/Private/Widgets/SHistogram.cpp
+++ b/Engine/Source/Developer/Profiler/Private/Widgets/SHistogram.cpp
@@ -49,7 +49,7 @@ int32 SHistogram::OnPaint( const FPaintArgs& Args, const FGeometry& AllottedGeom
 	(
 		OutDrawElements,
 		LayerId,
-		AllottedGeometry.ToPaintGeometry( FVector2D(0,0), FVector2D(AllottedGeometry.GetLocalSize().X,AllottedGeometry.GetLocalSize().Y) ),
+		AllottedGeometry.ToPaintGeometry( FVector2D(AllottedGeometry.GetLocalSize().X,AllottedGeometry.GetLocalSize().Y), FSlateLayoutTransform(FVector2D(0,0)) ),
 		TimelineAreaBrush,
 		DrawEffects,
 		TimelineAreaBrush->GetTint( InWidgetStyle ) * InWidgetStyle.GetColorAndOpacityTint()
@@ -159,7 +159,7 @@ int32 SHistogram::OnPaint( const FPaintArgs& Args, const FGeometry& AllottedGeom
 		FSlateDrawElement::MakeBox(
 			OutDrawElements,
 			LayerId,
-			AllottedGeometry.ToPaintGeometry( FVector2D(MarkerPosX, AllottedGeometry.GetLocalSize().Y - SizeY- LabelBuffer), FVector2D(CountX, SizeY) ),
+			AllottedGeometry.ToPaintGeometry( FVector2D(CountX, SizeY), FSlateLayoutTransform(FVector2D(MarkerPosX, AllottedGeometry.GetLocalSize().Y - SizeY- LabelBuffer)) ),
 			FillImage,
 			DrawEffects,
 			FLinearColor::Green
diff --git a/Engine/Source/Developer/Profiler/Private/Widgets/SProfilerMiniView.cpp b/Engine/Source/Developer/Profiler/Private/Widgets/SProfilerMiniView.cpp
--- a/Engine/Source/Developer/Profiler/Private/Widgets/SProfilerMiniView.cpp
+++ b/Engine/Source/Developer/Profiler/Private/Widgets/SProfilerMiniView.cpp
@@ -119,7 +119,7 @@ int32 SProfilerMiniView::OnPaint( const FPaintArgs& Args, const FGeometry& Allot
 	(
 		OutDrawElements,
 		LayerId,
-		AllottedGeometry.ToPaintGeometry( FVector2D( 0, 0 ), FVector2D( MiniViewSizeX, AllottedGeometry.Size.Y ) ),
+		AllottedGeometry.ToPaintGeometry( FVector2D( MiniViewSizeX, AllottedGeometry.Size.Y ), FSlateLayoutTransform() ),
 		MiniViewArea,
 		DrawEffects,
 		MiniViewArea->GetTint( InWidgetStyle ) * InWidgetStyle.GetColorAndOpacityTint()
@@ -155,7 +155,7 @@ int32 SProfilerMiniView::OnPaint( const FPaintArgs& Args, const FGeometry& Allot
 			(
 				OutDrawElements,
 				LayerId,
-				AllottedGeometry.ToPaintGeometry( FVector2D( DestSamplePosX0, MiniViewSizeY - GTSizeY ), FVector2D( DestSampleSizeX, GTSizeY ) ),
+				AllottedGeometry.ToPaintGeometry( FVector2D( DestSampleSizeX, GTSizeY ), FSlateLayoutTransform(FVector2D( DestSamplePosX0, MiniViewSizeY - GTSizeY )) ),
 				&SolidWhiteBrush,
 				DrawEffects,
 				GameThreadColor
@@ -166,7 +166,7 @@ int32 SProfilerMiniView::OnPaint( const FPaintArgs& Args, const FGeometry& Allot
 			(
 				OutDrawElements,
 				LayerId,
-				AllottedGeometry.ToPaintGeometry( FVector2D( DestSamplePosX0, MiniViewSizeY - GTSizeY - RTSizeY ), FVector2D( DestSampleSizeX, RTSizeY ) ),
+				AllottedGeometry.ToPaintGeometry( FVector2D( DestSampleSizeX, RTSizeY ), FSlateLayoutTransform(FVector2D( DestSamplePosX0, MiniViewSizeY - GTSizeY - RTSizeY )) ),
 				&SolidWhiteBrush,
 				DrawEffects,
 				RenderThreadColor
@@ -204,7 +204,7 @@ int32 SProfilerMiniView::OnPaint( const FPaintArgs& Args, const FGeometry& Allot
 			(
 				OutDrawElements,
 				LayerId,
-				AllottedGeometry.ToPaintGeometry( FVector2D( 0.0, 0.0 ), FVector2D( SelectionBoxX0, AllottedGeometry.Size.Y ) ),
+				AllottedGeometry.ToPaintGeometry( FVector2D( SelectionBoxX0, AllottedGeometry.Size.Y ), FSlateLayoutTransform(FVector2D( 0.0, 0.0 )) ),
 				&SolidWhiteBrush,
 				DrawEffects,
 				FColorList::Grey.WithAlpha( 192 )
@@ -217,7 +217,7 @@ int32 SProfilerMiniView::OnPaint( const FPaintArgs& Args, const FGeometry& Allot
 			(
 				OutDrawElements,
 				LayerId,
-				AllottedGeometry.ToPaintGeometry( FVector2D( SelectionBoxX1, 0.0 ), FVector2D( MiniViewSizeX - SelectionBoxX1, AllottedGeometry.Size.Y ) ),
+				AllottedGeometry.ToPaintGeometry( FVector2D( MiniViewSizeX - SelectionBoxX1, AllottedGeometry.Size.Y ), FSlateLayoutTransform(FVector2D( SelectionBoxX1, 0.0 )) ),
 				&SolidWhiteBrush,
 				DrawEffects,
 				FColorList::Grey.WithAlpha( 192 )
@@ -233,7 +233,7 @@ int32 SProfilerMiniView::OnPaint( const FPaintArgs& Args, const FGeometry& Allot
 			(
 				OutDrawElements,
 				LayerId,
-				AllottedGeometry.ToPaintGeometry( FVector2D( FillerPosX0, 0.0 ), FVector2D( FillerSizeX + 1.0, AllottedGeometry.Size.Y ) ),
+				AllottedGeometry.ToPaintGeometry( FVector2D( FillerSizeX + 1.0, AllottedGeometry.Size.Y ), FSlateLayoutTransform(FVector2D( FillerPosX0, 0.0 )) ),
 				&SolidWhiteBrush,
 				DrawEffects,
 				// #Profiler: 2014-04-09 How to get this color from Slate?
@@ -247,7 +247,7 @@ int32 SProfilerMiniView::OnPaint( const FPaintArgs& Args, const FGeometry& Allot
 		(
 			OutDrawElements,
 			LayerId,
-			AllottedGeometry.ToPaintGeometry( FVector2D( SelectionBoxX0, 0.0 ), FVector2D( SelectionBoxX1 - SelectionBoxX0, AllottedGeometry.Size.Y ) ),
+			AllottedGeometry.ToPaintGeometry( FVector2D( SelectionBoxX1 - SelectionBoxX0, AllottedGeometry.Size.Y ), FSlateLayoutTransform(FVector2D( SelectionBoxX0, 0.0 )) ),
 			FProfilerStyle::Get().GetBrush( "PlainBorder" ),
 			DrawEffects,
 			FColorList::Green
diff --git a/Engine/Source/Developer/Profiler/Private/Widgets/SProfilerThreadView.cpp b/Engine/Source/Developer/Profiler/Private/Widgets/SProfilerThreadView.cpp
--- a/Engine/Source/Developer/Profiler/Private/Widgets/SProfilerThreadView.cpp
+++ b/Engine/Source/Developer/Profiler/Private/Widgets/SProfilerThreadView.cpp
@@ -97,7 +97,7 @@ int32 SProfilerThreadView::OnPaint( const FPaintArgs& Args, const FGeometry& All
 	(
 		PaintState->OutDrawElements,
 		PaintState->LayerId,
-		PaintState->AllottedGeometry.ToPaintGeometry( FVector2D( 0, 0 ), PaintState->Size2D() ),
+		PaintState->AllottedGeometry.ToPaintGeometry( PaintState->Size2D(), FSlateLayoutTransform() ),
 		BackgroundBrush,
 		PaintState->DrawEffects,
 		BackgroundBrush->GetTint( InWidgetStyle ) * InWidgetStyle.GetColorAndOpacityTint()
@@ -190,7 +190,7 @@ void SProfilerThreadView::DrawFramesBackgroundAndTimelines() const
 			(
 				PaintState->OutDrawElements,
 				PaintState->LayerId,
-				PaintState->AllottedGeometry.ToPaintGeometry( ClippedFrameBackgroundRect.GetTopLeft(), ClippedFrameBackgroundRect.GetSize() ),
+				PaintState->AllottedGeometry.ToPaintGeometry( ClippedFrameBackgroundRect.GetSize(), FSlateLayoutTransform(ClippedFrameBackgroundRect.GetTopLeft()) ),
 				&SolidWhiteBrush,
 				PaintState->DrawEffects,
 				(ThreadNode.FrameIndex % 2) ? FColorList::White.WithAlpha( 64 ) : FColorList::White.WithAlpha( 128 )
@@ -287,7 +287,7 @@ void SProfilerThreadView::DrawUIStackNodes() const
 				(
 					PaintState->OutDrawElements,
 					PaintState->LayerId,
-					PaintState->AllottedGeometry.ToPaintGeometry( ClippedNodeRect.GetTopLeft(), ClippedNodeRect.GetSize() ),
+					PaintState->AllottedGeometry.ToPaintGeometry( ClippedNodeRect.GetSize(), FSlateLayoutTransform(ClippedNodeRect.GetTopLeft()) ),
 					BorderBrush,
 					PaintState->DrawEffects,
 					NodeColor
@@ -438,7 +438,7 @@ void SProfilerThreadView::DrawUIStackNodes_Recursively( const FProfilerUIStackNo
 		(
 			PaintState->OutDrawElements,
 			PaintState->LayerId,
-			PaintState->AllottedGeometry.ToPaintGeometry( Position, Size ),
+			PaintState->AllottedGeometry.ToPaintGeometry( Size, FSlateLayoutTransform(Position) ),
 			&SolidWhiteBrush,
 			PaintState->DrawEffects,
 			GameThreadColor
diff --git a/Engine/Source/Developer/Profiler/Private/Widgets/SProfilerThreadView.h b/Engine/Source/Developer/Profiler/Private/Widgets/SProfilerThreadView.h
--- a/Engine/Source/Developer/Profiler/Private/Widgets/SProfilerThreadView.h
+++ b/Engine/Source/Developer/Profiler/Private/Widgets/SProfilerThreadView.h
@@ -98,7 +98,7 @@ class SProfilerThreadView
 			, SummaryFont8Height( static_cast<float>(FontMeasureService->Measure( TEXT( "!" ), SummaryFont8 ).Y) )
 		{}
 
-		const FVector2D& Size2D() const
+		FVector2f Size2D() const
 		{
 			return AllottedGeometry.GetLocalSize();
 		}
diff --git a/Engine/Source/Developer/RealtimeProfiler/Private/SRealtimeProfilerLineGraph.cpp b/Engine/Source/Developer/RealtimeProfiler/Private/SRealtimeProfilerLineGraph.cpp
--- a/Engine/Source/Developer/RealtimeProfiler/Private/SRealtimeProfilerLineGraph.cpp
+++ b/Engine/Source/Developer/RealtimeProfiler/Private/SRealtimeProfilerLineGraph.cpp
@@ -121,7 +121,7 @@ int32 SRealtimeProfilerLineGraph::OnPaint( const FPaintArgs& Args, const FGeomet
 	FSlateDrawElement::MakeBox(
 		OutDrawElements,
 		LayerId,
-		AllottedGeometry.ToPaintGeometry( FVector2D(0,0), FVector2D(AllottedGeometry.GetLocalSize().X,AllottedGeometry.Size.Y) ),
+		AllottedGeometry.ToPaintGeometry( FVector2D(AllottedGeometry.GetLocalSize().X,AllottedGeometry.Size.Y), FSlateLayoutTransform() ),
 		TimelineAreaBrush,
 		DrawEffects,
 		TimelineAreaBrush->GetTint( InWidgetStyle ) * InWidgetStyle.GetColorAndOpacityTint()
diff --git a/Engine/Source/Developer/SlateReflector/Private/Models/NavigationSimulationNode.cpp b/Engine/Source/Developer/SlateReflector/Private/Models/NavigationSimulationNode.cpp
--- a/Engine/Source/Developer/SlateReflector/Private/Models/NavigationSimulationNode.cpp
+++ b/Engine/Source/Developer/SlateReflector/Private/Models/NavigationSimulationNode.cpp
@@ -112,6 +112,10 @@ TSharedRef<FJsonValue> FNavigationSimulationWidgetNode::ToJson(const FNavigation
 		{
 			return FWidgetReflectorNodeUtils::WidgetAddressToString(Value);
 		}
+		static TSharedRef<FJsonValue> CreateJsonValue(const UE::Slate::FDeprecateVector2DResult& InVec2D)
+		{
+			return CreateJsonValue(FVector2D(InVec2D));
+		}
 		static TSharedRef<FJsonValue> CreateJsonValue(const FVector2D& InVec2D)
 		{
 			TArray<TSharedPtr<FJsonValue>> StructJsonArray;
diff --git a/Engine/Source/Developer/SlateReflector/Private/Models/WidgetReflectorNode.cpp b/Engine/Source/Developer/SlateReflector/Private/Models/WidgetReflectorNode.cpp
--- a/Engine/Source/Developer/SlateReflector/Private/Models/WidgetReflectorNode.cpp
+++ b/Engine/Source/Developer/SlateReflector/Private/Models/WidgetReflectorNode.cpp
@@ -53,7 +53,7 @@ const FSlateRenderTransform& FWidgetReflectorNodeBase::GetAccumulatedRenderTrans
 	return WidgetGeometry.GetAccumulatedRenderTransform();
 }
 
-const FVector2D& FWidgetReflectorNodeBase::GetLocalSize() const
+FVector2f FWidgetReflectorNodeBase::GetLocalSize() const
 {
 	return WidgetGeometry.GetLocalSize();
 }
@@ -438,7 +438,13 @@ TSharedRef<FJsonValue> FSnapshotWidgetReflectorNode::ToJson(const TSharedRef<FSn
 			StructJsonArray.Add(MakeShareable(new FJsonValueNumber(InVec2D.Y)));
 			return MakeShareable(new FJsonValueArray(StructJsonArray));
 		}
-
+		static TSharedRef<FJsonValue> CreateVector2DJsonValue(const FVector2f& InVec2D)
+		{
+			TArray<TSharedPtr<FJsonValue>> StructJsonArray;
+			StructJsonArray.Add(MakeShareable(new FJsonValueNumber(InVec2D.X)));
+			StructJsonArray.Add(MakeShareable(new FJsonValueNumber(InVec2D.Y)));
+			return MakeShareable(new FJsonValueArray(StructJsonArray));
+		}
 		static TSharedRef<FJsonValue> CreateMatrix2x2JsonValue(const FMatrix2x2& InMatrix)
 		{
 			float m00, m01, m10, m11;
diff --git a/Engine/Source/Developer/SlateReflector/Private/Models/WidgetReflectorNode.h b/Engine/Source/Developer/SlateReflector/Private/Models/WidgetReflectorNode.h
--- a/Engine/Source/Developer/SlateReflector/Private/Models/WidgetReflectorNode.h
+++ b/Engine/Source/Developer/SlateReflector/Private/Models/WidgetReflectorNode.h
@@ -202,7 +202,7 @@ public:
 	/**
 	 * @return The local size of the widget we were initialized from
 	 */
-	const FVector2D& GetLocalSize() const;
+	FVector2f GetLocalSize() const;
 
 	/**
 	 * @return The basic hit-test of the widget we were initialized from
diff --git a/Engine/Source/Developer/SlateReflector/Private/Widgets/NavigationSimulationOverlay.cpp b/Engine/Source/Developer/SlateReflector/Private/Widgets/NavigationSimulationOverlay.cpp
--- a/Engine/Source/Developer/SlateReflector/Private/Widgets/NavigationSimulationOverlay.cpp
+++ b/Engine/Source/Developer/SlateReflector/Private/Widgets/NavigationSimulationOverlay.cpp
@@ -173,7 +173,7 @@ namespace NavigationSimulationOverlay
 				const FVector2D NewLocalOffset = RootDrawOffset + WidgetInfo.WidgetGeometry.GetAccumulatedLayoutTransform().GetTranslation();
 				const FVector2D NewLocalSize = TransformPoint(WidgetInfo.WidgetGeometry.GetAccumulatedLayoutTransform().GetScale(), WidgetInfo.WidgetGeometry.GetLocalSize());
 
-				Map.Add(WidgetInfo.WidgetPtr, AllottedGeometry.ToPaintGeometry(NewLocalOffset, NewLocalSize));
+				Map.Add(WidgetInfo.WidgetPtr, AllottedGeometry.ToPaintGeometry(NewLocalSize, FSlateLayoutTransform(NewLocalOffset)));
 				return Map.Find(WidgetInfo.WidgetPtr);
 			}
 
diff --git a/Engine/Source/Developer/SlateReflector/Private/Widgets/SAtlasVisualizer.cpp b/Engine/Source/Developer/SlateReflector/Private/Widgets/SAtlasVisualizer.cpp
--- a/Engine/Source/Developer/SlateReflector/Private/Widgets/SAtlasVisualizer.cpp
+++ b/Engine/Source/Developer/SlateReflector/Private/Widgets/SAtlasVisualizer.cpp
@@ -371,7 +371,7 @@ void SAtlasVisualizer::OnDrawViewport(const FGeometry& AllottedGeometry, const F
 		FSlateDrawElement::MakeBox(
 			OutDrawElements,
 			LayerId+1,
-			AllottedGeometry.ToPaintGeometry(FVector2D::Max((CurrentSlotRect.GetTopLeft()*ScrollPanel->GetZoomLevel())-FVector2D(1,1), FVector2D::ZeroVector), (CurrentSlotRect.GetSize() * ScrollPanel->GetZoomLevel())+FVector2D(1, 1)),
+			AllottedGeometry.ToPaintGeometry((CurrentSlotRect.GetSize() * ScrollPanel->GetZoomLevel())+FVector2D(1, 1), FSlateLayoutTransform(FVector2D::Max((CurrentSlotRect.GetTopLeft()*ScrollPanel->GetZoomLevel())-FVector2D(1,1), FVector2D::ZeroVector))),
 			HoveredSlotBorderBrush,
 			ESlateDrawEffect::None,
 			FLinearColor::Yellow
diff --git a/Engine/Source/Developer/SlateReflector/Private/Widgets/SWidgetSnapshotVisualizer.cpp b/Engine/Source/Developer/SlateReflector/Private/Widgets/SWidgetSnapshotVisualizer.cpp
--- a/Engine/Source/Developer/SlateReflector/Private/Widgets/SWidgetSnapshotVisualizer.cpp
+++ b/Engine/Source/Developer/SlateReflector/Private/Widgets/SWidgetSnapshotVisualizer.cpp
@@ -254,7 +254,7 @@ public:
 					FSlateDrawElement::MakeBox(
 						OutDrawElements,
 						++LayerId,
-						AllottedGeometry.ToPaintGeometry(RootDrawOffset + PickedWidget->GetAccumulatedLayoutTransform().GetTranslation(), TransformPoint(PickedWidget->GetAccumulatedLayoutTransform().GetScale(), PickedWidget->GetLocalSize())),
+						AllottedGeometry.ToPaintGeometry(TransformPoint(PickedWidget->GetAccumulatedLayoutTransform().GetScale(), PickedWidget->GetLocalSize()), FSlateLayoutTransform(RootDrawOffset + PickedWidget->GetAccumulatedLayoutTransform().GetTranslation())),
 						Brush,
 						ESlateDrawEffect::None,
 						FMath::Lerp(TopmostWidgetColor, LeafmostWidgetColor, ColorFactor)
@@ -268,7 +268,7 @@ public:
 					FSlateDrawElement::MakeBox(
 						OutDrawElements,
 						++LayerId,
-						AllottedGeometry.ToPaintGeometry(RootDrawOffset + SelectedWidget->GetAccumulatedLayoutTransform().GetTranslation(), TransformPoint(SelectedWidget->GetAccumulatedLayoutTransform().GetScale(), SelectedWidget->GetLocalSize())),
+						AllottedGeometry.ToPaintGeometry(TransformPoint(SelectedWidget->GetAccumulatedLayoutTransform().GetScale(), SelectedWidget->GetLocalSize()), FSlateLayoutTransform(RootDrawOffset + SelectedWidget->GetAccumulatedLayoutTransform().GetTranslation())),
 						Brush,
 						ESlateDrawEffect::None,
 						SelectedWidget->GetTint()
diff --git a/Engine/Source/Developer/ToolWidgets/Private/SSegmentedTimelineView.cpp b/Engine/Source/Developer/ToolWidgets/Private/SSegmentedTimelineView.cpp
--- a/Engine/Source/Developer/ToolWidgets/Private/SSegmentedTimelineView.cpp
+++ b/Engine/Source/Developer/ToolWidgets/Private/SSegmentedTimelineView.cpp
@@ -45,7 +45,7 @@ int32 SSegmentedTimelineView::PaintBlock(const FGeometry& AllottedGeometry, cons
 			float BoxMin = FMath::Max(0, RangeToScreen.InputToLocalX(LowerBound));
 			float BoxMax = FMath::Min(Size.X, RangeToScreen.InputToLocalX(UpperBound));
 
-			FPaintGeometry BoxGeometry = AllottedGeometry.ToPaintGeometry(FVector2D(BoxMin, 1), FVector2D(BoxMax - BoxMin, Size.Y - 2), 1);
+			FPaintGeometry BoxGeometry = AllottedGeometry.ToPaintGeometry(FVector2f(BoxMax - BoxMin, Size.Y - 2), FSlateLayoutTransform(FVector2f(BoxMin, 1.f)));
 
 			const FSlateBrush* Brush = FAppStyle::GetBrush("Sequencer.SectionArea.Background");
 
