NearestNeighborModel: fix check compatibility in MLDeformerModelInstance

diff --git a/Engine/Plugins/Animation/MLDeformer/MLDeformerFramework/Source/MLDeformerFramework/Private/MLDeformerModelInstance.cpp b/Engine/Plugins/Animation/MLDeformer/MLDeformerFramework/Source/MLDeformerFramework/Private/MLDeformerModelInstance.cpp
--- a/Engine/Plugins/Animation/MLDeformer/MLDeformerFramework/Source/MLDeformerFramework/Private/MLDeformerModelInstance.cpp
+++ b/Engine/Plugins/Animation/MLDeformer/MLDeformerFramework/Source/MLDeformerFramework/Private/MLDeformerModelInstance.cpp
@@ -151,9 +151,7 @@ FString UMLDeformerModelInstance::CheckCompatibility(USkeletalMeshComponent* InS
 	if (NeuralNetwork && NeuralNetwork->IsLoaded() && Model->GetDeformerAsset())
 	{
 		const int64 NumNeuralNetInputs = NeuralNetwork->GetInputTensor().Num();
-		const int32 NumFloatsPerBone = 6;
-		const int32 NumFloatsPerCurve = 1;
-		const int64 NumDeformerAssetInputs = static_cast<int64>(Model->GetInputInfo()->CalcNumNeuralNetInputs(NumFloatsPerBone, NumFloatsPerCurve));
+		const int64 NumDeformerAssetInputs = static_cast<int64>(Model->GetInputInfo()->CalcNumNeuralNetInputs(Model->GetNumFloatsPerBone(), Model->GetNumFloatsPerCurve()));
 		if (NumNeuralNetInputs != NumDeformerAssetInputs)
 		{
 			const FString InputErrorString = "The number of network inputs doesn't match the asset. Please retrain the asset."; 
@@ -364,9 +362,7 @@ bool UMLDeformerModelInstance::SetupInputs()
 
 	// If the neural network expects a different number of inputs, do nothing.
 	const int64 NumNeuralNetInputs = NeuralNetwork->GetInputTensorForContext(NeuralNetworkInferenceHandle).Num();
-	const int32 NumFloatsPerBone = 6;
-	const int32 NumFloatsPerCurve = 1;
-	const int64 NumDeformerAssetInputs = Model->GetInputInfo()->CalcNumNeuralNetInputs(NumFloatsPerBone, NumFloatsPerCurve);
+	const int64 NumDeformerAssetInputs = Model->GetInputInfo()->CalcNumNeuralNetInputs(Model->GetNumFloatsPerBone(), Model->GetNumFloatsPerCurve());
 	if (NumNeuralNetInputs != NumDeformerAssetInputs)
 	{
 		return false;
diff --git a/Engine/Plugins/Animation/MLDeformer/MLDeformerFramework/Source/MLDeformerFramework/Public/MLDeformerModel.h b/Engine/Plugins/Animation/MLDeformer/MLDeformerFramework/Source/MLDeformerFramework/Public/MLDeformerModel.h
--- a/Engine/Plugins/Animation/MLDeformer/MLDeformerFramework/Source/MLDeformerFramework/Public/MLDeformerModel.h
+++ b/Engine/Plugins/Animation/MLDeformer/MLDeformerFramework/Source/MLDeformerFramework/Public/MLDeformerModel.h
@@ -326,7 +326,7 @@ public:
 	 * Set the neural network object that we use during inference.
 	 * @param InNeuralNetwork The new neural network to use inside this deformer model.
 	 */
-	void SetNeuralNetwork(UNeuralNetwork* InNeuralNetwork);
+	virtual void SetNeuralNetwork(UNeuralNetwork* InNeuralNetwork);
 
 	/**
 	 * Get the neural network modified delegate.
@@ -507,6 +507,12 @@ public:
 	static FName GetDeltaCutoffLengthPropertyName()		{ return GET_MEMBER_NAME_CHECKED(UMLDeformerModel, DeltaCutoffLength); }
 #endif	// #if WITH_EDITORONLY_DATA
 
+	/** @return The number of floats per bone in network input. */
+	virtual int32 GetNumFloatsPerBone() const { return NumFloatsPerBone; }
+
+	/** @return The number of floats per curve in network input. */
+	virtual int32 GetNumFloatsPerCurve() const { return NumFloatsPerCurve; }
+
 protected:
 	/**
 	 * Set the training input information.
@@ -596,6 +602,12 @@ private:
 	UPROPERTY(EditAnywhere, Category = "Base Mesh")
 	TObjectPtr<USkeletalMesh> SkeletalMesh = nullptr;
 
+	/** The number of floats per bone in network input. */
+	static constexpr int32 NumFloatsPerBone = 6;
+
+	/** The number of floats per curve in network input. */
+	static constexpr int32 NumFloatsPerCurve = 1;
+
 #if WITH_EDITORONLY_DATA
 	UPROPERTY()
 	TObjectPtr<UMLDeformerVizSettings> VizSettings = nullptr;
diff --git a/Engine/Plugins/Animation/MLDeformer/NearestNeighborModel/Source/NearestNeighborModel/Public/NearestNeighborModel.h b/Engine/Plugins/Animation/MLDeformer/NearestNeighborModel/Source/NearestNeighborModel/Public/NearestNeighborModel.h
--- a/Engine/Plugins/Animation/MLDeformer/NearestNeighborModel/Source/NearestNeighborModel/Public/NearestNeighborModel.h
+++ b/Engine/Plugins/Animation/MLDeformer/NearestNeighborModel/Source/NearestNeighborModel/Public/NearestNeighborModel.h
@@ -235,8 +235,8 @@ public:
 
 	void ResetMorphBuffers();
 
-	virtual int32 GetNumFloatsPerBone() const { return NearestNeighborNumFloatsPerBone; }
-	virtual int32 GetNumFloatsPerCurve() const { return NearestNeighborNumFloatsPerCurve; }
+	virtual int32 GetNumFloatsPerBone() const override { return NearestNeighborNumFloatsPerBone; }
+	virtual int32 GetNumFloatsPerCurve() const override { return NearestNeighborNumFloatsPerCurve; }
 
 	bool DoesUseOptimizedNetwork() const;
 	UNearestNeighborOptimizedNetwork* GetOptimizedNetwork() const { return OptimizedNetwork.Get(); }
