MLDeformer: Implement GetShaderVirtualPath() for all optimus data interfaces that can support it.

diff --git a/Engine/Plugins/Animation/MLDeformer/MLDeformerFramework/Source/MLDeformerFramework/Private/MLDeformerGraphDataInterface.cpp b/Engine/Plugins/Animation/MLDeformer/MLDeformerFramework/Source/MLDeformerFramework/Private/MLDeformerGraphDataInterface.cpp
--- a/Engine/Plugins/Animation/MLDeformer/MLDeformerFramework/Source/MLDeformerFramework/Private/MLDeformerGraphDataInterface.cpp
+++ b/Engine/Plugins/Animation/MLDeformer/MLDeformerFramework/Source/MLDeformerFramework/Private/MLDeformerGraphDataInterface.cpp
@@ -57,6 +57,18 @@ void UMLDeformerGraphDataInterface::GetShaderParameters(TCHAR const* UID, FShade
 	InOutBuilder.AddNestedStruct<FMLDeformerGraphDataInterfaceParameters>(UID);
 }
 
+TCHAR const* UMLDeformerGraphDataInterface::TemplateFilePath = TEXT("/Plugin/MLDeformerFramework/Private/MLDeformerGraphDataInterface.ush");
+
+TCHAR const* UMLDeformerGraphDataInterface::GetShaderVirtualPath() const
+{
+	return TemplateFilePath;
+}
+
+void UMLDeformerGraphDataInterface::GetShaderHash(FString& InOutKey) const
+{
+	GetShaderFileHash(TemplateFilePath, EShaderPlatform::SP_PCD3D_SM5).AppendString(InOutKey);
+}
+
 void UMLDeformerGraphDataInterface::GetHLSL(FString& OutHLSL, FString const& InDataInterfaceName) const
 {
 	TMap<FString, FStringFormatArg> TemplateArgs =
@@ -65,7 +77,7 @@ void UMLDeformerGraphDataInterface::GetHLSL(FString& OutHLSL, FString const& InD
 	};
 
 	FString TemplateFile;
-	LoadShaderSourceFile(TEXT("/Plugin/MLDeformerFramework/Private/MLDeformerGraphDataInterface.ush"), EShaderPlatform::SP_PCD3D_SM5, &TemplateFile, nullptr);
+	LoadShaderSourceFile(TemplateFilePath, EShaderPlatform::SP_PCD3D_SM5, &TemplateFile, nullptr);
 	OutHLSL += FString::Format(*TemplateFile, TemplateArgs);
 }
 
diff --git a/Engine/Plugins/Animation/MLDeformer/MLDeformerFramework/Source/MLDeformerFramework/Private/MLDeformerGraphDebugDataInterface.cpp b/Engine/Plugins/Animation/MLDeformer/MLDeformerFramework/Source/MLDeformerFramework/Private/MLDeformerGraphDebugDataInterface.cpp
--- a/Engine/Plugins/Animation/MLDeformer/MLDeformerFramework/Source/MLDeformerFramework/Private/MLDeformerGraphDebugDataInterface.cpp
+++ b/Engine/Plugins/Animation/MLDeformer/MLDeformerFramework/Source/MLDeformerFramework/Private/MLDeformerGraphDebugDataInterface.cpp
@@ -74,6 +74,18 @@ void UMLDeformerGraphDebugDataInterface::GetShaderParameters(TCHAR const* UID, F
 	InOutBuilder.AddNestedStruct<FMLDeformerGraphDebugDataInterfaceParameters>(UID);
 }
 
+TCHAR const* UMLDeformerGraphDebugDataInterface::TemplateFilePath = TEXT("/Plugin/MLDeformerFramework/Private/MLDeformerGraphHeatMapDataInterface.ush");
+
+TCHAR const* UMLDeformerGraphDebugDataInterface::GetShaderVirtualPath() const
+{
+	return TemplateFilePath;
+}
+
+void UMLDeformerGraphDebugDataInterface::GetShaderHash(FString& InOutKey) const
+{
+	GetShaderFileHash(TemplateFilePath, EShaderPlatform::SP_PCD3D_SM5).AppendString(InOutKey);
+}
+
 void UMLDeformerGraphDebugDataInterface::GetHLSL(FString& OutHLSL, FString const& InDataInterfaceName) const
 {
 	TMap<FString, FStringFormatArg> TemplateArgs =
@@ -82,7 +94,7 @@ void UMLDeformerGraphDebugDataInterface::GetHLSL(FString& OutHLSL, FString const
 	};
 
 	FString TemplateFile;
-	LoadShaderSourceFile(TEXT("/Plugin/MLDeformerFramework/Private/MLDeformerGraphHeatMapDataInterface.ush"), EShaderPlatform::SP_PCD3D_SM5, &TemplateFile, nullptr);
+	LoadShaderSourceFile(TemplateFilePath, EShaderPlatform::SP_PCD3D_SM5, &TemplateFile, nullptr);
 	OutHLSL += FString::Format(*TemplateFile, TemplateArgs);
 }
 
diff --git a/Engine/Plugins/Animation/MLDeformer/MLDeformerFramework/Source/MLDeformerFramework/Public/MLDeformerGraphDataInterface.h b/Engine/Plugins/Animation/MLDeformer/MLDeformerFramework/Source/MLDeformerFramework/Public/MLDeformerGraphDataInterface.h
--- a/Engine/Plugins/Animation/MLDeformer/MLDeformerFramework/Source/MLDeformerFramework/Public/MLDeformerGraphDataInterface.h
+++ b/Engine/Plugins/Animation/MLDeformer/MLDeformerFramework/Source/MLDeformerFramework/Public/MLDeformerGraphDataInterface.h
@@ -69,9 +69,14 @@ public:
 	TCHAR const* GetClassName() const override { return TEXT("MLDeformer"); }
 	virtual void GetSupportedInputs(TArray<FShaderFunctionDefinition>& OutFunctions) const override;
 	virtual void GetShaderParameters(TCHAR const* UID, FShaderParametersMetadataBuilder& InOutBuilder, FShaderParametersMetadataAllocations& InOutAllocations) const override;
+	virtual TCHAR const* GetShaderVirtualPath() const override;
+	virtual void GetShaderHash(FString& InOutKey) const override;
 	virtual void GetHLSL(FString& OutHLSL, FString const& InDataInterfaceName) const override;
 	virtual UComputeDataProvider* CreateDataProvider(TObjectPtr<UObject> InBinding, uint64 InInputMask, uint64 InOutputMask) const override;
 	// ~END UComputeDataInterface overrides.
+
+private:
+	static TCHAR const* TemplateFilePath;
 };
 
 /** Compute Framework Data Provider for MLDeformer data. */
diff --git a/Engine/Plugins/Animation/MLDeformer/MLDeformerFramework/Source/MLDeformerFramework/Public/MLDeformerGraphDebugDataInterface.h b/Engine/Plugins/Animation/MLDeformer/MLDeformerFramework/Source/MLDeformerFramework/Public/MLDeformerGraphDebugDataInterface.h
--- a/Engine/Plugins/Animation/MLDeformer/MLDeformerFramework/Source/MLDeformerFramework/Public/MLDeformerGraphDebugDataInterface.h
+++ b/Engine/Plugins/Animation/MLDeformer/MLDeformerFramework/Source/MLDeformerFramework/Public/MLDeformerGraphDebugDataInterface.h
@@ -57,9 +57,14 @@ public:
 	TCHAR const* GetClassName() const override { return TEXT("MLDeformerDebug"); }
 	virtual void GetSupportedInputs(TArray<FShaderFunctionDefinition>& OutFunctions) const override;
 	virtual void GetShaderParameters(TCHAR const* UID, FShaderParametersMetadataBuilder& OutBuilder, FShaderParametersMetadataAllocations& InOutAllocations) const override;
+	virtual TCHAR const* GetShaderVirtualPath() const override;
+	virtual void GetShaderHash(FString& InOutKey) const override;
 	virtual void GetHLSL(FString& OutHLSL, FString const& InDataInterfaceName) const override;
 	virtual UComputeDataProvider* CreateDataProvider(TObjectPtr<UObject> InBinding, uint64 InInputMask, uint64 InOutputMask) const override;
 	// ~END UComputeDataInterface overrides.
+
+private:
+	static TCHAR const* TemplateFilePath;
 };
 
 /** Compute Framework Data Provider for MLDeformer debugging data. */
