Fix assert when loading ML Deformer Framework

diff --git a/Engine/Plugins/Animation/MLDeformer/MLDeformerFramework/Source/MLDeformerFramework/Private/MLDeformerModule.cpp b/Engine/Plugins/Animation/MLDeformer/MLDeformerFramework/Source/MLDeformerFramework/Private/MLDeformerModule.cpp
--- a/Engine/Plugins/Animation/MLDeformer/MLDeformerFramework/Source/MLDeformerFramework/Private/MLDeformerModule.cpp
+++ b/Engine/Plugins/Animation/MLDeformer/MLDeformerFramework/Source/MLDeformerFramework/Private/MLDeformerModule.cpp
@@ -13,14 +13,4 @@ MLDEFORMERFRAMEWORK_API DEFINE_LOG_CATEGORY(LogMLDeformer)
 
 IMPLEMENT_MODULE(UE::MLDeformer::FMLDeformerModule, MLDeformerFramework)
 
-namespace UE::MLDeformer
-{
-	void FMLDeformerModule::StartupModule()
-	{
-		// Register an additional shader path for our shaders used inside the deformer graph system.
-		const FString PluginShaderDir = FPaths::Combine(IPluginManager::Get().FindPlugin(TEXT("MLDeformerFramework"))->GetBaseDir(), TEXT("Shaders"));
-		AddShaderSourceDirectoryMapping(TEXT("/Plugin/MLDeformerFramework"), PluginShaderDir);
-	}
-}	// namespace UE::MLDeformer
-
 #undef LOCTEXT_NAMESPACE
diff --git a/Engine/Plugins/Animation/MLDeformer/MLDeformerFramework/Source/MLDeformerFramework/Public/MLDeformerModule.h b/Engine/Plugins/Animation/MLDeformer/MLDeformerFramework/Source/MLDeformerFramework/Public/MLDeformerModule.h
--- a/Engine/Plugins/Animation/MLDeformer/MLDeformerFramework/Source/MLDeformerFramework/Public/MLDeformerModule.h
+++ b/Engine/Plugins/Animation/MLDeformer/MLDeformerFramework/Source/MLDeformerFramework/Public/MLDeformerModule.h
@@ -17,8 +17,5 @@ namespace UE::MLDeformer
 		: public IModuleInterface
 	{
 	public:
-		// IModuleInterface overrides.
-		virtual void StartupModule() override;
-		// ~END IModuleInterface overrides.
 	};
 }	// namespace UE::MLDeformer
