Adding missing RHI includes

diff --git a/Engine/Plugins/Animation/MLDeformer/MLDeformerFramework/Source/MLDeformerFramework/Private/MLDeformerModel.cpp b/Engine/Plugins/Animation/MLDeformer/MLDeformerFramework/Source/MLDeformerFramework/Private/MLDeformerModel.cpp
--- a/Engine/Plugins/Animation/MLDeformer/MLDeformerFramework/Source/MLDeformerFramework/Private/MLDeformerModel.cpp
+++ b/Engine/Plugins/Animation/MLDeformer/MLDeformerFramework/Source/MLDeformerFramework/Private/MLDeformerModel.cpp
@@ -12,6 +12,7 @@
 #include "Rendering/SkeletalMeshLODModel.h"
 #include "Animation/AnimSequence.h"
 #include "UObject/UObjectGlobals.h"
+#include "RHICommandList.h"
 
 #include UE_INLINE_GENERATED_CPP_BY_NAME(MLDeformerModel)
 
diff --git a/Engine/Source/Runtime/RHI/Private/RHITextureReference.cpp b/Engine/Source/Runtime/RHI/Private/RHITextureReference.cpp
--- a/Engine/Source/Runtime/RHI/Private/RHITextureReference.cpp
+++ b/Engine/Source/Runtime/RHI/Private/RHITextureReference.cpp
@@ -1,6 +1,9 @@
 // Copyright Epic Games, Inc. All Rights Reserved.
 
 #include "RHITextureReference.h"
+#include "RHI.h"
+#include "RHICommandList.h"
+#include "RHIShaderPlatform.h"
 
 FRHITextureReference::FRHITextureReference()
 	: FRHITexture(RRT_TextureReference)
diff --git a/Engine/Source/Runtime/Renderer/Private/RayTracing/RayTracingScene.h b/Engine/Source/Runtime/Renderer/Private/RayTracing/RayTracingScene.h
--- a/Engine/Source/Runtime/Renderer/Private/RayTracing/RayTracingScene.h
+++ b/Engine/Source/Runtime/Renderer/Private/RayTracing/RayTracingScene.h
@@ -9,6 +9,7 @@
 
 #include "Async/TaskGraphInterfaces.h"
 #include "RHI.h"
+#include "RHIUtilities.h"
 #include "RenderGraphResources.h"
 #include "Misc/MemStack.h"
 #include "Containers/ArrayView.h"
