Fixed non unity/pch compile errors

diff --git a/Engine/Plugins/Animation/MLDeformer/MLDeformerFramework/Source/MLDeformerFramework/Private/MLDeformerModel.cpp b/Engine/Plugins/Animation/MLDeformer/MLDeformerFramework/Source/MLDeformerFramework/Private/MLDeformerModel.cpp
--- a/Engine/Plugins/Animation/MLDeformer/MLDeformerFramework/Source/MLDeformerFramework/Private/MLDeformerModel.cpp
+++ b/Engine/Plugins/Animation/MLDeformer/MLDeformerFramework/Source/MLDeformerFramework/Private/MLDeformerModel.cpp
@@ -11,6 +11,7 @@
 #include "Rendering/SkeletalMeshModel.h"
 #include "Rendering/SkeletalMeshLODModel.h"
 #include "Animation/AnimSequence.h"
+#include "RHICommandList.h"
 #include "UObject/UObjectGlobals.h"
 #include "RHICommandList.h"
 
diff --git a/Engine/Source/Developer/WidgetRegistration/Public/ToolkitBuilder.h b/Engine/Source/Developer/WidgetRegistration/Public/ToolkitBuilder.h
--- a/Engine/Source/Developer/WidgetRegistration/Public/ToolkitBuilder.h
+++ b/Engine/Source/Developer/WidgetRegistration/Public/ToolkitBuilder.h
@@ -2,11 +2,12 @@
 
 #pragma once
 
-#include "ToolMenus.h"
+#include "Framework/MultiBox/SToolBarButtonBlock.h"
 #include "Framework/Commands/UICommandInfo.h"
-#include "Widgets/SBoxPanel.h"
 #include "InteractiveToolManager.h"
 #include "ToolElementRegistry.h"
+#include "ToolMenus.h"
+#include "Widgets/SBoxPanel.h"
 
 
 class SWidget;
diff --git a/Engine/Source/Runtime/Slate/Public/Framework/MultiBox/MultiBoxBuilder.h b/Engine/Source/Runtime/Slate/Public/Framework/MultiBox/MultiBoxBuilder.h
--- a/Engine/Source/Runtime/Slate/Public/Framework/MultiBox/MultiBoxBuilder.h
+++ b/Engine/Source/Runtime/Slate/Public/Framework/MultiBox/MultiBoxBuilder.h
@@ -5,7 +5,6 @@
 #include "Framework/MultiBox/MultiBoxExtender.h"
 #include "Framework/SlateDelegates.h"
 #include "CoreMinimal.h"
-#include "SToolBarButtonBlock.h"
 #include "Misc/Attribute.h"
 #include "Layout/Visibility.h"
 #include "Widgets/SWidget.h"
