[MLDeformer]
- Made some fixes that could cause the MLD asset to be in a bad state, resulting in crashes and weird behavior. This could likely happen when aborting during training, or changing assets after training, and saving.
- Now setting correct section indices inside the morph targets, rather than just including all sections.
- Sorting the morph target deltas by vertex number. Perhaps this could help with some GPU caching.
- Detect changes in vertex count, and show a warning in the morph target section, when dynamic morph updates aren't possible until after retraining.
- Prevent these dynamic morph updates when it is not possible, which fixes a crash.
- Added some progress bars when creating and compressing morph targets.

diff --git a/Engine/Plugins/Animation/MLDeformer/MLDeformerFramework/Source/MLDeformerFramework/Private/MLDeformerMorphModel.cpp b/Engine/Plugins/Animation/MLDeformer/MLDeformerFramework/Source/MLDeformerFramework/Private/MLDeformerMorphModel.cpp
--- a/Engine/Plugins/Animation/MLDeformer/MLDeformerFramework/Source/MLDeformerFramework/Private/MLDeformerMorphModel.cpp
+++ b/Engine/Plugins/Animation/MLDeformer/MLDeformerFramework/Source/MLDeformerFramework/Private/MLDeformerMorphModel.cpp
@@ -56,6 +56,11 @@ void UMLDeformerMorphModel::PostLoad()
 #endif
 }
 
+bool UMLDeformerMorphModel::CanDynamicallyUpdateMorphTargets() const
+{
+	return GetMorphTargetDeltas().Num() == (GetNumBaseMeshVerts() * GetNumMorphTargets());
+}
+
 UMLDeformerModelInstance* UMLDeformerMorphModel::CreateModelInstance(UMLDeformerComponent* Component)
 {
 	return NewObject<UMLDeformerMorphModelInstance>(Component);
diff --git a/Engine/Plugins/Animation/MLDeformer/MLDeformerFramework/Source/MLDeformerFramework/Public/MLDeformerMorphModel.h b/Engine/Plugins/Animation/MLDeformer/MLDeformerFramework/Source/MLDeformerFramework/Public/MLDeformerMorphModel.h
--- a/Engine/Plugins/Animation/MLDeformer/MLDeformerFramework/Source/MLDeformerFramework/Public/MLDeformerMorphModel.h
+++ b/Engine/Plugins/Animation/MLDeformer/MLDeformerFramework/Source/MLDeformerFramework/Public/MLDeformerMorphModel.h
@@ -46,6 +46,8 @@ public:
 	EMLDeformerMaskChannel GetMaskChannel() const					{ return MaskChannel; }
 	bool GetInvertMaskChannel() const								{ return bInvertMaskChannel; }
 
+	bool CanDynamicallyUpdateMorphTargets() const;
+
 	void SetMorphDeltaZeroThreshold(float Threshold)				{ MorphDeltaZeroThreshold = Threshold; }
 	void SetMorphCompressionlevel(float Tolerance)					{ MorphCompressionLevel = Tolerance; }
 	void SetIncludeMorphTargetNormals(bool bInclude)				{ bIncludeNormals = bInclude; }
diff --git a/Engine/Plugins/Animation/MLDeformer/MLDeformerFramework/Source/MLDeformerFrameworkEditor/Private/MLDeformerEditorModel.cpp b/Engine/Plugins/Animation/MLDeformer/MLDeformerFramework/Source/MLDeformerFrameworkEditor/Private/MLDeformerEditorModel.cpp
--- a/Engine/Plugins/Animation/MLDeformer/MLDeformerFramework/Source/MLDeformerFrameworkEditor/Private/MLDeformerEditorModel.cpp
+++ b/Engine/Plugins/Animation/MLDeformer/MLDeformerFramework/Source/MLDeformerFrameworkEditor/Private/MLDeformerEditorModel.cpp
@@ -35,6 +35,7 @@
 #include "GeometryCacheMeshData.h"
 #include "GeometryCacheTrack.h"
 #include "BoneContainer.h"
+#include "Misc/ScopedSlowTask.h"
 
 #define LOCTEXT_NAMESPACE "MLDeformerEditorModel"
 
@@ -159,7 +160,6 @@ namespace UE::MLDeformer
 		SkelMeshComponent->SetVisibility(false);
 		SkelMeshComponent->MarkRenderStateDirty();
 
-
 		// Register the editor actor.
 		const FLinearColor LabelColor = FMLDeformerEditorStyle::Get().GetColor("MLDeformer.BaseMesh.LabelColor");
 		FMLDeformerEditorActor::FConstructSettings Settings;
@@ -1556,10 +1556,18 @@ namespace UE::MLDeformer
 		OnInputAssetsChanged();
 		OnPostInputAssetChanged();
 		UpdateMemoryUsage();
-		GetEditor()->GetModelDetailsView()->ForceRefresh();
+
+		if (Editor->GetModelDetailsView())
+		{
+			Editor->GetModelDetailsView()->ForceRefresh();
+		}
+
 		if (bRefreshVizSettings)
 		{
-			GetEditor()->GetVizSettingsDetailsView()->ForceRefresh();
+			if (Editor->GetVizSettingsDetailsView())
+			{
+				Editor->GetVizSettingsDetailsView()->ForceRefresh();
+			}
 		}
 	}
 
@@ -1652,12 +1660,24 @@ namespace UE::MLDeformer
 			return;
 		}
 
+		// Used to sort vertices.
+		struct FMLDeformerCompareMorphTargetDeltas
+		{
+			FORCEINLINE bool operator()(const FMorphTargetDelta& A, const FMorphTargetDelta& B) const
+			{
+				return A.SourceIdx < B.SourceIdx;
+			}
+		};
+
 		const int32 NumBaseMeshVerts = Model->GetNumBaseMeshVerts();
 		check(Deltas.Num() % NumBaseMeshVerts == 0);
 		const int32 NumMorphTargets = Deltas.Num() / NumBaseMeshVerts;
 		check((Deltas.Num() / NumMorphTargets) == NumBaseMeshVerts);
 		check(!Model->GetVertexMap().IsEmpty());
 
+		FScopedSlowTask Task(NumMorphTargets, LOCTEXT("CreateEngineMorphTargetProgress", "Creating morph targets"));
+		Task.MakeDialog(false);	
+
 		USkeletalMesh* SkelMesh = Model->GetSkeletalMesh();
 		FSkeletalMeshRenderData* RenderData = SkelMesh->GetResourceForRendering();
 		check(RenderData);
@@ -1666,6 +1686,8 @@ namespace UE::MLDeformer
 
 		// Calculate the normals for the base mesh.
 		const FSkeletalMeshModel* ImportedModel = SkelMesh->GetImportedModel();
+		check(ImportedModel);
+		check(ImportedModel->LODModels.IsValidIndex(LOD));
 		const TArrayView<const uint32> IndexArray = ImportedModel->LODModels[LOD].IndexBuffer;
 		const TArrayView<const int32> VertexMap = ImportedModel->LODModels[LOD].MeshToImportVertexMap;
 		const TArrayView<const FVector3f> BaseVertexPositions = Sampler->GetUnskinnedVertexPositions();
@@ -1702,15 +1724,6 @@ namespace UE::MLDeformer
 			MorphLODModel.Reset();
 			MorphLODModel.bGeneratedByEngine = true;
 			MorphLODModel.NumBaseMeshVerts = NumRenderVertices;
-			MorphLODModel.NumVertices = NumRenderVertices;
-
-			// Init sections.
-			const int32 NumSections = RenderData->LODRenderData[LOD].RenderSections.Num();
-			MorphLODModel.SectionIndices.AddUninitialized(NumSections);
-			for (int32 SectionIndex = 0; SectionIndex < NumSections; ++SectionIndex)
-			{
-				MorphLODModel.SectionIndices[SectionIndex] = SectionIndex;
-			}
 
 			// Init deltas for this morph target.
 			MorphLODModel.Vertices.Reserve(NumRenderVertices);
@@ -1740,24 +1753,37 @@ namespace UE::MLDeformer
 						}
 					}
 
-					const FVector3f Delta = Deltas[ImportedVertexNumber + MorphTargetIndex * NumBaseMeshVerts];
-					if (Delta.Length() * VertexWeight > DeltaThreshold)
+					// Update the sections.
+					int32 SectionIndex = INDEX_NONE;
+					int32 SectionVertexIndex = INDEX_NONE;
+					ImportedModel->LODModels[LOD].GetSectionFromVertexIndex(VertexIndex, SectionIndex, SectionVertexIndex);
+
+					const FVector3f Delta = Deltas[ImportedVertexNumber + MorphTargetIndex * NumBaseMeshVerts] * VertexWeight;
+					if (Delta.Length() > DeltaThreshold)
 					{
 						MorphLODModel.Vertices.AddDefaulted();
 						FMorphTargetDelta& MorphTargetDelta = MorphLODModel.Vertices.Last();
-						MorphTargetDelta.PositionDelta = Delta * VertexWeight;
+						MorphTargetDelta.PositionDelta = Delta;
 						MorphTargetDelta.SourceIdx = VertexIndex;
 						MorphTargetDelta.TangentZDelta = bIncludeNormals ? DeltaNormals[ImportedVertexNumber] * VertexWeight : FVector3f::ZeroVector;
+
+						MorphLODModel.SectionIndices.AddUnique(SectionIndex);
 					}
 				}
 			}
 
 			MorphLODModel.Vertices.Shrink();
+			MorphLODModel.Vertices.Sort(FMLDeformerCompareMorphTargetDeltas());
+			MorphLODModel.NumVertices = MorphLODModel.Vertices.Num();
+			Task.EnterProgressFrame();
 		}
 	}
 
 	void FMLDeformerEditorModel::CompressEngineMorphTargets(FMorphTargetVertexInfoBuffers& OutMorphBuffers, const TArray<UMorphTarget*>& MorphTargets, int32 LOD, float MorphErrorTolerance)
 	{
+		FScopedSlowTask Task(2, LOCTEXT("CompressEngineMorphTargetProgress", "Compressing morph targets"));
+		Task.MakeDialog(false);	
+
 		USkeletalMesh* SkelMesh = Model->GetSkeletalMesh();
 		FSkeletalMeshRenderData* RenderData = SkelMesh->GetResourceForRendering();
 		check(RenderData);
@@ -1785,6 +1811,8 @@ namespace UE::MLDeformer
 			MorphErrorTolerance
 		);
 
+		Task.EnterProgressFrame();
+
 		// Reinit the render resources.
 		if (OutMorphBuffers.IsMorphCPUDataValid() && OutMorphBuffers.GetNumMorphs() > 0 && OutMorphBuffers.GetNumBatches() > 0)
 		{
@@ -1803,6 +1831,7 @@ namespace UE::MLDeformer
 		}
 
 		UpdateMemoryUsage();
+		Task.EnterProgressFrame();
 	}
 
 	void FMLDeformerEditorModel::DrawMorphTarget(FPrimitiveDrawInterface* PDI, const TArray<FVector3f>& MorphDeltas, float DeltaThreshold, int32 MorphTargetIndex, const FVector& DrawOffset)
diff --git a/Engine/Plugins/Animation/MLDeformer/MLDeformerFramework/Source/MLDeformerFrameworkEditor/Private/MLDeformerEditorToolkit.cpp b/Engine/Plugins/Animation/MLDeformer/MLDeformerFramework/Source/MLDeformerFrameworkEditor/Private/MLDeformerEditorToolkit.cpp
--- a/Engine/Plugins/Animation/MLDeformer/MLDeformerFramework/Source/MLDeformerFrameworkEditor/Private/MLDeformerEditorToolkit.cpp
+++ b/Engine/Plugins/Animation/MLDeformer/MLDeformerFramework/Source/MLDeformerFrameworkEditor/Private/MLDeformerEditorToolkit.cpp
@@ -388,11 +388,11 @@ namespace UE::MLDeformer
 						}
 
 						// Initialize the training inputs.
-						ActiveModel->InitInputInfo(Model->GetInputInfo());
+						ActiveModel->UpdateEditorInputInfo();
 
 						// Make sure we have something to train on.
 						// If this triggers, the train button most likely was enabled while it shouldn't be.
-						check(!Model->GetInputInfo()->IsEmpty());
+						check(!ActiveModel->GetEditorInputInfo()->IsEmpty());
 
 						// Train the model, which executes the Python code's "train" function.
 						const double StartTime = FPlatformTime::Seconds();
@@ -492,6 +492,7 @@ namespace UE::MLDeformer
 					FFormatNamedArguments SuccessArgs;
 					SuccessArgs.Add(TEXT("Duration"), TrainingDurationText);
 					WindowMessage = FText::Format(LOCTEXT("TrainingSuccess", "Training completed successfully!\n\nTraining time: {Duration}"), SuccessArgs);
+					ActiveModel->InitInputInfo(ActiveModel->GetModel()->GetInputInfo());
 					bMarkDirty = true;
 				}
 			}
@@ -516,6 +517,7 @@ namespace UE::MLDeformer
 					else
 					{
 						ShowNotification(LOCTEXT("PartialTrainingSuccess", "Training partially completed!"), SNotificationItem::ECompletionState::CS_Success, true);
+						ActiveModel->InitInputInfo(ActiveModel->GetModel()->GetInputInfo());
 						bMarkDirty = true;
 					}
 				}
@@ -637,7 +639,7 @@ namespace UE::MLDeformer
 			ActiveModel->SetDefaultDeformerGraphIfNeeded();
 			ActiveModel->CreateActors(InPersonaPreviewScene);
 			ActiveModel->UpdateActorVisibility();
-			ActiveModel->OnInputAssetsChanged();
+			ActiveModel->TriggerInputAssetChanged(false);
 			ActiveModel->CreateHeatMapAssets();
 			ActiveModel->SetHeatMapMaterialEnabled(ActiveModel->GetModel()->GetVizSettings()->GetShowHeatMap());
 			ActiveModel->SetResamplingInputOutputsNeeded(true);
diff --git a/Engine/Plugins/Animation/MLDeformer/MLDeformerFramework/Source/MLDeformerFrameworkEditor/Private/MLDeformerMorphModelDetails.cpp b/Engine/Plugins/Animation/MLDeformer/MLDeformerFramework/Source/MLDeformerFrameworkEditor/Private/MLDeformerMorphModelDetails.cpp
--- a/Engine/Plugins/Animation/MLDeformer/MLDeformerFramework/Source/MLDeformerFrameworkEditor/Private/MLDeformerMorphModelDetails.cpp
+++ b/Engine/Plugins/Animation/MLDeformer/MLDeformerFramework/Source/MLDeformerFrameworkEditor/Private/MLDeformerMorphModelDetails.cpp
@@ -7,6 +7,8 @@
 #include "DetailCategoryBuilder.h"
 #include "DetailWidgetRow.h"
 #include "IDetailGroup.h"
+#include "SWarningOrErrorBox.h"
+#include "Widgets/Layout/SBox.h"
 
 #define LOCTEXT_NAMESPACE "MLDeformerMorphModelDetails"
 
@@ -48,6 +50,22 @@ namespace UE::MLDeformer
 		MaskGroup.AddPropertyRow(DetailLayoutBuilder->GetProperty(UMLDeformerMorphModel::GetInvertMaskChannelPropertyName(), UMLDeformerMorphModel::StaticClass()));
 
 		MorphTargetCategoryBuilder->AddProperty(DetailLayoutBuilder->GetProperty(UMLDeformerMorphModel::GetQualityLevelsPropertyName(), UMLDeformerMorphModel::StaticClass()));
+
+		if (MorphModel && !MorphModel->CanDynamicallyUpdateMorphTargets())
+		{
+			const FText DeltaCountMismatchErrorText = LOCTEXT("MorphDeltaCountMismatch", "Dynamic morph target updates disabled until retrained. This is because the vertex count changed after the model was trained.");
+			FDetailWidgetRow& DeltaMismatchErrorRow = MorphTargetCategoryBuilder->AddCustomRow(FText::FromString("MorphDeltaCountMismatchError"))
+				.WholeRowContent()
+				[
+					SNew(SBox)
+					.Padding(FMargin(0.0f, 4.0f))
+					[
+						SNew(SWarningOrErrorBox)
+						.MessageStyle(EMessageStyle::Warning)
+						.Message(DeltaCountMismatchErrorText)
+					]
+				];
+		}
 	}
 }	// namespace UE::MLDeformer
 
diff --git a/Engine/Plugins/Animation/MLDeformer/MLDeformerFramework/Source/MLDeformerFrameworkEditor/Private/MLDeformerMorphModelEditorModel.cpp b/Engine/Plugins/Animation/MLDeformer/MLDeformerFramework/Source/MLDeformerFrameworkEditor/Private/MLDeformerMorphModelEditorModel.cpp
--- a/Engine/Plugins/Animation/MLDeformer/MLDeformerFramework/Source/MLDeformerFrameworkEditor/Private/MLDeformerMorphModelEditorModel.cpp
+++ b/Engine/Plugins/Animation/MLDeformer/MLDeformerFramework/Source/MLDeformerFrameworkEditor/Private/MLDeformerMorphModelEditorModel.cpp
@@ -47,7 +47,10 @@ namespace UE::MLDeformer
 		{
 			if (PropertyChangedEvent.ChangeType == EPropertyChangeType::ValueSet)
 			{
-				InitEngineMorphTargets(GetMorphModel()->GetMorphTargetDeltas());
+				if (GetMorphModel()->CanDynamicallyUpdateMorphTargets())
+				{
+					InitEngineMorphTargets(GetMorphModel()->GetMorphTargetDeltas());
+				}
 			}
 		}
 		else if (Property->GetFName() == UMLDeformerMorphModelVizSettings::GetMorphTargetNumberPropertyName())
diff --git a/Engine/Plugins/Animation/MLDeformer/MLDeformerFramework/Source/MLDeformerFrameworkEditor/Private/MLDeformerTrainingModel.cpp b/Engine/Plugins/Animation/MLDeformer/MLDeformerFramework/Source/MLDeformerFrameworkEditor/Private/MLDeformerTrainingModel.cpp
--- a/Engine/Plugins/Animation/MLDeformer/MLDeformerFramework/Source/MLDeformerFrameworkEditor/Private/MLDeformerTrainingModel.cpp
+++ b/Engine/Plugins/Animation/MLDeformer/MLDeformerFramework/Source/MLDeformerFrameworkEditor/Private/MLDeformerTrainingModel.cpp
@@ -28,12 +28,12 @@ UE::MLDeformer::FMLDeformerEditorModel* UMLDeformerTrainingModel::GetEditorModel
 
 int32 UMLDeformerTrainingModel::GetNumberSampleTransforms() const
 {
-	return GetModel()->GetInputInfo()->GetNumBones();
+	return EditorModel->GetEditorInputInfo()->GetNumBones();
 }
 
 int32 UMLDeformerTrainingModel::GetNumberSampleCurves() const
 {
-	return GetModel()->GetInputInfo()->GetNumCurves();
+	return EditorModel->GetEditorInputInfo()->GetNumCurves();
 }
 
 int32 UMLDeformerTrainingModel::NumSamples() const
@@ -43,7 +43,7 @@ int32 UMLDeformerTrainingModel::NumSamples() const
 
 int32 UMLDeformerTrainingModel::GetNumberSampleDeltas() const
 {
-	return GetModel()->GetNumBaseMeshVerts();
+	return EditorModel->GetEditorInputInfo()->GetNumBaseMeshVertices();
 }
 
 void UMLDeformerTrainingModel::SetNumFloatsPerCurve(int32 NumFloatsPerCurve)
