[MLDeformer]
- Fixed heatmaps.
- Fixed a crash when enabling morph target visualization.
- Fixed broken training when creating a new asset from scratch.
- Fixed a crash when changing MLD model type.

diff --git a/Engine/Plugins/Animation/MLDeformer/MLDeformerFramework/Shaders/Private/MLDeformerModelHeatMap.ush b/Engine/Plugins/Animation/MLDeformer/MLDeformerFramework/Shaders/Private/MLDeformerModelHeatMap.ush
new file mode 100644
--- /dev/null
+++ b/Engine/Plugins/Animation/MLDeformer/MLDeformerFramework/Shaders/Private/MLDeformerModelHeatMap.ush
@@ -0,0 +1,41 @@
+// Copyright Epic Games, Inc. All Rights Reserved.
+
+uint {DataInterfaceName}_NumVertices;
+uint {DataInterfaceName}_InputStreamStart;
+uint {DataInterfaceName}_HeatMapMode;
+float {DataInterfaceName}_HeatMapMax;
+float {DataInterfaceName}_GroundTruthLerp;
+uint {DataInterfaceName}_GroundTruthBufferSize;
+StructuredBuffer<float> {DataInterfaceName}_PositionGroundTruthBuffer;
+Buffer<uint> {DataInterfaceName}_VertexMapBuffer;
+
+uint ReadNumVertices_{DataInterfaceName}()
+{
+	return {DataInterfaceName}_NumVertices;
+}
+
+float3 ReadPositionGroundTruth_{DataInterfaceName}(uint VertexIndex)
+{
+	uint BufferIndex = {DataInterfaceName}_VertexMapBuffer[VertexIndex + {DataInterfaceName}_InputStreamStart];
+	if (BufferIndex >= {DataInterfaceName}_GroundTruthBufferSize)
+	{
+		return float3(0, 0, 0);
+	}
+	BufferIndex *= 3;
+	return float3({DataInterfaceName}_PositionGroundTruthBuffer[BufferIndex], {DataInterfaceName}_PositionGroundTruthBuffer[BufferIndex + 1], {DataInterfaceName}_PositionGroundTruthBuffer[BufferIndex + 2]);
+}
+
+uint ReadHeatMapMode_{DataInterfaceName}()
+{
+	return {DataInterfaceName}_HeatMapMode;
+}
+
+float ReadHeatMapMax_{DataInterfaceName}()
+{
+	return {DataInterfaceName}_HeatMapMax;
+}
+
+float ReadGroundTruthLerp_{DataInterfaceName}()
+{
+	return {DataInterfaceName}_GroundTruthLerp;
+}
diff --git a/Engine/Plugins/Animation/MLDeformer/MLDeformerFramework/Source/MLDeformerFramework/Private/MLDeformerGraphDebugDataInterface.cpp b/Engine/Plugins/Animation/MLDeformer/MLDeformerFramework/Source/MLDeformerFramework/Private/MLDeformerGraphDebugDataInterface.cpp
new file mode 100644
--- /dev/null
+++ b/Engine/Plugins/Animation/MLDeformer/MLDeformerFramework/Source/MLDeformerFramework/Private/MLDeformerGraphDebugDataInterface.cpp
@@ -0,0 +1,227 @@
+// Copyright Epic Games, Inc. All Rights Reserved.
+
+#include "MLDeformerGraphDebugDataInterface.h"
+#include "MLDeformerModelInstance.h"
+#include "MLDeformerAsset.h"
+#include "MLDeformerComponent.h"
+#include "MLDeformerVizSettings.h"
+#include "Components/SkeletalMeshComponent.h"
+#include "ComputeFramework/ComputeKernelPermutationSet.h"
+#include "ComputeFramework/ShaderParamTypeDefinition.h"
+#include "OptimusDataDomain.h"
+#include "RenderGraphBuilder.h"
+#include "RenderGraphUtils.h"
+#include "RenderGraphResources.h"
+#include "Rendering/SkeletalMeshModel.h"
+#include "Rendering/SkeletalMeshLODRenderData.h"
+#include "Rendering/SkeletalMeshRenderData.h"
+#include "ShaderCompilerCore.h"
+#include "ShaderParameterMetadataBuilder.h"
+#include "SkeletalRenderPublic.h"
+
+#include UE_INLINE_GENERATED_CPP_BY_NAME(MLDeformerGraphDebugDataInterface)
+
+TArray<FOptimusCDIPinDefinition> UMLDeformerGraphDebugDataInterface::GetPinDefinitions() const
+{
+	TArray<FOptimusCDIPinDefinition> Defs;
+	Defs.Add({ "HeatMapMode", "ReadHeatMapMode" });
+	Defs.Add({ "HeatMapMax", "ReadHeatMapMax" });
+	Defs.Add({ "GroundTruthLerp", "ReadGroundTruthLerp" });
+	Defs.Add({ "PositionGroundTruth", "ReadPositionGroundTruth", Optimus::DomainName::Vertex, "ReadNumVertices" });
+	return Defs;
+}
+
+TSubclassOf<UActorComponent> UMLDeformerGraphDebugDataInterface::GetRequiredComponentClass() const
+{
+	return UMLDeformerComponent::StaticClass();
+}
+
+void UMLDeformerGraphDebugDataInterface::GetSupportedInputs(TArray<FShaderFunctionDefinition>& OutFunctions) const
+{
+	OutFunctions.AddDefaulted_GetRef()
+		.SetName(TEXT("ReadNumVertices"))
+		.AddReturnType(EShaderFundamentalType::Uint);
+
+	OutFunctions.AddDefaulted_GetRef()
+		.SetName(TEXT("ReadHeatMapMode"))
+		.AddReturnType(EShaderFundamentalType::Int);
+
+	OutFunctions.AddDefaulted_GetRef()
+		.SetName(TEXT("ReadHeatMapMax"))
+		.AddReturnType(EShaderFundamentalType::Float);
+
+	OutFunctions.AddDefaulted_GetRef()
+		.SetName(TEXT("ReadGroundTruthLerp"))
+		.AddReturnType(EShaderFundamentalType::Float);
+
+	OutFunctions.AddDefaulted_GetRef()
+		.SetName(TEXT("ReadPositionGroundTruth"))
+		.AddReturnType(EShaderFundamentalType::Float, 3)
+		.AddParam(EShaderFundamentalType::Uint);
+}
+
+BEGIN_SHADER_PARAMETER_STRUCT(FMLDeformerGraphDebugDataInterfaceParameters, )
+	SHADER_PARAMETER(uint32, NumVertices)
+	SHADER_PARAMETER(uint32, InputStreamStart)
+	SHADER_PARAMETER(int32, HeatMapMode)
+	SHADER_PARAMETER(float, HeatMapMax)
+	SHADER_PARAMETER(float, GroundTruthLerp)
+	SHADER_PARAMETER(uint32, GroundTruthBufferSize)
+	SHADER_PARAMETER_RDG_BUFFER_SRV(StructuredBuffer<float>, PositionGroundTruthBuffer)
+	SHADER_PARAMETER_SRV(Buffer<uint>, VertexMapBuffer)
+END_SHADER_PARAMETER_STRUCT()
+
+FString UMLDeformerGraphDebugDataInterface::GetDisplayName() const
+{
+	return TEXT("MLD Model Debug");
+}
+
+void UMLDeformerGraphDebugDataInterface::GetShaderParameters(TCHAR const* UID, FShaderParametersMetadataBuilder& InOutBuilder, FShaderParametersMetadataAllocations& InOutAllocations) const
+{
+	InOutBuilder.AddNestedStruct<FMLDeformerGraphDebugDataInterfaceParameters>(UID);
+}
+
+TCHAR const* UMLDeformerGraphDebugDataInterface::TemplateFilePath = TEXT("/Plugin/MLDeformerFramework/Private/MLDeformerModelHeatMap.ush");
+
+TCHAR const* UMLDeformerGraphDebugDataInterface::GetShaderVirtualPath() const
+{
+	return TemplateFilePath;
+}
+
+void UMLDeformerGraphDebugDataInterface::GetShaderHash(FString& InOutKey) const
+{
+	GetShaderFileHash(TemplateFilePath, EShaderPlatform::SP_PCD3D_SM5).AppendString(InOutKey);
+}
+
+void UMLDeformerGraphDebugDataInterface::GetHLSL(FString& OutHLSL, FString const& InDataInterfaceName) const
+{
+	TMap<FString, FStringFormatArg> TemplateArgs =
+	{
+		{ TEXT("DataInterfaceName"), InDataInterfaceName },
+	};
+
+	FString TemplateFile;
+	LoadShaderSourceFile(TemplateFilePath, EShaderPlatform::SP_PCD3D_SM5, &TemplateFile, nullptr);
+	OutHLSL += FString::Format(*TemplateFile, TemplateArgs);
+}
+
+UComputeDataProvider* UMLDeformerGraphDebugDataInterface::CreateDataProvider(TObjectPtr<UObject> InBinding, uint64 InInputMask, uint64 InOutputMask) const
+{
+	UMLDeformerGraphDebugDataProvider* Provider = NewObject<UMLDeformerGraphDebugDataProvider>();
+	Provider->DeformerComponent = Cast<UMLDeformerComponent>(InBinding);
+	if (Provider->DeformerComponent)
+	{
+		Provider->DeformerAsset = Provider->DeformerComponent->GetDeformerAsset();
+	}
+
+	#if WITH_EDITORONLY_DATA
+		if (Provider->DeformerAsset != nullptr)
+		{
+			Provider->Init();
+		}
+	#endif
+
+	return Provider;
+}
+
+FComputeDataProviderRenderProxy* UMLDeformerGraphDebugDataProvider::GetRenderProxy()
+{
+#if WITH_EDITORONLY_DATA
+	if (DeformerComponent && DeformerAsset && DeformerComponent->GetModelInstance() && DeformerComponent->GetModelInstance()->IsValidForDataProvider())
+	{
+		UE::MLDeformer::FMLDeformerGraphDebugDataProviderProxy* Proxy = new UE::MLDeformer::FMLDeformerGraphDebugDataProviderProxy(DeformerComponent, DeformerAsset, this);
+		const float SampleTime = DeformerComponent->GetModelInstance()->GetSkeletalMeshComponent()->GetPosition();
+		UMLDeformerModel* Model = DeformerAsset->GetModel();
+		Model->SampleGroundTruthPositions(SampleTime, Proxy->GetGroundTruthPositions());
+		Proxy->HandleZeroGroundTruthPositions();
+		return Proxy;
+	}
+#endif
+
+	// Return default invalid proxy.
+	return new FComputeDataProviderRenderProxy();
+}
+
+#if WITH_EDITORONLY_DATA
+namespace UE::MLDeformer
+{
+	FMLDeformerGraphDebugDataProviderProxy::FMLDeformerGraphDebugDataProviderProxy(UMLDeformerComponent* DeformerComponent, UMLDeformerAsset* DeformerAsset, UMLDeformerGraphDebugDataProvider* InProvider)
+		: FComputeDataProviderRenderProxy()
+	{
+		Provider = InProvider;
+
+		if (DeformerComponent && DeformerAsset)
+		{
+			const UMLDeformerModel* Model = DeformerAsset->GetModel();	
+			const UMLDeformerVizSettings* VizSettings = Model->GetVizSettings();
+			const UMLDeformerModelInstance* ModelInstance = DeformerComponent->GetModelInstance();
+
+			SkeletalMeshObject = ModelInstance->GetSkeletalMeshComponent()->MeshObject;
+			VertexMapBufferSRV = Model->GetVertexMapBuffer().ShaderResourceViewRHI;
+			HeatMapMode = static_cast<int32>(VizSettings->GetHeatMapMode());
+			HeatMapMax = 1.0f / FMath::Max(VizSettings->GetHeatMapMax(), 0.00001f);
+			GroundTruthLerp = VizSettings->GetGroundTruthLerp();
+		}
+	}
+
+	void FMLDeformerGraphDebugDataProviderProxy::HandleZeroGroundTruthPositions()
+	{
+		if (GroundTruthPositions.Num() == 0)
+		{	
+			// We didn't get valid ground truth vertices.
+			// Make non empty array for later buffer generation.
+			GroundTruthPositions.Add(FVector3f::ZeroVector);
+
+			// Silently disable relevant debug things.
+			if (HeatMapMode == static_cast<int32>(EMLDeformerHeatMapMode::GroundTruth))
+			{
+				HeatMapMode = -1;
+				HeatMapMax = 0.0f;
+				GroundTruthLerp = 0.0f;
+			}
+		}
+	}
+
+	bool FMLDeformerGraphDebugDataProviderProxy::IsValid(FValidationData const& InValidationData) const
+	{
+		if (InValidationData.ParameterStructSize != sizeof(FMLDeformerGraphDebugDataInterfaceParameters))
+		{
+			return false;
+		}
+		if (SkeletalMeshObject == nullptr || VertexMapBufferSRV == nullptr)
+		{
+			return false;
+		}
+
+		return true;
+	}
+
+	void FMLDeformerGraphDebugDataProviderProxy::AllocateResources(FRDGBuilder& GraphBuilder)
+	{
+		GroundTruthBuffer = GraphBuilder.CreateBuffer(FRDGBufferDesc::CreateStructuredDesc(sizeof(float), 3 * GroundTruthPositions.Num()), TEXT("MLDeformer.GroundTruthPositions"));
+		GroundTruthBufferSRV = GraphBuilder.CreateSRV(GroundTruthBuffer);
+		GraphBuilder.QueueBufferUpload(GroundTruthBuffer, GroundTruthPositions.GetData(), sizeof(FVector3f) * GroundTruthPositions.Num(), ERDGInitialDataFlags::None);
+	}
+
+	void FMLDeformerGraphDebugDataProviderProxy::GatherDispatchData(FDispatchData const& InDispatchData)
+	{
+		const FSkeletalMeshRenderData& SkeletalMeshRenderData = SkeletalMeshObject->GetSkeletalMeshRenderData();
+		const FSkeletalMeshLODRenderData* LodRenderData = SkeletalMeshRenderData.GetPendingFirstLOD(0);
+		const TStridedView<FMLDeformerGraphDebugDataInterfaceParameters> ParameterArray = MakeStridedParameterView<FMLDeformerGraphDebugDataInterfaceParameters>(InDispatchData);
+		for (int32 InvocationIndex = 0; InvocationIndex < ParameterArray.Num(); ++InvocationIndex)
+		{
+			const FSkelMeshRenderSection& RenderSection = LodRenderData->RenderSections[InvocationIndex];
+			FMLDeformerGraphDebugDataInterfaceParameters& Parameters = ParameterArray[InvocationIndex];
+			Parameters.NumVertices = 0;
+			Parameters.InputStreamStart = RenderSection.BaseVertexIndex;
+			Parameters.HeatMapMode = HeatMapMode;
+			Parameters.HeatMapMax = HeatMapMax;
+			Parameters.GroundTruthLerp = GroundTruthLerp;
+			Parameters.GroundTruthBufferSize = GroundTruthPositions.Num();
+			Parameters.PositionGroundTruthBuffer = GroundTruthBufferSRV;
+			Parameters.VertexMapBuffer = VertexMapBufferSRV;
+		}
+	}
+}	// namespace UE::MLDeformer
+#endif // WITH_EDITORONLY_DATA
+
diff --git a/Engine/Plugins/Animation/MLDeformer/MLDeformerFramework/Source/MLDeformerFramework/Private/MLDeformerModule.cpp b/Engine/Plugins/Animation/MLDeformer/MLDeformerFramework/Source/MLDeformerFramework/Private/MLDeformerModule.cpp
--- a/Engine/Plugins/Animation/MLDeformer/MLDeformerFramework/Source/MLDeformerFramework/Private/MLDeformerModule.cpp
+++ b/Engine/Plugins/Animation/MLDeformer/MLDeformerFramework/Source/MLDeformerFramework/Private/MLDeformerModule.cpp
@@ -13,4 +13,14 @@ MLDEFORMERFRAMEWORK_API DEFINE_LOG_CATEGORY(LogMLDeformer)
 
 IMPLEMENT_MODULE(UE::MLDeformer::FMLDeformerModule, MLDeformerFramework)
 
+namespace UE::MLDeformer
+{
+	void FMLDeformerModule::StartupModule()
+	{
+		// Register an additional shader path for our shaders used inside the deformer graph system.
+		const FString PluginShaderDir = FPaths::Combine(IPluginManager::Get().FindPlugin(TEXT("MLDeformerFramework"))->GetBaseDir(), TEXT("Shaders"));
+		AddShaderSourceDirectoryMapping(TEXT("/Plugin/MLDeformerFramework"), PluginShaderDir);
+	}
+}	// namespace UE::MLDeformer
+
 #undef LOCTEXT_NAMESPACE
diff --git a/Engine/Plugins/Animation/MLDeformer/MLDeformerFramework/Source/MLDeformerFramework/Public/MLDeformerGraphDebugDataInterface.h b/Engine/Plugins/Animation/MLDeformer/MLDeformerFramework/Source/MLDeformerFramework/Public/MLDeformerGraphDebugDataInterface.h
new file mode 100644
--- /dev/null
+++ b/Engine/Plugins/Animation/MLDeformer/MLDeformerFramework/Source/MLDeformerFramework/Public/MLDeformerGraphDebugDataInterface.h
@@ -0,0 +1,103 @@
+// Copyright Epic Games, Inc. All Rights Reserved.
+
+#pragma once
+
+#include "CoreMinimal.h"
+#include "MLDeformerModel.h"
+#include "OptimusComputeDataInterface.h"
+#include "ComputeFramework/ComputeDataProvider.h"
+#include "MLDeformerGraphDebugDataInterface.generated.h"
+
+class FRDGBuffer;
+class FRDGBufferSRV;
+class FRHIShaderResourceView;
+class FSkeletalMeshObject;
+class UMLDeformerAsset;
+class USkeletalMeshComponent;
+
+/** 
+ * Compute Framework Data Interface for MLDeformer debugging data. 
+ * This interfaces to editor only data, and so will only give valid results in that context.
+ */
+UCLASS(Category = ComputeFramework)
+class MLDEFORMERFRAMEWORK_API UMLDeformerGraphDebugDataInterface
+	: public UOptimusComputeDataInterface
+{
+	GENERATED_BODY()
+
+public:
+	// UOptimusComputeDataInterface overrides.
+	virtual FString GetDisplayName() const override;
+	virtual TArray<FOptimusCDIPinDefinition> GetPinDefinitions() const override;
+	virtual TSubclassOf<UActorComponent> GetRequiredComponentClass() const override;
+	// ~END UOptimusComputeDataInterface overrides.
+
+	// UComputeDataInterface overrides.
+	TCHAR const* GetClassName() const override { return TEXT("MLDeformerGraphDebugData"); }
+	virtual void GetSupportedInputs(TArray<FShaderFunctionDefinition>& OutFunctions) const override;
+	virtual void GetShaderParameters(TCHAR const* UID, FShaderParametersMetadataBuilder& OutBuilder, FShaderParametersMetadataAllocations& InOutAllocations) const override;
+	virtual TCHAR const* GetShaderVirtualPath() const override;
+	virtual void GetShaderHash(FString& InOutKey) const override;
+	virtual void GetHLSL(FString& OutHLSL, FString const& InDataInterfaceName) const override;
+	virtual UComputeDataProvider* CreateDataProvider(TObjectPtr<UObject> InBinding, uint64 InInputMask, uint64 InOutputMask) const override;
+	// ~END UComputeDataInterface overrides.
+
+private:
+	static TCHAR const* TemplateFilePath;
+};
+
+/** Compute Framework Data Provider for MLDeformer debugging data. */
+UCLASS(BlueprintType, EditInlineNew, Category = ComputeFramework)
+class MLDEFORMERFRAMEWORK_API UMLDeformerGraphDebugDataProvider
+	: public UComputeDataProvider
+{
+	GENERATED_BODY()
+
+public:
+	virtual void Init() {};
+
+	// UComputeDataProvider overrides.
+	virtual FComputeDataProviderRenderProxy* GetRenderProxy() override;
+	// ~END UComputeDataProvider overrides.
+
+public:
+	UPROPERTY(EditAnywhere, BlueprintReadWrite, Category = Binding)
+	TObjectPtr<UMLDeformerComponent> DeformerComponent = nullptr;
+
+	UPROPERTY(EditAnywhere, BlueprintReadWrite, Category = Binding)
+	TObjectPtr<UMLDeformerAsset> DeformerAsset = nullptr;
+};
+
+/** Compute Framework Data Provider Proxy for MLDeformer debugging data. */
+#if WITH_EDITORONLY_DATA
+namespace UE::MLDeformer
+{
+	class MLDEFORMERFRAMEWORK_API FMLDeformerGraphDebugDataProviderProxy
+		: public FComputeDataProviderRenderProxy
+	{
+	public:
+		FMLDeformerGraphDebugDataProviderProxy(UMLDeformerComponent* DeformerComponent, UMLDeformerAsset* DeformerAsset, UMLDeformerGraphDebugDataProvider* InProvider);
+
+		virtual void HandleZeroGroundTruthPositions();
+
+		// FComputeDataProviderRenderProxy overrides.
+		bool IsValid(FValidationData const& InValidationData) const override;
+		void AllocateResources(FRDGBuilder& GraphBuilder) override;
+		void GatherDispatchData(FDispatchData const& InDispatchData) override;
+		// ~END FComputeDataProviderRenderProxy overrides.
+
+		TArray<FVector3f>& GetGroundTruthPositions() { return GroundTruthPositions; }
+
+	protected:
+		TObjectPtr<UMLDeformerGraphDebugDataProvider> Provider = nullptr;
+		FSkeletalMeshObject* SkeletalMeshObject = nullptr;
+		TArray<FVector3f> GroundTruthPositions;
+		FRHIShaderResourceView* VertexMapBufferSRV = nullptr;
+		FRDGBuffer* GroundTruthBuffer = nullptr;
+		FRDGBufferSRV* GroundTruthBufferSRV = nullptr;
+		int32 HeatMapMode = 0;
+		float HeatMapMax = 0.0f;
+		float GroundTruthLerp = 0.0f;
+	};
+}	// namespace UE::MLDeformer
+#endif
diff --git a/Engine/Plugins/Animation/MLDeformer/MLDeformerFramework/Source/MLDeformerFramework/Public/MLDeformerModule.h b/Engine/Plugins/Animation/MLDeformer/MLDeformerFramework/Source/MLDeformerFramework/Public/MLDeformerModule.h
--- a/Engine/Plugins/Animation/MLDeformer/MLDeformerFramework/Source/MLDeformerFramework/Public/MLDeformerModule.h
+++ b/Engine/Plugins/Animation/MLDeformer/MLDeformerFramework/Source/MLDeformerFramework/Public/MLDeformerModule.h
@@ -17,5 +17,8 @@ namespace UE::MLDeformer
 		: public IModuleInterface
 	{
 	public:
+		// IModuleInterface overrides.
+		void StartupModule() override;
+		// ~END IModuleInterface overrides.
 	};
 }	// namespace UE::MLDeformer
diff --git a/Engine/Plugins/Animation/MLDeformer/MLDeformerFramework/Source/MLDeformerFrameworkEditor/Private/MLDeformerEditorToolkit.cpp b/Engine/Plugins/Animation/MLDeformer/MLDeformerFramework/Source/MLDeformerFrameworkEditor/Private/MLDeformerEditorToolkit.cpp
--- a/Engine/Plugins/Animation/MLDeformer/MLDeformerFramework/Source/MLDeformerFrameworkEditor/Private/MLDeformerEditorToolkit.cpp
+++ b/Engine/Plugins/Animation/MLDeformer/MLDeformerFramework/Source/MLDeformerFrameworkEditor/Private/MLDeformerEditorToolkit.cpp
@@ -639,6 +639,16 @@ namespace UE::MLDeformer
 			ActiveModel->SetDefaultDeformerGraphIfNeeded();
 			ActiveModel->CreateActors(InPersonaPreviewScene);
 			ActiveModel->UpdateActorVisibility();
+
+			if (GetModelDetailsView())
+			{
+				GetModelDetailsView()->SetObject(nullptr, true);
+			}
+
+			if (GetVizSettingsDetailsView())
+			{
+				GetVizSettingsDetailsView()->SetObject(nullptr, true);
+			}
 			ActiveModel->TriggerInputAssetChanged(false);
 			ActiveModel->CreateHeatMapAssets();
 			ActiveModel->SetHeatMapMaterialEnabled(ActiveModel->GetModel()->GetVizSettings()->GetShowHeatMap());
diff --git a/Engine/Plugins/Animation/MLDeformer/MLDeformerFramework/Source/MLDeformerFrameworkEditor/Private/MLDeformerMorphModelVizSettingsDetails.cpp b/Engine/Plugins/Animation/MLDeformer/MLDeformerFramework/Source/MLDeformerFrameworkEditor/Private/MLDeformerMorphModelVizSettingsDetails.cpp
--- a/Engine/Plugins/Animation/MLDeformer/MLDeformerFramework/Source/MLDeformerFrameworkEditor/Private/MLDeformerMorphModelVizSettingsDetails.cpp
+++ b/Engine/Plugins/Animation/MLDeformer/MLDeformerFramework/Source/MLDeformerFrameworkEditor/Private/MLDeformerMorphModelVizSettingsDetails.cpp
@@ -70,7 +70,7 @@ namespace UE::MLDeformer
 	{
 		IDetailGroup& MorphsGroup = LiveSettingsCategory->AddGroup("Morph Targets", LOCTEXT("MorphTargetsLabel", "Morph Targets"), false, true);
 		MorphsGroup.AddPropertyRow(DetailLayoutBuilder->GetProperty(UMLDeformerMorphModelVizSettings::GetDrawMorphTargetsPropertyName(), UMLDeformerMorphModelVizSettings::StaticClass()))
-			.EditCondition(!MorphModel->GetMorphTargetDeltas().IsEmpty(), nullptr);
+			.EditCondition(!MorphModel->GetMorphTargetDeltas().IsEmpty() && MorphModel->CanDynamicallyUpdateMorphTargets(), nullptr);
 
 		MorphsGroup.AddPropertyRow(DetailLayoutBuilder->GetProperty(UMLDeformerMorphModelVizSettings::GetMorphTargetNumberPropertyName(), UMLDeformerMorphModelVizSettings::StaticClass()))
 			.EditCondition(TAttribute<bool>(this, &FMLDeformerMorphModelVizSettingsDetails::IsMorphTargetsEnabled), nullptr);
diff --git a/Engine/Plugins/Animation/MLDeformer/MLDeformerFramework/Source/MLDeformerFrameworkEditor/Private/MLDeformerSampler.cpp b/Engine/Plugins/Animation/MLDeformer/MLDeformerFramework/Source/MLDeformerFrameworkEditor/Private/MLDeformerSampler.cpp
--- a/Engine/Plugins/Animation/MLDeformer/MLDeformerFramework/Source/MLDeformerFrameworkEditor/Private/MLDeformerSampler.cpp
+++ b/Engine/Plugins/Animation/MLDeformer/MLDeformerFramework/Source/MLDeformerFrameworkEditor/Private/MLDeformerSampler.cpp
@@ -90,13 +90,13 @@ namespace UE::MLDeformer
 
 	void FMLDeformerSampler::UpdateBoneRotations()
 	{
-		const UMLDeformerInputInfo* InputInfo = Model->GetInputInfo();
+		const UMLDeformerInputInfo* InputInfo = EditorModel->GetEditorInputInfo();
 		InputInfo->ExtractBoneRotations(SkeletalMeshComponent, BoneRotations);
 	}
 
 	void FMLDeformerSampler::UpdateCurveValues()
 	{
-		const UMLDeformerInputInfo* InputInfo = Model->GetInputInfo();
+		const UMLDeformerInputInfo* InputInfo = EditorModel->GetEditorInputInfo();
 		InputInfo->ExtractCurveValues(SkeletalMeshComponent, CurveValues, NumFloatsPerCurve);
 	}
 
@@ -230,7 +230,7 @@ namespace UE::MLDeformer
 
 	int32 FMLDeformerSampler::GetNumBones() const
 	{
-		const UMLDeformerInputInfo* InputInfo = Model->GetInputInfo();
+		const UMLDeformerInputInfo* InputInfo = EditorModel->GetEditorInputInfo();
 		return InputInfo->GetNumBones();
 	}
 
diff --git a/Engine/Plugins/Animation/MLDeformer/VertexDeltaModel/Source/VertexDeltaModel/Private/VertexDeltaGraphDebugDataInterface.cpp b/Engine/Plugins/Animation/MLDeformer/VertexDeltaModel/Source/VertexDeltaModel/Private/VertexDeltaGraphDebugDataInterface.cpp
--- a/Engine/Plugins/Animation/MLDeformer/VertexDeltaModel/Source/VertexDeltaModel/Private/VertexDeltaGraphDebugDataInterface.cpp
+++ b/Engine/Plugins/Animation/MLDeformer/VertexDeltaModel/Source/VertexDeltaModel/Private/VertexDeltaGraphDebugDataInterface.cpp
@@ -1,7 +1,6 @@
 // Copyright Epic Games, Inc. All Rights Reserved.
 
 #include "VertexDeltaGraphDebugDataInterface.h"
-#include "VertexDeltaGraphDataInterface.h"
 #include "MLDeformerModelInstance.h"
 #include "MLDeformerAsset.h"
 #include "MLDeformerComponent.h"
@@ -21,7 +20,9 @@
 
 #include UE_INLINE_GENERATED_CPP_BY_NAME(VertexDeltaGraphDebugDataInterface)
 
-TArray<FOptimusCDIPinDefinition> UVertexDeltaGraphDebugDataInterface::GetPinDefinitions() const
+PRAGMA_DISABLE_DEPRECATION_WARNINGS
+
+TArray<FOptimusCDIPinDefinition> UDEPRECATED_VertexDeltaGraphDebugDataInterface::GetPinDefinitions() const
 {
 	TArray<FOptimusCDIPinDefinition> Defs;
 	Defs.Add({ "HeatMapMode", "ReadHeatMapMode" });
@@ -31,12 +32,12 @@ TArray<FOptimusCDIPinDefinition> UVertexDeltaGraphDebugDataInterface::GetPinDefi
 	return Defs;
 }
 
-TSubclassOf<UActorComponent> UVertexDeltaGraphDebugDataInterface::GetRequiredComponentClass() const
+TSubclassOf<UActorComponent> UDEPRECATED_VertexDeltaGraphDebugDataInterface::GetRequiredComponentClass() const
 {
 	return UMLDeformerComponent::StaticClass();
 }
 
-void UVertexDeltaGraphDebugDataInterface::GetSupportedInputs(TArray<FShaderFunctionDefinition>& OutFunctions) const
+void UDEPRECATED_VertexDeltaGraphDebugDataInterface::GetSupportedInputs(TArray<FShaderFunctionDefinition>& OutFunctions) const
 {
 	OutFunctions.AddDefaulted_GetRef()
 		.SetName(TEXT("ReadNumVertices"))
@@ -71,29 +72,29 @@ BEGIN_SHADER_PARAMETER_STRUCT(FVertexDeltaGraphDebugDataInterfaceParameters, )
 	SHADER_PARAMETER_SRV(Buffer<uint>, VertexMapBuffer)
 END_SHADER_PARAMETER_STRUCT()
 
-FString UVertexDeltaGraphDebugDataInterface::GetDisplayName() const
+FString UDEPRECATED_VertexDeltaGraphDebugDataInterface::GetDisplayName() const
 {
 	return TEXT("MLD Vertex Delta Model Debug");
 }
 
-void UVertexDeltaGraphDebugDataInterface::GetShaderParameters(TCHAR const* UID, FShaderParametersMetadataBuilder& InOutBuilder, FShaderParametersMetadataAllocations& InOutAllocations) const
+void UDEPRECATED_VertexDeltaGraphDebugDataInterface::GetShaderParameters(TCHAR const* UID, FShaderParametersMetadataBuilder& InOutBuilder, FShaderParametersMetadataAllocations& InOutAllocations) const
 {
 	InOutBuilder.AddNestedStruct<FVertexDeltaGraphDebugDataInterfaceParameters>(UID);
 }
 
-TCHAR const* UVertexDeltaGraphDebugDataInterface::TemplateFilePath = TEXT("/Plugin/VertexDeltaModel/Private/VertexDeltaModelHeatMap.ush");
+TCHAR const* UDEPRECATED_VertexDeltaGraphDebugDataInterface::TemplateFilePath = TEXT("/Plugin/VertexDeltaModel/Private/VertexDeltaModelHeatMap.ush");
 
-TCHAR const* UVertexDeltaGraphDebugDataInterface::GetShaderVirtualPath() const
+TCHAR const* UDEPRECATED_VertexDeltaGraphDebugDataInterface::GetShaderVirtualPath() const
 {
 	return TemplateFilePath;
 }
 
-void UVertexDeltaGraphDebugDataInterface::GetShaderHash(FString& InOutKey) const
+void UDEPRECATED_VertexDeltaGraphDebugDataInterface::GetShaderHash(FString& InOutKey) const
 {
 	GetShaderFileHash(TemplateFilePath, EShaderPlatform::SP_PCD3D_SM5).AppendString(InOutKey);
 }
 
-void UVertexDeltaGraphDebugDataInterface::GetHLSL(FString& OutHLSL, FString const& InDataInterfaceName) const
+void UDEPRECATED_VertexDeltaGraphDebugDataInterface::GetHLSL(FString& OutHLSL, FString const& InDataInterfaceName) const
 {
 	TMap<FString, FStringFormatArg> TemplateArgs =
 	{
@@ -105,24 +106,25 @@ void UVertexDeltaGraphDebugDataInterface::GetHLSL(FString& OutHLSL, FString cons
 	OutHLSL += FString::Format(*TemplateFile, TemplateArgs);
 }
 
-UComputeDataProvider* UVertexDeltaGraphDebugDataInterface::CreateDataProvider(TObjectPtr<UObject> InBinding, uint64 InInputMask, uint64 InOutputMask) const
+UComputeDataProvider* UDEPRECATED_VertexDeltaGraphDebugDataInterface::CreateDataProvider(TObjectPtr<UObject> InBinding, uint64 InInputMask, uint64 InOutputMask) const
 {
-	UVertexDeltaGraphDebugDataProvider* Provider = NewObject<UVertexDeltaGraphDebugDataProvider>();
+	UDEPRECATED_VertexDeltaGraphDebugDataProvider* Provider = NewObject<UDEPRECATED_VertexDeltaGraphDebugDataProvider>();
 	Provider->DeformerComponent = Cast<UMLDeformerComponent>(InBinding);
 	if (Provider->DeformerComponent)
 	{
 		Provider->DeformerAsset = Provider->DeformerComponent->GetDeformerAsset();
 	}
-	MLDEFORMER_EDITORDATA_ONLY(
+
+	#if WITH_EDITORONLY_DATA
 		if (Provider->DeformerAsset != nullptr)
 		{
 			Provider->Init();
 		}
-	,)
+	#endif
 	return Provider;
 }
 
-FComputeDataProviderRenderProxy* UVertexDeltaGraphDebugDataProvider::GetRenderProxy()
+FComputeDataProviderRenderProxy* UDEPRECATED_VertexDeltaGraphDebugDataProvider::GetRenderProxy()
 {
 #if WITH_EDITORONLY_DATA
 	if (DeformerComponent && DeformerAsset && DeformerComponent->GetModelInstance() && DeformerComponent->GetModelInstance()->IsValidForDataProvider())
@@ -143,7 +145,7 @@ FComputeDataProviderRenderProxy* UVertexDeltaGraphDebugDataProvider::GetRenderPr
 #if WITH_EDITORONLY_DATA
 namespace UE::VertexDeltaModel
 {
-	FVertexDeltaGraphDebugDataProviderProxy::FVertexDeltaGraphDebugDataProviderProxy(UMLDeformerComponent* DeformerComponent, UMLDeformerAsset* DeformerAsset, UVertexDeltaGraphDebugDataProvider* InProvider)
+	FVertexDeltaGraphDebugDataProviderProxy::FVertexDeltaGraphDebugDataProviderProxy(UMLDeformerComponent* DeformerComponent, UMLDeformerAsset* DeformerAsset, UDEPRECATED_VertexDeltaGraphDebugDataProvider* InProvider)
 		: FComputeDataProviderRenderProxy()
 	{
 		Provider = InProvider;
@@ -223,3 +225,4 @@ namespace UE::VertexDeltaModel
 }	// namespace UE::VertexDeltaModel
 #endif // WITH_EDITORONLY_DATA
 
+PRAGMA_ENABLE_DEPRECATION_WARNINGS
diff --git a/Engine/Plugins/Animation/MLDeformer/VertexDeltaModel/Source/VertexDeltaModel/Private/VertexDeltaModel.cpp b/Engine/Plugins/Animation/MLDeformer/VertexDeltaModel/Source/VertexDeltaModel/Private/VertexDeltaModel.cpp
--- a/Engine/Plugins/Animation/MLDeformer/VertexDeltaModel/Source/VertexDeltaModel/Private/VertexDeltaModel.cpp
+++ b/Engine/Plugins/Animation/MLDeformer/VertexDeltaModel/Source/VertexDeltaModel/Private/VertexDeltaModel.cpp
@@ -21,7 +21,7 @@ namespace UE::VertexDeltaModel
 	class VERTEXDELTAMODEL_API FVertexDeltaModelModule
 		: public IModuleInterface
 	{
-		void StartupModule()
+		void StartupModule() override
 		{
 			// Register an additional shader path for our shaders used inside the deformer graph system.
 			const FString PluginShaderDir = FPaths::Combine(IPluginManager::Get().FindPlugin(TEXT("VertexDeltaModel"))->GetBaseDir(), TEXT("Shaders"));
diff --git a/Engine/Plugins/Animation/MLDeformer/VertexDeltaModel/Source/VertexDeltaModel/Public/VertexDeltaGraphDataInterface.h b/Engine/Plugins/Animation/MLDeformer/VertexDeltaModel/Source/VertexDeltaModel/Public/VertexDeltaGraphDataInterface.h
--- a/Engine/Plugins/Animation/MLDeformer/VertexDeltaModel/Source/VertexDeltaModel/Public/VertexDeltaGraphDataInterface.h
+++ b/Engine/Plugins/Animation/MLDeformer/VertexDeltaModel/Source/VertexDeltaModel/Public/VertexDeltaGraphDataInterface.h
@@ -19,12 +19,6 @@ class UNeuralNetwork;
 class USkeletalMeshComponent;
 class UMLDeformerModel;
 
-#if WITH_EDITORONLY_DATA
-	#define MLDEFORMER_EDITORDATA_ONLY(Statement, ElseStatement) Statement
-#else
-	#define MLDEFORMER_EDITORDATA_ONLY(X, ElseStatement) ElseStatement
-#endif
-
 /** Compute Framework Data Interface for MLDefomer data. */
 UCLASS(Category = ComputeFramework)
 class VERTEXDELTAMODEL_API UVertexDeltaGraphDataInterface
diff --git a/Engine/Plugins/Animation/MLDeformer/VertexDeltaModel/Source/VertexDeltaModel/Public/VertexDeltaGraphDebugDataInterface.h b/Engine/Plugins/Animation/MLDeformer/VertexDeltaModel/Source/VertexDeltaModel/Public/VertexDeltaGraphDebugDataInterface.h
--- a/Engine/Plugins/Animation/MLDeformer/VertexDeltaModel/Source/VertexDeltaModel/Public/VertexDeltaGraphDebugDataInterface.h
+++ b/Engine/Plugins/Animation/MLDeformer/VertexDeltaModel/Source/VertexDeltaModel/Public/VertexDeltaGraphDebugDataInterface.h
@@ -20,8 +20,8 @@ class USkeletalMeshComponent;
  * Compute Framework Data Interface for MLDeformer debugging data. 
  * This interfaces to editor only data, and so will only give valid results in that context.
  */
-UCLASS(Category = ComputeFramework)
-class VERTEXDELTAMODEL_API UVertexDeltaGraphDebugDataInterface
+UCLASS(Deprecated, Category = ComputeFramework, meta = (DeprecationMessage = "Please use UMLDeformerGraphDebugDataInterface instead."))
+class VERTEXDELTAMODEL_API UDEPRECATED_VertexDeltaGraphDebugDataInterface
 	: public UOptimusComputeDataInterface
 {
 	GENERATED_BODY()
@@ -48,8 +48,8 @@ private:
 };
 
 /** Compute Framework Data Provider for MLDeformer debugging data. */
-UCLASS(BlueprintType, EditInlineNew, Category = ComputeFramework)
-class VERTEXDELTAMODEL_API UVertexDeltaGraphDebugDataProvider
+UCLASS(Deprecated, BlueprintType, EditInlineNew, Category = ComputeFramework, meta = (DeprecationMessage = "Please use UMLDeformerGraphDebugDataProvider instead."))
+class VERTEXDELTAMODEL_API UDEPRECATED_VertexDeltaGraphDebugDataProvider
 	: public UComputeDataProvider
 {
 	GENERATED_BODY()
@@ -73,11 +73,11 @@ public:
 #if WITH_EDITORONLY_DATA
 namespace UE::VertexDeltaModel
 {
-	class VERTEXDELTAMODEL_API FVertexDeltaGraphDebugDataProviderProxy
+	class VERTEXDELTAMODEL_API UE_DEPRECATED(5.2, "Please use UE::MLDeformer::FMLDeformerDebugDataProviderProxy instead.") FVertexDeltaGraphDebugDataProviderProxy
 		: public FComputeDataProviderRenderProxy
 	{
 	public:
-		FVertexDeltaGraphDebugDataProviderProxy(UMLDeformerComponent* DeformerComponent, UMLDeformerAsset* DeformerAsset, UVertexDeltaGraphDebugDataProvider* InProvider);
+		FVertexDeltaGraphDebugDataProviderProxy(UMLDeformerComponent* DeformerComponent, UMLDeformerAsset* DeformerAsset, UDEPRECATED_VertexDeltaGraphDebugDataProvider* InProvider);
 
 		virtual void HandleZeroGroundTruthPositions();
 
@@ -90,7 +90,7 @@ namespace UE::VertexDeltaModel
 		TArray<FVector3f>& GetGroundTruthPositions() { return GroundTruthPositions; }
 
 	protected:
-		TObjectPtr<UVertexDeltaGraphDebugDataProvider> Provider = nullptr;
+		TObjectPtr<UDEPRECATED_VertexDeltaGraphDebugDataProvider> Provider = nullptr;
 		FSkeletalMeshObject* SkeletalMeshObject = nullptr;
 		TArray<FVector3f> GroundTruthPositions;
 		FRHIShaderResourceView* VertexMapBufferSRV = nullptr;
