[MLDeformer]
- Fix some issue that could cause unexpected behavior when using curves in the Local Model. This also triggered a check to fail in debug mode.
- Fixed a crash when the local bone transforms on the skeletal mesh component was empty.

diff --git a/Engine/Plugins/Animation/MLDeformer/MLDeformerFramework/Source/MLDeformerFramework/Private/MLDeformerModelInstance.cpp b/Engine/Plugins/Animation/MLDeformer/MLDeformerFramework/Source/MLDeformerFramework/Private/MLDeformerModelInstance.cpp
--- a/Engine/Plugins/Animation/MLDeformer/MLDeformerFramework/Source/MLDeformerFramework/Private/MLDeformerModelInstance.cpp
+++ b/Engine/Plugins/Animation/MLDeformer/MLDeformerFramework/Source/MLDeformerFramework/Private/MLDeformerModelInstance.cpp
@@ -151,24 +151,26 @@ void UMLDeformerModelInstance::UpdateBoneTransforms()
 	if (LeaderPoseComponent)
 	{
 		const TArray<FTransform>& LeaderTransforms = LeaderPoseComponent->GetComponentSpaceTransforms();
-		USkinnedAsset* SkinnedAsset = LeaderPoseComponent->GetSkinnedAsset();
-		const int32 NumTrainingBones = AssetBonesToSkelMeshMappings.Num();
-		for (int32 Index = 0; Index < NumTrainingBones; ++Index)
+		if (!LeaderTransforms.IsEmpty())
 		{
-			const int32 ComponentBoneIndex = AssetBonesToSkelMeshMappings[Index];
-			check(LeaderTransforms.IsValidIndex(ComponentBoneIndex));
-			
-			const FTransform& ComponentSpaceTransform = LeaderTransforms[ComponentBoneIndex];
-			const int32 ParentIndex = SkinnedAsset->GetRefSkeleton().GetParentIndex(ComponentBoneIndex);
-			if (LeaderTransforms.IsValidIndex(ParentIndex))
+			USkinnedAsset* SkinnedAsset = LeaderPoseComponent->GetSkinnedAsset();
+			const int32 NumTrainingBones = AssetBonesToSkelMeshMappings.Num();
+			for (int32 Index = 0; Index < NumTrainingBones; ++Index)
 			{
-				TrainingBoneTransforms[Index] = ComponentSpaceTransform.GetRelativeTransform(LeaderTransforms[ParentIndex]);
+				const int32 ComponentBoneIndex = AssetBonesToSkelMeshMappings[Index];
+				checkSlow(LeaderTransforms.IsValidIndex(ComponentBoneIndex));			
+				const FTransform& ComponentSpaceTransform = LeaderTransforms[ComponentBoneIndex];
+				const int32 ParentIndex = SkinnedAsset->GetRefSkeleton().GetParentIndex(ComponentBoneIndex);
+				if (LeaderTransforms.IsValidIndex(ParentIndex))
+				{
+					TrainingBoneTransforms[Index] = ComponentSpaceTransform.GetRelativeTransform(LeaderTransforms[ParentIndex]);
+				}
+				else
+				{
+					TrainingBoneTransforms[Index] = ComponentSpaceTransform;
+				}
+				TrainingBoneTransforms[Index].NormalizeRotation();
 			}
-			else
-			{
-				TrainingBoneTransforms[Index] = ComponentSpaceTransform;
-			}
-			TrainingBoneTransforms[Index].NormalizeRotation();
 		}
 	}
 	else
@@ -176,11 +178,14 @@ void UMLDeformerModelInstance::UpdateBoneTransforms()
 		// Grab the transforms from our own skeletal mesh component.
 		// These are local space transforms, relative to the parent bone.
 		const TArrayView<const FTransform> Transforms = SkeletalMeshComponent->GetBoneSpaceTransformsView();
-		const int32 NumTrainingBones = AssetBonesToSkelMeshMappings.Num();
-		for (int32 Index = 0; Index < NumTrainingBones; ++Index)
+		if (!Transforms.IsEmpty())
 		{
-			const int32 ComponentBoneIndex = AssetBonesToSkelMeshMappings[Index];
-			TrainingBoneTransforms[Index] = Transforms[ComponentBoneIndex];
+			const int32 NumTrainingBones = AssetBonesToSkelMeshMappings.Num();
+			for (int32 Index = 0; Index < NumTrainingBones; ++Index)
+			{
+				const int32 ComponentBoneIndex = AssetBonesToSkelMeshMappings[Index];
+				TrainingBoneTransforms[Index] = Transforms[ComponentBoneIndex];
+			}
 		}
 	}
 }
@@ -254,6 +259,29 @@ bool UMLDeformerModelInstance::IsValidForDataProvider() const
 	return true;
 }
 
+bool UMLDeformerModelInstance::HasValidTransforms() const
+{
+	if (!SkeletalMeshComponent)
+	{
+		return false;
+	}
+
+	const USkinnedMeshComponent* LeaderPoseComponent = SkeletalMeshComponent->LeaderPoseComponent.Get();
+	if (LeaderPoseComponent)
+	{
+		if (LeaderPoseComponent->GetComponentSpaceTransforms().IsEmpty())
+		{
+			return false;
+		}
+	}
+	else if (SkeletalMeshComponent->GetBoneSpaceTransformsView().IsEmpty())
+	{
+		return false;
+	}
+
+	return true;
+}
+
 void UMLDeformerModelInstance::Tick(float DeltaTime, float ModelWeight)
 {
 	// Post init hasn't yet succeeded, try again.
@@ -263,7 +291,7 @@ void UMLDeformerModelInstance::Tick(float DeltaTime, float ModelWeight)
 		PostMLDeformerComponentInit();
 	}
 
-	if (ModelWeight > 0.0001f && SetupInputs())
+	if (ModelWeight > 0.0001f && HasValidTransforms() && SetupInputs())
 	{
 		// Execute the model instance.
 		// For models using neural networks this will perform the inference, 
diff --git a/Engine/Plugins/Animation/MLDeformer/MLDeformerFramework/Source/MLDeformerFramework/Public/MLDeformerModelInstance.h b/Engine/Plugins/Animation/MLDeformer/MLDeformerFramework/Source/MLDeformerFramework/Public/MLDeformerModelInstance.h
--- a/Engine/Plugins/Animation/MLDeformer/MLDeformerFramework/Source/MLDeformerFramework/Public/MLDeformerModelInstance.h
+++ b/Engine/Plugins/Animation/MLDeformer/MLDeformerFramework/Source/MLDeformerFramework/Public/MLDeformerModelInstance.h
@@ -192,6 +192,13 @@ protected:
 	 */
 	void UpdateBoneTransforms();
 
+	/** 
+	 * Check whether we have valid transforms in the skeletal mesh component.
+	 * When this returns false, we can't really execute the deformer.
+	 * @return Returns true if the skeletal mesh component, or its leader component has non-empty transform buffers.
+	 */
+	bool HasValidTransforms() const;
+
 protected:
 	/** The fence that let's us wait for all render commands to finish, before we continue. */
 	FRenderCommandFence RenderCommandFence;
diff --git a/Engine/Plugins/Animation/MLDeformer/NeuralMorphModel/Source/NeuralMorphModel/Private/NeuralMorphModelInstance.cpp b/Engine/Plugins/Animation/MLDeformer/NeuralMorphModel/Source/NeuralMorphModel/Private/NeuralMorphModelInstance.cpp
--- a/Engine/Plugins/Animation/MLDeformer/NeuralMorphModel/Source/NeuralMorphModel/Private/NeuralMorphModelInstance.cpp
+++ b/Engine/Plugins/Animation/MLDeformer/NeuralMorphModel/Source/NeuralMorphModel/Private/NeuralMorphModelInstance.cpp
@@ -134,7 +134,7 @@ int64 UNeuralMorphModelInstance::SetCurveValues(float* OutputBuffer, int64 Outpu
 	const int32 NumCurveFloats = AssetNumCurves * NumFloatsPerCurve;
 	checkf((Index + NumCurveFloats) <= OutputBufferSize, TEXT("Writing curves past the end of the input buffer"));
 
-	if (NumCurveFloats > 1)
+	if (NumFloatsPerCurve > 1)
 	{
 		// First write all zeros.
 		for (int32 CurveIndex = 0; CurveIndex < NumCurveFloats; ++CurveIndex)
@@ -156,7 +156,7 @@ int64 UNeuralMorphModelInstance::SetCurveValues(float* OutputBuffer, int64 Outpu
 	}
 	else
 	{
-		checkSlow(NumCurveFloats == 1);
+		checkSlow(NumFloatsPerCurve == 1);
 		UAnimInstance* AnimInstance = SkeletalMeshComponent->GetAnimInstance();
 		if (AnimInstance)
 		{
