Add Cycle Stats for ML Deformer

diff --git a/Engine/Plugins/Animation/MLDeformer/MLDeformerFramework/Source/MLDeformerFramework/Private/MLDeformerComponent.cpp b/Engine/Plugins/Animation/MLDeformer/MLDeformerFramework/Source/MLDeformerFramework/Private/MLDeformerComponent.cpp
--- a/Engine/Plugins/Animation/MLDeformer/MLDeformerFramework/Source/MLDeformerFramework/Private/MLDeformerComponent.cpp
+++ b/Engine/Plugins/Animation/MLDeformer/MLDeformerFramework/Source/MLDeformerFramework/Private/MLDeformerComponent.cpp
@@ -223,7 +223,7 @@ void UMLDeformerComponent::TickComponent(float DeltaTime, enum ELevelTick TickTy
 	#if WITH_EDITOR
 		TickPerfCounter.BeginSample();
 	#endif
-
+	SCOPE_CYCLE_COUNTER(STAT_MLDeformerInference);
 	Super::TickComponent(DeltaTime, TickType, ThisTickFunction);
 	if (TickType != ELevelTick::LEVELTICK_PauseTick)
 	{
diff --git a/Engine/Plugins/Animation/MLDeformer/MLDeformerFramework/Source/MLDeformerFramework/Private/MLDeformerModelInstance.cpp b/Engine/Plugins/Animation/MLDeformer/MLDeformerFramework/Source/MLDeformerFramework/Private/MLDeformerModelInstance.cpp
--- a/Engine/Plugins/Animation/MLDeformer/MLDeformerFramework/Source/MLDeformerFramework/Private/MLDeformerModelInstance.cpp
+++ b/Engine/Plugins/Animation/MLDeformer/MLDeformerFramework/Source/MLDeformerFramework/Private/MLDeformerModelInstance.cpp
@@ -16,6 +16,10 @@
 
 CSV_DEFINE_CATEGORY_MODULE(MLDEFORMERFRAMEWORK_API, MLDeformer, false);
 
+#if STATS
+DEFINE_STAT(STAT_MLDeformerInference);
+#endif
+
 void UMLDeformerModelInstance::BeginDestroy()
 {
 	Release();
diff --git a/Engine/Plugins/Animation/MLDeformer/MLDeformerFramework/Source/MLDeformerFramework/Public/MLDeformerModelInstance.h b/Engine/Plugins/Animation/MLDeformer/MLDeformerFramework/Source/MLDeformerFramework/Public/MLDeformerModelInstance.h
--- a/Engine/Plugins/Animation/MLDeformer/MLDeformerFramework/Source/MLDeformerFramework/Public/MLDeformerModelInstance.h
+++ b/Engine/Plugins/Animation/MLDeformer/MLDeformerFramework/Source/MLDeformerFramework/Public/MLDeformerModelInstance.h
@@ -13,6 +13,10 @@ class UMLDeformerModel;
 class USkeletalMeshComponent;
 class UMLDeformerComponent;
 
+#if STATS
+DECLARE_STATS_GROUP(TEXT("MLDeformer"), STATGROUP_MLDeformer, STATCAT_Advanced);
+DECLARE_CYCLE_STAT_EXTERN(TEXT("ML Deformer Inference"), STAT_MLDeformerInference, STATGROUP_MLDeformer, );
+#endif
 /**
  * An instance of the ML Deformer model.
  * The ML Deformer model contains shared data, while this instance contains data unique to the actor it is being applied to.
