Static analysis fixes

Fallout from the curve refactor

diff --git a/Engine/Plugins/Animation/GameplayInsights/Source/GameplayInsights/Private/PoseWatchTrack.cpp b/Engine/Plugins/Animation/GameplayInsights/Source/GameplayInsights/Private/PoseWatchTrack.cpp
--- a/Engine/Plugins/Animation/GameplayInsights/Source/GameplayInsights/Private/PoseWatchTrack.cpp
+++ b/Engine/Plugins/Animation/GameplayInsights/Source/GameplayInsights/Private/PoseWatchTrack.cpp
@@ -207,12 +207,12 @@ bool FPoseWatchTrack::UpdateInternal()
 
 			bChanged = bChanged || Children[i]->Update();
 		}
-	}
 
-	if(const TCHAR* FoundName = AnimationProvider->GetName(NameId))
-	{
-		bChanged = bChanged || TrackName.ToString().Compare(FoundName) != 0;
-		TrackName = FText::FromString(FoundName);
+		if(const TCHAR* FoundName = AnimationProvider->GetName(NameId))
+		{
+			bChanged = bChanged || TrackName.ToString().Compare(FoundName) != 0;
+			TrackName = FText::FromString(FoundName);
+		}
 	}
 
 	return bChanged;
diff --git a/Engine/Plugins/Animation/MLDeformer/MLDeformerFramework/Source/MLDeformerFrameworkEditor/Private/MLDeformerEditorModel.cpp b/Engine/Plugins/Animation/MLDeformer/MLDeformerFramework/Source/MLDeformerFrameworkEditor/Private/MLDeformerEditorModel.cpp
--- a/Engine/Plugins/Animation/MLDeformer/MLDeformerFramework/Source/MLDeformerFrameworkEditor/Private/MLDeformerEditorModel.cpp
+++ b/Engine/Plugins/Animation/MLDeformer/MLDeformerFramework/Source/MLDeformerFrameworkEditor/Private/MLDeformerEditorModel.cpp
@@ -1050,18 +1050,21 @@ namespace UE::MLDeformer
 		if (bIncludeCurves && SkeletalMesh)
 		{
 			// Anim curves.
-			// Include all curves when no list was provided.
-			if (Model->GetCurveIncludeList().IsEmpty())
+			if(Skeleton)
 			{
-				Skeleton->GetCurveMetaDataNames(CurveNames);
-			}
-			else // A list of curve names was provided.
-			{
-				for (const FMLDeformerCurveReference& CurveReference : Model->GetCurveIncludeList())
+				// Include all curves when no list was provided.
+				if (Model->GetCurveIncludeList().IsEmpty())
+				{
+					Skeleton->GetCurveMetaDataNames(CurveNames);
+				}
+				else // A list of curve names was provided.
 				{
-					if (CurveReference.CurveName.IsValid())
+					for (const FMLDeformerCurveReference& CurveReference : Model->GetCurveIncludeList())
 					{
-						CurveNames.Add(CurveReference.CurveName);
+						if (CurveReference.CurveName.IsValid())
+						{
+							CurveNames.Add(CurveReference.CurveName);
+						}
 					}
 				}
 			}
diff --git a/Engine/Source/Runtime/Engine/Private/Animation/AnimCompressionTypes.cpp b/Engine/Source/Runtime/Engine/Private/Animation/AnimCompressionTypes.cpp
--- a/Engine/Source/Runtime/Engine/Private/Animation/AnimCompressionTypes.cpp
+++ b/Engine/Source/Runtime/Engine/Private/Animation/AnimCompressionTypes.cpp
@@ -727,7 +727,7 @@ void FCompressibleAnimData::FetchData(const ITargetPlatform* InPlatform)
 	const UAnimSequence* AnimSequence = WeakSequence.GetEvenIfUnreachable();
 	checkf(AnimSequence, TEXT("Invalid animation sequence while trying to fetch to-compress animation data"));
 	USkeleton* Skeleton = AnimSequence->GetSkeleton();
-	checkf(Skeleton, TEXT("Invalid Skeleton while trying to fetch to-compress animation data, %s"), *Skeleton->GetPathName());
+	checkf(Skeleton, TEXT("Invalid Skeleton while trying to fetch to-compress animation data, %s"), *AnimSequence->GetPathName());
 
 	FAnimationUtils::BuildSkeletonMetaData(Skeleton, BoneData);
 
