[MLDeformer] Removed a substring redirect, which wasn't needed anymore and triggered a warning in the log.

diff --git a/Engine/Plugins/Animation/MLDeformer/MLDeformerFramework/Config/BaseMLDeformerFramework.ini b/Engine/Plugins/Animation/MLDeformer/MLDeformerFramework/Config/BaseMLDeformerFramework.ini
--- a/Engine/Plugins/Animation/MLDeformer/MLDeformerFramework/Config/BaseMLDeformerFramework.ini
+++ b/Engine/Plugins/Animation/MLDeformer/MLDeformerFramework/Config/BaseMLDeformerFramework.ini
@@ -1,7 +1,6 @@
 [CoreRedirects]
 ; Redirects from MLDeformer to MLDeformerFramework.
 +ClassRedirects=(OldName="/Script/MLDeformer.MLDeformerAsset",NewName="/Script/MLDeformerFramework.MLDeformerAsset")
-+PackageRedirects=(OldName="MLDeformer",NewName="/MLDeformerFramework/",MatchSubstring=true)
 
 ; Redirects from BoneNameStrings to BoneNames and from CurveNameStrings to CurveNames.
 +PropertyRedirects=(OldName="/Script/MLDeformerFramework.MLDeformerInputInfo.BoneNameStrings",NewName="BoneNames")
