[MLDeformer]
- Fix the editor from going into some weird state where the timeline doesn't work anymore after changing input assets, like changing a property on a skeletal mesh.
- Fixed a bug where compatibility checks would fail after training, while this wasn't really true.
- Improved the component reinits by reduced the number of inits.
- Fixed a check fail in debug mode.

diff --git a/Engine/Plugins/Animation/MLDeformer/MLDeformerFramework/Config/BaseMLDeformerFramework.ini b/Engine/Plugins/Animation/MLDeformer/MLDeformerFramework/Config/BaseMLDeformerFramework.ini
--- a/Engine/Plugins/Animation/MLDeformer/MLDeformerFramework/Config/BaseMLDeformerFramework.ini
+++ b/Engine/Plugins/Animation/MLDeformer/MLDeformerFramework/Config/BaseMLDeformerFramework.ini
@@ -2,13 +2,9 @@
 ; Redirects from MLDeformer to MLDeformerFramework.
 +ClassRedirects=(OldName="/Script/MLDeformer.MLDeformerAsset",NewName="/Script/MLDeformerFramework.MLDeformerAsset")
 
-; Redirects from BoneNameStrings to BoneNames and from CurveNameStrings to CurveNames.
+; Property redirects.
 +PropertyRedirects=(OldName="/Script/MLDeformerFramework.MLDeformerInputInfo.BoneNameStrings",NewName="BoneNames")
 +PropertyRedirects=(OldName="/Script/MLDeformerFramework.MLDeformerInputInfo.CurveNameStrings",NewName="CurveNames")
-
-; Redirect some morph target settings.
 +PropertyRedirects=(OldName="/Script/MLDeformerFramework.MLDeformerMorphModel.MorphTargetErrorTolerance",NewName="MorphCompressionLevel")
 +PropertyRedirects=(OldName="/Script/MLDeformerFramework.MLDeformerMorphModel.MorphTargetDeltaThreshold",NewName="MorphDeltaZeroThreshold")
-
-; Deprecated value inside the morph model's viz settings.
 +PropertyRedirects=(OldName="/Script/MLDeformerFramework.MLDeformerMorphModelVizSettings.MorphTargetDeltaThreshold",NewName="MorphTargetDeltaThreshold_DEPRECATED")
diff --git a/Engine/Plugins/Animation/MLDeformer/MLDeformerFramework/Source/MLDeformerFramework/Private/MLDeformerComponent.cpp b/Engine/Plugins/Animation/MLDeformer/MLDeformerFramework/Source/MLDeformerFramework/Private/MLDeformerComponent.cpp
--- a/Engine/Plugins/Animation/MLDeformer/MLDeformerFramework/Source/MLDeformerFramework/Private/MLDeformerComponent.cpp
+++ b/Engine/Plugins/Animation/MLDeformer/MLDeformerFramework/Source/MLDeformerFramework/Private/MLDeformerComponent.cpp
@@ -36,12 +36,7 @@ void UMLDeformerComponent::Init()
 	// If there is no deformer asset linked, release what we currently have.
 	if (DeformerAsset == nullptr)
 	{
-		if (ModelInstance)
-		{
-			ModelInstance->Release();
-			ModelInstance->ConditionalBeginDestroy();
-			ModelInstance = nullptr;
-		}
+		ReleaseModelInstance();
 		return;
 	}
 
@@ -49,12 +44,7 @@ void UMLDeformerComponent::Init()
 	UMLDeformerModel* Model = DeformerAsset->GetModel();
 	if (Model)
 	{
-		if (ModelInstance)
-		{
-			ModelInstance->Release();
-			ModelInstance->ConditionalBeginDestroy();
-			ModelInstance = nullptr;
-		}
+		ReleaseModelInstance();
 		ModelInstance = Model->CreateModelInstance(this);
 		ModelInstance->SetModel(Model);
 		ModelInstance->Init(SkelMeshComponent);
@@ -90,9 +80,7 @@ void UMLDeformerComponent::SetupComponent(UMLDeformerAsset* InDeformerAsset, USk
 	SkelMeshComponent = InSkelMeshComponent;
 
 	// Initialize and make sure we have a model instance.
-	RemoveNeuralNetworkModifyDelegate();
 	Init();
-	AddNeuralNetworkModifyDelegate();
 
 	// Verify that a mesh deformer has been setup when the used ML model requires one.
 	if (!bSuppressMeshDeformerLogWarnings && DeformerAsset && ModelInstance && ModelInstance->GetModel() && SkelMeshComponent)
@@ -118,42 +106,18 @@ void UMLDeformerComponent::SetupComponent(UMLDeformerAsset* InDeformerAsset, USk
 	#endif
 }
 
-void UMLDeformerComponent::AddNeuralNetworkModifyDelegate()
+void UMLDeformerComponent::ReleaseModelInstance()
 {
-	if (DeformerAsset == nullptr)
-	{
-		return;
-	}
-
 	UMLDeformerModel* Model = DeformerAsset->GetModel();
 	if (Model)
 	{
-		NeuralNetworkModifyDelegateHandle = Model->GetNeuralNetworkModifyDelegate().AddLambda
-		(
-			([this]()
-			{
-				Init();
-			})
-		);
-	}
-}
-
-void UMLDeformerComponent::RemoveNeuralNetworkModifyDelegate()
-{
-	if (DeformerAsset && 
-		NeuralNetworkModifyDelegateHandle != FDelegateHandle() && 
-		DeformerAsset->GetModel())
-	{
-		DeformerAsset->GetModel()->GetNeuralNetworkModifyDelegate().Remove(NeuralNetworkModifyDelegateHandle);
+		if (ModelInstance)
+		{
+			ModelInstance->Release();
+			ModelInstance->ConditionalBeginDestroy();
+			ModelInstance = nullptr;
+		}
 	}
-	
-	NeuralNetworkModifyDelegateHandle = FDelegateHandle();
-}
-
-void UMLDeformerComponent::BeginDestroy()
-{
-	RemoveNeuralNetworkModifyDelegate();
-	Super::BeginDestroy();
 }
 
 USkeletalMeshComponent* UMLDeformerComponent::FindSkeletalMeshComponent(const UMLDeformerAsset* const Asset) const
@@ -209,7 +173,6 @@ void UMLDeformerComponent::Deactivate()
 		TickPerfCounter.Reset();
 	#endif
 
-	RemoveNeuralNetworkModifyDelegate();
 	if (ModelInstance)
 	{
 		ModelInstance->Release();
diff --git a/Engine/Plugins/Animation/MLDeformer/MLDeformerFramework/Source/MLDeformerFramework/Public/MLDeformerComponent.h b/Engine/Plugins/Animation/MLDeformer/MLDeformerFramework/Source/MLDeformerFramework/Public/MLDeformerComponent.h
--- a/Engine/Plugins/Animation/MLDeformer/MLDeformerFramework/Source/MLDeformerFramework/Public/MLDeformerComponent.h
+++ b/Engine/Plugins/Animation/MLDeformer/MLDeformerFramework/Source/MLDeformerFramework/Public/MLDeformerComponent.h
@@ -30,7 +30,6 @@ class MLDEFORMERFRAMEWORK_API UMLDeformerComponent
 
 public:
 	// UObject overrides.
-	void BeginDestroy() override;
 #if WITH_EDITOR
 	void PostEditChangeProperty(FPropertyChangedEvent& PropertyChangedEvent) override;
 #endif
@@ -165,10 +164,20 @@ protected:
 	void Init();
 
 	/** Bind to the MLDeformerModel's NeuralNetworkModifyDelegate. */
-	void AddNeuralNetworkModifyDelegate();
+	UE_DEPRECATED(5.2, "This method will be removed.")
+	void AddNeuralNetworkModifyDelegate() {}
 
 	/** Unbind from the MLDeformerModel's NeuralNetworkModifyDelegate. */
-	void RemoveNeuralNetworkModifyDelegate();
+	UE_DEPRECATED(5.2, "This method will be removed.")
+	void RemoveNeuralNetworkModifyDelegate() {}
+
+	UE_DEPRECATED(5.3, "This method will be removed.")
+	void AddReleaseModelInstancesDelegate();
+	
+	UE_DEPRECATED(5.3, "This method will be removed.")
+	void RemoveReleaseModelInstancesDelegate();
+
+	void ReleaseModelInstance();
 
 	/** Set the ML Deformer weight. */
 	virtual void SetWeightInternal(const float NormalizedWeightValue) { Weight = FMath::Clamp<float>(NormalizedWeightValue, 0.0f, 1.0f); }
@@ -197,8 +206,9 @@ protected:
 	 */
 	TObjectPtr<USkeletalMeshComponent> SkelMeshComponent = nullptr;
 
-	/** DelegateHandle for NeuralNetwork modification. */
-	FDelegateHandle NeuralNetworkModifyDelegateHandle;
+	/** DelegateHandle for NeuralNetwork modification. This has been deprecated. */
+	UE_DEPRECATED(5.3, "This member has been deprecated.")
+	FDelegateHandle NeuralNetworkModifyDelegateHandle_DEPRECATED;
 
 	/** Suppress mesh deformer logging warnings? This is used by the ML Deformer editor, as we don't want to show some warnings when using that. */
 	bool bSuppressMeshDeformerLogWarnings = false;
diff --git a/Engine/Plugins/Animation/MLDeformer/MLDeformerFramework/Source/MLDeformerFramework/Public/MLDeformerModel.h b/Engine/Plugins/Animation/MLDeformer/MLDeformerFramework/Source/MLDeformerFramework/Public/MLDeformerModel.h
--- a/Engine/Plugins/Animation/MLDeformer/MLDeformerFramework/Source/MLDeformerFramework/Public/MLDeformerModel.h
+++ b/Engine/Plugins/Animation/MLDeformer/MLDeformerFramework/Source/MLDeformerFramework/Public/MLDeformerModel.h
@@ -81,8 +81,12 @@ namespace UE::MLDeformer
 	};
 }	// namespace UE::MLDeformer
 
+// DEPRECATED: Use FMLDeformerReleaseModelInstancesDelegate instead.
 DECLARE_EVENT_OneParam(UMLDeformerModel, FMLDeformerModelOnPostEditProperty, FPropertyChangedEvent&)
 
+/** Delegate used to signal that the UMLDeformerModelInstance's of the specified model should get released. */
+DECLARE_EVENT(UMLDeformerModel, FMLDeformerReleaseModelInstancesDelegate)
+
 /**
  * The ML Deformer runtime model base class.
  * All models should be inherited from this class.
@@ -313,14 +317,16 @@ public:
 	 * @return A reference to the GPU buffer resource object that holds the vertex map.
 	 * @see GetVertexMap
 	 */
-	const UE::MLDeformer::FVertexMapBuffer& GetVertexMapBuffer() const { return VertexMapBuffer; }
+	const UE::MLDeformer::FVertexMapBuffer& GetVertexMapBuffer() const	{ return VertexMapBuffer; }
 
 	/**
 	 * Get the neural network modified delegate.
 	 * This triggers when the neural network pointers changes.
 	 * @return A reference to the delegate.
 	 */
-	FNeuralNetworkModifyDelegate& GetNeuralNetworkModifyDelegate() { return NeuralNetworkModifyDelegate; }
+	UE_DEPRECATED(5.2, "This delegate will be removed.")
+	FNeuralNetworkModifyDelegate& GetNeuralNetworkModifyDelegate()		{ return NeuralNetworkModifyDelegate_DEPRECATED; }
+
 
 #if WITH_EDITORONLY_DATA
 	// UObject overrides.
@@ -561,7 +567,7 @@ private:
 	FRenderCommandFence RenderResourceDestroyFence;
 
 	/** Delegate that will be called immediately before the NeuralNetwork is changed. */
-	FNeuralNetworkModifyDelegate NeuralNetworkModifyDelegate;
+	FNeuralNetworkModifyDelegate NeuralNetworkModifyDelegate_DEPRECATED;
 
 	/** Cached number of skeletal mesh vertices. */
 	UPROPERTY()
diff --git a/Engine/Plugins/Animation/MLDeformer/MLDeformerFramework/Source/MLDeformerFrameworkEditor/Private/MLDeformerEditorModel.cpp b/Engine/Plugins/Animation/MLDeformer/MLDeformerFramework/Source/MLDeformerFrameworkEditor/Private/MLDeformerEditorModel.cpp
--- a/Engine/Plugins/Animation/MLDeformer/MLDeformerFramework/Source/MLDeformerFrameworkEditor/Private/MLDeformerEditorModel.cpp
+++ b/Engine/Plugins/Animation/MLDeformer/MLDeformerFramework/Source/MLDeformerFrameworkEditor/Private/MLDeformerEditorModel.cpp
@@ -43,6 +43,16 @@ namespace UE::MLDeformer
 {
 	FMLDeformerEditorModel::~FMLDeformerEditorModel()
 	{
+		if (InputObjectModifiedHandle.IsValid())
+		{
+			FCoreUObjectDelegates::OnObjectModified.Remove(InputObjectModifiedHandle);
+		}
+
+		if (InputObjectPropertyChangedHandle.IsValid())
+		{
+			FCoreUObjectDelegates::OnObjectPropertyChanged.Remove(InputObjectPropertyChangedHandle);
+		}
+
 		if (PostEditPropertyDelegateHandle.IsValid())
 		{
 			Model->OnPostEditChangeProperty().Remove(PostEditPropertyDelegateHandle);
@@ -73,7 +83,44 @@ namespace UE::MLDeformer
 		WorkingRange = TRange<double>(0.0, 100.0);
 		PlaybackRange = TRange<double>(0.0, 100.0);
 
+		// Start listening to property changes in the ML Deformer model.
 		PostEditPropertyDelegateHandle = Model->OnPostEditChangeProperty().AddRaw(this, &FMLDeformerEditorModel::OnPostEditChangeProperty);
+
+		// Start listening for object modified and object property changed events.
+		// We use this to trigger reinits of the UI and components, when one of our input assets got modified (by property changes, or reimports, etc).
+		// This listens to changes in ALL objects in the engine.
+		InputObjectModifiedHandle = FCoreUObjectDelegates::OnObjectModified.AddRaw(this, &FMLDeformerEditorModel::OnObjectModified);
+		InputObjectPropertyChangedHandle = FCoreUObjectDelegates::OnObjectPropertyChanged.AddLambda
+		(
+			[this](UObject* Object, const FPropertyChangedEvent& PropertyChangedEvent)
+			{
+				if (PropertyChangedEvent.ChangeType != EPropertyChangeType::Interactive)
+				{
+					OnObjectModified(Object);
+				}
+			}
+		);
+	}
+
+	void FMLDeformerEditorModel::OnObjectModified(UObject* Object)
+	{
+		if (Model->GetSkeletalMesh() == Object || Model->GetInputInfo()->GetSkeletalMesh() == Object)
+		{
+			UE_LOG(LogMLDeformer, Display, TEXT("Detected a modification in skeletal mesh %s, reinitializing inputs."), *Object->GetName());
+			bNeedsAssetReinit = true;
+		}
+
+		if (Model->GetAnimSequence() == Object)
+		{
+			UE_LOG(LogMLDeformer, Display, TEXT("Detected a modification in training anim sequence %s, reinitializing inputs."), *Object->GetName());
+			bNeedsAssetReinit = true;
+		}
+
+		if (Model->GetVizSettings()->GetTestAnimSequence() == Object)
+		{
+			UE_LOG(LogMLDeformer, Display, TEXT("Detected a modification in test anim sequence %s, reinitializing inputs."), *Object->GetName());
+			bNeedsAssetReinit = true;
+		}
 	}
 
 	void FMLDeformerEditorModel::AddReferencedObjects(FReferenceCollector& Collector)
@@ -310,7 +357,15 @@ namespace UE::MLDeformer
 	}
 
 	void FMLDeformerEditorModel::Tick(FEditorViewportClient* ViewportClient, float DeltaTime)
-	{
+	{		
+		if (bNeedsAssetReinit)
+		{
+			TriggerInputAssetChanged();
+			Model->InitVertexMap();
+			Model->InitGPUData();
+			bNeedsAssetReinit = false;
+		}
+
 		// Force the training sequence to use Step interpolation.
 		// We do this in the Tick, as this is reset to false in the engine sometimes.
 		UAnimSequence* TrainingAnimSequence = Model->GetAnimSequence();
@@ -324,7 +379,7 @@ namespace UE::MLDeformer
 		{
 			if (Actor && Actor->GetSkeletalMeshComponent())
 			{
-				USkeletalMeshComponent* SkelMeshComp = Actor->GetSkeletalMeshComponent();
+				USkeletalMeshComponent* SkelMeshComp = Actor->GetSkeletalMeshComponent();			
 				if (SkelMeshComp->GetAnimInstance())
 				{
 					SkelMeshComp->GetAnimInstance()->GetRequiredBones().SetUseRAWData(true);
@@ -1482,20 +1537,9 @@ namespace UE::MLDeformer
 
 	void FMLDeformerEditorModel::OnPostTraining(ETrainingResult TrainingResult, bool bUsePartiallyTrainedWhenAborted)
 	{
-		for (FMLDeformerEditorActor* EditorActor : EditorActors)
-		{
-			if (EditorActor && EditorActor->GetMLDeformerComponent())
-			{
-				USkeletalMeshComponent* SkelMeshComponent = EditorActor->GetSkeletalMeshComponent();
-				UMLDeformerAsset* DeformerAsset = Model->GetDeformerAsset();
-				EditorActor->GetMLDeformerComponent()->SetupComponent(DeformerAsset, SkelMeshComponent);
-			}
-		}
-
 		Sampler->SetVertexDeltaSpace(EVertexDeltaSpace::PostSkinning);
 		SampleDeltas();
 		Model->InitGPUData();
-
 		UpdateMemoryUsage();
 	}
 
diff --git a/Engine/Plugins/Animation/MLDeformer/MLDeformerFramework/Source/MLDeformerFrameworkEditor/Private/MLDeformerGeomCacheEditorModel.cpp b/Engine/Plugins/Animation/MLDeformer/MLDeformerFramework/Source/MLDeformerFrameworkEditor/Private/MLDeformerGeomCacheEditorModel.cpp
--- a/Engine/Plugins/Animation/MLDeformer/MLDeformerFramework/Source/MLDeformerFrameworkEditor/Private/MLDeformerGeomCacheEditorModel.cpp
+++ b/Engine/Plugins/Animation/MLDeformer/MLDeformerFramework/Source/MLDeformerFrameworkEditor/Private/MLDeformerGeomCacheEditorModel.cpp
@@ -272,6 +272,26 @@ namespace UE::MLDeformer
 
 		GetGeomCacheModel()->GetGeomCacheMeshMappings().Reset();
 	}
+
+	void FMLDeformerGeomCacheEditorModel::OnObjectModified(UObject* Object)
+	{
+		// Handle changes for the skeletal mesh and anim sequence.
+		FMLDeformerEditorModel::OnObjectModified(Object);
+
+		// Check if the object is our training geometry cache.
+		if (GetGeomCacheModel()->GetGeometryCache() == Object)
+		{
+			UE_LOG(LogMLDeformer, Display, TEXT("Detected a modification in training geometry cache %s, reinitializing inputs."), *Object->GetName());
+			bNeedsAssetReinit = true;
+		}
+
+		// Check if it is the ground truth test geom cache.
+		if (GetGeomCacheModel()->GetGeomCacheVizSettings()->GetTestGroundTruth() == Object)
+		{
+			UE_LOG(LogMLDeformer, Display, TEXT("Detected a modification in test ground truth geometry cache %s, reinitializing inputs."), *Object->GetName());
+			bNeedsAssetReinit = true;
+		}
+	}
 }	// namespace UE::MLDeformer
 
 #undef LOCTEXT_NAMESPACE
diff --git a/Engine/Plugins/Animation/MLDeformer/MLDeformerFramework/Source/MLDeformerFrameworkEditor/Public/MLDeformerEditorModel.h b/Engine/Plugins/Animation/MLDeformer/MLDeformerFramework/Source/MLDeformerFrameworkEditor/Public/MLDeformerEditorModel.h
--- a/Engine/Plugins/Animation/MLDeformer/MLDeformerFramework/Source/MLDeformerFrameworkEditor/Public/MLDeformerEditorModel.h
+++ b/Engine/Plugins/Animation/MLDeformer/MLDeformerFramework/Source/MLDeformerFrameworkEditor/Public/MLDeformerEditorModel.h
@@ -784,6 +784,15 @@ namespace UE::MLDeformer
 		 */
 		void OnPostEditChangeProperty(FPropertyChangedEvent& PropertyChangedEvent);
 
+		/**
+		 * Called whenever an engine object is being modified.
+		 * Use this to set the reinitialize assets flag. For example when a skeletal mesh or anim sequence change (being modified or reimported), we should 
+		 * set the bNeedsAssetReinit member to true, to trigger the actor components and UI etc to be reinitialized.
+		 * The default implementation of this method will check for changes in the skeletal mesh and training and test anim sequence.
+		 * If you want to check more things (your target mesh), you need to check for that.
+		 */
+		virtual void OnObjectModified(UObject* Object);
+
 		/**
 		 * Delete all editor actors.
 		 * This does not actually delete the actors in the world.
@@ -923,6 +932,15 @@ namespace UE::MLDeformer
 
 		/** Display frame numbers? */
 		bool bDisplayFrames = true;
+
+		/** The delegate that handles when an object got modified (any object). */
+		FDelegateHandle InputObjectModifiedHandle;
+
+		/** The delegate that handles when an object property got changed (any object). */
+		FDelegateHandle InputObjectPropertyChangedHandle;
+
+		/** Set to true when on next tick we need to trigger an input assets changed event. */
+		bool bNeedsAssetReinit = false;
 	};
 
 	/**
diff --git a/Engine/Plugins/Animation/MLDeformer/MLDeformerFramework/Source/MLDeformerFrameworkEditor/Public/MLDeformerGeomCacheEditorModel.h b/Engine/Plugins/Animation/MLDeformer/MLDeformerFramework/Source/MLDeformerFrameworkEditor/Public/MLDeformerGeomCacheEditorModel.h
--- a/Engine/Plugins/Animation/MLDeformer/MLDeformerFramework/Source/MLDeformerFrameworkEditor/Public/MLDeformerGeomCacheEditorModel.h
+++ b/Engine/Plugins/Animation/MLDeformer/MLDeformerFramework/Source/MLDeformerFrameworkEditor/Public/MLDeformerGeomCacheEditorModel.h
@@ -39,6 +39,7 @@ namespace UE::MLDeformer
 		virtual void UpdateIsReadyForTrainingState() override;
 		virtual void OnPropertyChanged(FPropertyChangedEvent& PropertyChangedEvent) override;
 		virtual void OnInputAssetsChanged() override;
+		virtual void OnObjectModified(UObject* Object) override;
 		virtual ETrainingResult Train() override;
 		// ~END FMLDeformerEditorModel overrides.
 		
diff --git a/Engine/Plugins/Animation/MLDeformer/NearestNeighborModel/Source/NearestNeighborModel/Private/NearestNeighborModel.cpp b/Engine/Plugins/Animation/MLDeformer/NearestNeighborModel/Source/NearestNeighborModel/Private/NearestNeighborModel.cpp
--- a/Engine/Plugins/Animation/MLDeformer/NearestNeighborModel/Source/NearestNeighborModel/Private/NearestNeighborModel.cpp
+++ b/Engine/Plugins/Animation/MLDeformer/NearestNeighborModel/Source/NearestNeighborModel/Private/NearestNeighborModel.cpp
@@ -65,12 +65,8 @@ UMLDeformerInputInfo* UNearestNeighborModel::CreateInputInfo()
 	return NearestNeighborModelInputInfo;
 }
 
-void UNearestNeighborModel::SetNNEModelData(TObjectPtr<UNNEModelData> ModelData, bool bBroadcast)
+void UNearestNeighborModel::SetNNEModelData(TObjectPtr<UNNEModelData> ModelData)
 {
-	if (bBroadcast)
-	{
-		GetNeuralNetworkModifyDelegate().Broadcast();
-	}
 	NNEModel = ModelData;
 }
 
@@ -497,11 +493,9 @@ bool UNearestNeighborModel::CheckPCAData(int32 PartId) const
 
 void UNearestNeighborModel::SetOptimizedNetwork(UNearestNeighborOptimizedNetwork* InOptimizedNetwork)
 {
-	GetNeuralNetworkModifyDelegate().Broadcast();
 	OptimizedNetwork = InOptimizedNetwork;
 }
 
-
 bool UNearestNeighborModel::DoesUseOptimizedNetwork() const
 {
 	return bUseOptimizedNetwork;
diff --git a/Engine/Plugins/Animation/MLDeformer/NearestNeighborModel/Source/NearestNeighborModel/Public/NearestNeighborModel.h b/Engine/Plugins/Animation/MLDeformer/NearestNeighborModel/Source/NearestNeighborModel/Public/NearestNeighborModel.h
--- a/Engine/Plugins/Animation/MLDeformer/NearestNeighborModel/Source/NearestNeighborModel/Public/NearestNeighborModel.h
+++ b/Engine/Plugins/Animation/MLDeformer/NearestNeighborModel/Source/NearestNeighborModel/Public/NearestNeighborModel.h
@@ -261,7 +261,7 @@ private:
 	void SetOptimizedNetwork(UNearestNeighborOptimizedNetwork* InOptimizedNetwork);
 	bool LoadOptimizedNetwork(const FString& OnnxPath);
 
-	void SetNNEModelData(TObjectPtr<UNNEModelData> ModelData, bool bBroadcast);
+	void SetNNEModelData(TObjectPtr<UNNEModelData> ModelData);
 	TObjectPtr<UNNEModelData> GetNNEModelData() const;
 
 	const FString GetNNERuntimeName() const {return TEXT("NNERuntimeORTDml"); }
diff --git a/Engine/Plugins/Animation/MLDeformer/NearestNeighborModel/Source/NearestNeighborModelEditor/Private/NearestNeighborEditorModel.cpp b/Engine/Plugins/Animation/MLDeformer/NearestNeighborModel/Source/NearestNeighborModelEditor/Private/NearestNeighborEditorModel.cpp
--- a/Engine/Plugins/Animation/MLDeformer/NearestNeighborModel/Source/NearestNeighborModelEditor/Private/NearestNeighborEditorModel.cpp
+++ b/Engine/Plugins/Animation/MLDeformer/NearestNeighborModel/Source/NearestNeighborModelEditor/Private/NearestNeighborEditorModel.cpp
@@ -151,7 +151,7 @@ namespace UE::NearestNeighborModel
 			{
 				TObjectPtr<UNNEModelData> NNEModel = LoadNeuralNetworkFromOnnx(OnnxFile);
 				NearestNeighborModel->SetUseOptimizedNetwork(false);
-				NearestNeighborModel->SetNNEModelData(NNEModel, false);
+				NearestNeighborModel->SetNNEModelData(NNEModel);
 				return true;
 			}
 		}
diff --git a/Engine/Plugins/Animation/MLDeformer/NeuralMorphModel/Source/NeuralMorphModel/Private/NeuralMorphModel.cpp b/Engine/Plugins/Animation/MLDeformer/NeuralMorphModel/Source/NeuralMorphModel/Private/NeuralMorphModel.cpp
--- a/Engine/Plugins/Animation/MLDeformer/NeuralMorphModel/Source/NeuralMorphModel/Private/NeuralMorphModel.cpp
+++ b/Engine/Plugins/Animation/MLDeformer/NeuralMorphModel/Source/NeuralMorphModel/Private/NeuralMorphModel.cpp
@@ -47,6 +47,11 @@ UMLDeformerInputInfo* UNeuralMorphModel::CreateInputInfo()
 	return NewObject<UNeuralMorphInputInfo>(this);
 }
 
+void UNeuralMorphModel::SetNeuralMorphNetwork(UNeuralMorphNetwork* Net)
+{ 
+	NeuralMorphNetwork = Net;
+}
+
 void UNeuralMorphModel::Serialize(FArchive& Archive)
 {
 	if (Archive.IsSaving() && Archive.IsCooking())
diff --git a/Engine/Plugins/Animation/MLDeformer/NeuralMorphModel/Source/NeuralMorphModel/Private/NeuralMorphNetwork.cpp b/Engine/Plugins/Animation/MLDeformer/NeuralMorphModel/Source/NeuralMorphModel/Private/NeuralMorphNetwork.cpp
--- a/Engine/Plugins/Animation/MLDeformer/NeuralMorphModel/Source/NeuralMorphModel/Private/NeuralMorphNetwork.cpp
+++ b/Engine/Plugins/Animation/MLDeformer/NeuralMorphModel/Source/NeuralMorphModel/Private/NeuralMorphNetwork.cpp
@@ -594,7 +594,6 @@ void UNeuralMorphNetworkInstance::RunLocalMLP(UNeuralMorphMLP& MLP, const FRunSe
 	{
 		const UNeuralMorphMLPLayer& CurLayer = MLP.GetLayer(LayerIndex);
 		const int32 NumInputsPerBlock = CurLayer.NumInputs;
-		checkSlow(NumInputsPerBlock == 6);
 		const int32 NumBlocks = CurLayer.Depth;
 
 		// Copy the inputs to the temp buffer.
diff --git a/Engine/Plugins/Animation/MLDeformer/NeuralMorphModel/Source/NeuralMorphModel/Public/NeuralMorphModel.h b/Engine/Plugins/Animation/MLDeformer/NeuralMorphModel/Source/NeuralMorphModel/Public/NeuralMorphModel.h
--- a/Engine/Plugins/Animation/MLDeformer/NeuralMorphModel/Source/NeuralMorphModel/Public/NeuralMorphModel.h
+++ b/Engine/Plugins/Animation/MLDeformer/NeuralMorphModel/Source/NeuralMorphModel/Public/NeuralMorphModel.h
@@ -57,10 +57,11 @@ public:
 	virtual UMLDeformerInputInfo* CreateInputInfo();
 	// ~END UMLDeformerModel overrides.
 
+	void SetNeuralMorphNetwork(UNeuralMorphNetwork* Net);
+
 	const TArray<FNeuralMorphBoneGroup>& GetBoneGroups() const		{ return BoneGroups; }
 	const TArray<FNeuralMorphCurveGroup>& GetCurveGroups() const	{ return CurveGroups; }
 	UNeuralMorphNetwork* GetNeuralMorphNetwork() const		{ return NeuralMorphNetwork.Get(); }
-	void SetNeuralMorphNetwork(UNeuralMorphNetwork* Net)	{ NeuralMorphNetwork = Net; }
 	ENeuralMorphMode GetModelMode() const					{ return Mode; }
 	int32 GetLocalNumHiddenLayers() const					{ return LocalNumHiddenLayers; }
 	int32 GetLocalNumNeuronsPerLayer() const				{ return LocalNumNeuronsPerLayer; }
