[MLDeformer] Fix a crash when creating an MLD component inside a blueprint.

diff --git a/Engine/Plugins/Animation/MLDeformer/MLDeformerFramework/Source/MLDeformerFramework/Private/MLDeformerComponent.cpp b/Engine/Plugins/Animation/MLDeformer/MLDeformerFramework/Source/MLDeformerFramework/Private/MLDeformerComponent.cpp
--- a/Engine/Plugins/Animation/MLDeformer/MLDeformerFramework/Source/MLDeformerFramework/Private/MLDeformerComponent.cpp
+++ b/Engine/Plugins/Animation/MLDeformer/MLDeformerFramework/Source/MLDeformerFramework/Private/MLDeformerComponent.cpp
@@ -108,6 +108,11 @@ void UMLDeformerComponent::SetupComponent(UMLDeformerAsset* InDeformerAsset, USk
 
 void UMLDeformerComponent::ReleaseModelInstance()
 {
+	if (!DeformerAsset)
+	{
+		return;
+	}
+
 	UMLDeformerModel* Model = DeformerAsset->GetModel();
 	if (Model)
 	{
