Fixed an issue where the ML Deformer time sampling was incorrect

diff --git a/Engine/Plugins/Animation/MLDeformer/MLDeformerFramework/Source/MLDeformerFrameworkEditor/Private/MLDeformerEditorModel.cpp b/Engine/Plugins/Animation/MLDeformer/MLDeformerFramework/Source/MLDeformerFrameworkEditor/Private/MLDeformerEditorModel.cpp
--- a/Engine/Plugins/Animation/MLDeformer/MLDeformerFramework/Source/MLDeformerFrameworkEditor/Private/MLDeformerEditorModel.cpp
+++ b/Engine/Plugins/Animation/MLDeformer/MLDeformerFramework/Source/MLDeformerFrameworkEditor/Private/MLDeformerEditorModel.cpp
@@ -36,6 +36,7 @@
 #include "GeometryCacheTrack.h"
 #include "BoneContainer.h"
 #include "Misc/ScopedSlowTask.h"
+#include "Animation/AnimData/IAnimationDataModel.h"
 
 #define LOCTEXT_NAMESPACE "MLDeformerEditorModel"
 
@@ -600,10 +601,7 @@ namespace UE::MLDeformer
 			{
 				if (EditorActor && EditorActor->IsTrainingActor())
 				{
-					if (Model->HasTrainingGroundTruth())
-					{
-						PlayOffset = GetTrainingTimeAtFrame(TargetFrame);
-					}
+					PlayOffset = GetTrainingTimeAtFrame(TargetFrame);
 					EditorActor->SetPlayPosition(PlayOffset);
 				}
 			}
@@ -627,9 +625,32 @@ namespace UE::MLDeformer
 		}
 	}
 
+	float FMLDeformerEditorModel::CorrectedFrameTime(int32 FrameNumber, float FrameTimeForFrameNumber, FFrameRate FrameRate)
+	{
+		const FFrameTime PlayOffsetCurrentTime(FrameRate.AsFrameTime(FrameTimeForFrameNumber));
+
+		if (PlayOffsetCurrentTime.FloorToFrame() != FrameNumber)
+		{
+			// correct the float for error
+			const double CorrectedPlayOffset = FrameRate.AsInterval() * (double)(FrameNumber);
+			float FloatPlayOffset = CorrectedPlayOffset;
+			FloatPlayOffset = nextafterf(FloatPlayOffset, FloatPlayOffset + 1.0f);
+			return FloatPlayOffset;
+		}
+		return FrameTimeForFrameNumber;
+	}
+
 	double FMLDeformerEditorModel::GetTrainingTimeAtFrame(int32 FrameNumber) const
 	{
-		return Model->GetAnimSequence() ? Model->GetAnimSequence()->GetTimeAtFrame(FrameNumber) : 0.0;
+		UAnimSequence* AnimSequence = Model->GetAnimSequence();
+		if (AnimSequence)
+		{
+			const FFrameRate FrameRate = AnimSequence->GetSamplingFrameRate();
+			const float UncorrectedTime = AnimSequence->GetTimeAtFrame(FrameNumber);
+			// due to floating point errors, return a double that when converted to a float cleanly maps to the value
+			return CorrectedFrameTime(FrameNumber, UncorrectedTime, FrameRate);
+		}
+		return 0.0; 
 	}
 
 	int32 FMLDeformerEditorModel::GetTrainingFrameAtTime(double TimeInSeconds) const
diff --git a/Engine/Plugins/Animation/MLDeformer/MLDeformerFramework/Source/MLDeformerFrameworkEditor/Private/MLDeformerGeomCacheEditorModel.cpp b/Engine/Plugins/Animation/MLDeformer/MLDeformerFramework/Source/MLDeformerFrameworkEditor/Private/MLDeformerGeomCacheEditorModel.cpp
--- a/Engine/Plugins/Animation/MLDeformer/MLDeformerFramework/Source/MLDeformerFrameworkEditor/Private/MLDeformerGeomCacheEditorModel.cpp
+++ b/Engine/Plugins/Animation/MLDeformer/MLDeformerFramework/Source/MLDeformerFrameworkEditor/Private/MLDeformerGeomCacheEditorModel.cpp
@@ -54,29 +54,6 @@ namespace UE::MLDeformer
 		return static_cast<FMLDeformerGeomCacheActor*>(FindEditorActor(ID));
 	}
 
-	double FMLDeformerGeomCacheEditorModel::GetTrainingTimeAtFrame(int32 FrameNumber) const
-	{
-		// Try to get the frame from the geometry cache.
-		const FMLDeformerGeomCacheActor* EditorActor = static_cast<FMLDeformerGeomCacheActor*>(FindEditorActor(ActorID_Train_GroundTruth));
-		if (EditorActor && EditorActor->GetGeometryCacheComponent() && EditorActor->GetGeometryCacheComponent()->GeometryCache.Get())
-		{
-			return EditorActor->GetGeometryCacheComponent()->GetTimeAtFrame(FrameNumber);
-		}
-
-		return FMLDeformerEditorModel::GetTrainingTimeAtFrame(FrameNumber);
-	}
-
-	int32 FMLDeformerGeomCacheEditorModel::GetTrainingFrameAtTime(double TimeInSeconds) const
-	{
-		const FMLDeformerGeomCacheActor* EditorActor = static_cast<FMLDeformerGeomCacheActor*>(FindEditorActor(ActorID_Train_GroundTruth));
-		if (EditorActor && EditorActor->GetGeometryCacheComponent() && EditorActor->GetGeometryCacheComponent()->GeometryCache.Get())
-		{
-			return EditorActor->GetGeometryCacheComponent()->GetFrameAtTime(TimeInSeconds);
-		}
-
-		return FMLDeformerEditorModel::GetTrainingFrameAtTime(TimeInSeconds);
-	}
-
 	double FMLDeformerGeomCacheEditorModel::GetTestTimeAtFrame(int32 FrameNumber) const
 	{
 		// Try to get the frame from the geometry cache.
diff --git a/Engine/Plugins/Animation/MLDeformer/MLDeformerFramework/Source/MLDeformerFrameworkEditor/Private/MLDeformerGeomCacheSampler.cpp b/Engine/Plugins/Animation/MLDeformer/MLDeformerFramework/Source/MLDeformerFrameworkEditor/Private/MLDeformerGeomCacheSampler.cpp
--- a/Engine/Plugins/Animation/MLDeformer/MLDeformerFramework/Source/MLDeformerFrameworkEditor/Private/MLDeformerGeomCacheSampler.cpp
+++ b/Engine/Plugins/Animation/MLDeformer/MLDeformerFramework/Source/MLDeformerFrameworkEditor/Private/MLDeformerGeomCacheSampler.cpp
@@ -9,6 +9,7 @@
 #include "MLDeformerGeomCacheHelpers.h"
 #include "Animation/DebugSkelMeshComponent.h"
 #include "Animation/AnimSequence.h"
+#include "Animation/AnimSingleNodeInstance.h"
 #include "Engine/SkeletalMesh.h"
 #include "Rendering/SkeletalMeshModel.h"
 #include "Rendering/SkeletalMeshLODModel.h"
@@ -136,6 +137,27 @@ namespace UE::MLDeformer
 
 	float FMLDeformerGeomCacheSampler::GetTimeAtFrame(int32 InAnimFrameIndex) const
 	{
-		return GeometryCacheComponent.Get() ? GeometryCacheComponent->GetTimeAtFrame(InAnimFrameIndex) : 0.0f;
+		// the animation instance must drive the time update
+		if (SkeletalMeshComponent)
+		{
+			UAnimSingleNodeInstance* SingleNodeInstance = Cast<UAnimSingleNodeInstance>(SkeletalMeshComponent->GetAnimInstance());
+			if (SingleNodeInstance)
+			{
+				const UAnimSequenceBase* SequenceBase = Cast<UAnimSequenceBase>(SingleNodeInstance->CurrentAsset);
+				if (SequenceBase)
+				{
+					const FFrameRate FrameRate = SequenceBase->GetSamplingFrameRate();
+					const float UncorrectedTime = SequenceBase->GetTimeAtFrame(InAnimFrameIndex);
+					const FFrameTime PlayOffsetCurrentTime(FrameRate.AsFrameTime(UncorrectedTime));
+					return FMLDeformerEditorModel::CorrectedFrameTime(InAnimFrameIndex, UncorrectedTime, FrameRate);
+				}
+			}
+		}
+
+		if (GeometryCacheComponent.Get())
+		{
+			return  GeometryCacheComponent->GetTimeAtFrame(InAnimFrameIndex);
+		}
+		return 0.0f;
 	}
 }	// namespace UE::MLDeformer
diff --git a/Engine/Plugins/Animation/MLDeformer/MLDeformerFramework/Source/MLDeformerFrameworkEditor/Public/MLDeformerEditorModel.h b/Engine/Plugins/Animation/MLDeformer/MLDeformerFramework/Source/MLDeformerFrameworkEditor/Public/MLDeformerEditorModel.h
--- a/Engine/Plugins/Animation/MLDeformer/MLDeformerFramework/Source/MLDeformerFrameworkEditor/Public/MLDeformerEditorModel.h
+++ b/Engine/Plugins/Animation/MLDeformer/MLDeformerFramework/Source/MLDeformerFrameworkEditor/Public/MLDeformerEditorModel.h
@@ -777,6 +777,15 @@ namespace UE::MLDeformer
 		 */
 		UMLDeformerComponent* FindMLDeformerComponent(int32 ActorID = ActorID_Test_MLDeformed) const;
 
+
+		/**
+		 * Correct floating point errors that can cause issues when sampling animation using step timing.
+		 * @param FrameNumber The desired frame number
+		 * @param FrameTimeForFrameNumber The frame time computed for the frame number, which may have conversion errors
+		 * @Param FrameRate The frame rate for conversion
+		 * @return A corrected frame time, which is equal to FrameTimeForFrameNumber unless a floating point error means the floor is not equal to the FrameNumber
+		 */
+		static float CorrectedFrameTime(int32 FrameNumber, float FrameTimeForFrameNumber, FFrameRate FrameRate);
 	protected:
 		/**
 		 * Executed when a property changes by a change in the UI.
diff --git a/Engine/Plugins/Animation/MLDeformer/MLDeformerFramework/Source/MLDeformerFrameworkEditor/Public/MLDeformerGeomCacheEditorModel.h b/Engine/Plugins/Animation/MLDeformer/MLDeformerFramework/Source/MLDeformerFrameworkEditor/Public/MLDeformerGeomCacheEditorModel.h
--- a/Engine/Plugins/Animation/MLDeformer/MLDeformerFramework/Source/MLDeformerFrameworkEditor/Public/MLDeformerGeomCacheEditorModel.h
+++ b/Engine/Plugins/Animation/MLDeformer/MLDeformerFramework/Source/MLDeformerFrameworkEditor/Public/MLDeformerGeomCacheEditorModel.h
@@ -32,8 +32,6 @@ namespace UE::MLDeformer
 		virtual void CreateTrainingGroundTruthActor(UWorld* World) override;
 		virtual void CreateTestGroundTruthActor(UWorld* World) override;
 		virtual int32 GetNumTrainingFrames() const override;
-		virtual double GetTrainingTimeAtFrame(int32 FrameNumber) const override;
-		virtual int32 GetTrainingFrameAtTime(double TimeInSeconds) const override;
 		virtual double GetTestTimeAtFrame(int32 FrameNumber) const override;
 		virtual int32 GetTestFrameAtTime(double TimeInSeconds) const override;
 		virtual void UpdateIsReadyForTrainingState() override;
