MLDeformerFramework: refactoring how render fence is used to avoid waiting in `BeginDestroy()` that flushes render fence bundling on shutdown

diff --git a/Engine/Plugins/Animation/MLDeformer/MLDeformerFramework/Source/MLDeformerFramework/Private/MLDeformerComponent.cpp b/Engine/Plugins/Animation/MLDeformer/MLDeformerFramework/Source/MLDeformerFramework/Private/MLDeformerComponent.cpp
--- a/Engine/Plugins/Animation/MLDeformer/MLDeformerFramework/Source/MLDeformerFramework/Private/MLDeformerComponent.cpp
+++ b/Engine/Plugins/Animation/MLDeformer/MLDeformerFramework/Source/MLDeformerFramework/Private/MLDeformerComponent.cpp
@@ -108,21 +108,8 @@ void UMLDeformerComponent::SetupComponent(UMLDeformerAsset* InDeformerAsset, USk
 
 void UMLDeformerComponent::ReleaseModelInstance()
 {
-	if (!DeformerAsset)
-	{
-		return;
-	}
-
-	UMLDeformerModel* Model = DeformerAsset->GetModel();
-	if (Model)
-	{
-		if (ModelInstance)
-		{
-			ModelInstance->Release();
-			ModelInstance->ConditionalBeginDestroy();
-			ModelInstance = nullptr;
-		}
-	}
+	ModelInstance->ConditionalBeginDestroy(); // force destruction immediately instead of waiting for the next GC
+	ModelInstance = nullptr;
 }
 
 USkeletalMeshComponent* UMLDeformerComponent::FindSkeletalMeshComponent(const UMLDeformerAsset* const Asset) const
@@ -180,7 +167,7 @@ void UMLDeformerComponent::Deactivate()
 
 	if (ModelInstance)
 	{
-		ModelInstance->Release();
+		ModelInstance->ConditionalBeginDestroy();
 		ModelInstance = nullptr;
 	}
 	Super::Deactivate();
diff --git a/Engine/Plugins/Animation/MLDeformer/MLDeformerFramework/Source/MLDeformerFramework/Private/MLDeformerModelInstance.cpp b/Engine/Plugins/Animation/MLDeformer/MLDeformerFramework/Source/MLDeformerFramework/Private/MLDeformerModelInstance.cpp
--- a/Engine/Plugins/Animation/MLDeformer/MLDeformerFramework/Source/MLDeformerFramework/Private/MLDeformerModelInstance.cpp
+++ b/Engine/Plugins/Animation/MLDeformer/MLDeformerFramework/Source/MLDeformerFramework/Private/MLDeformerModelInstance.cpp
@@ -22,16 +22,15 @@ DEFINE_STAT(STAT_MLDeformerInference);
 
 void UMLDeformerModelInstance::BeginDestroy()
 {
-	Release();
+	RenderCommandFence.BeginFence();
 	Super::BeginDestroy();
 }
 
-void UMLDeformerModelInstance::Release()
+bool UMLDeformerModelInstance::IsReadyForFinishDestroy()
 {
 	// Wait for the render commands to finish, because some neural network might still be executing, using the inference handle
 	// that we are about to delete.
-	RenderCommandFence.BeginFence();
-	RenderCommandFence.Wait();
+	return Super::IsReadyForFinishDestroy() && RenderCommandFence.IsFenceComplete();
 }
 
 USkeletalMeshComponent* UMLDeformerModelInstance::GetSkeletalMeshComponent() const
diff --git a/Engine/Plugins/Animation/MLDeformer/MLDeformerFramework/Source/MLDeformerFramework/Private/MLDeformerMorphModel.cpp b/Engine/Plugins/Animation/MLDeformer/MLDeformerFramework/Source/MLDeformerFramework/Private/MLDeformerMorphModel.cpp
--- a/Engine/Plugins/Animation/MLDeformer/MLDeformerFramework/Source/MLDeformerFramework/Private/MLDeformerMorphModel.cpp
+++ b/Engine/Plugins/Animation/MLDeformer/MLDeformerFramework/Source/MLDeformerFramework/Private/MLDeformerMorphModel.cpp
@@ -88,17 +88,21 @@ int32 UMLDeformerMorphModel::GetMorphTargetDeltaStartIndex(int32 BlendShapeIndex
 
 void UMLDeformerMorphModel::BeginDestroy()
 {
-	if (MorphTargetSet.IsValid())
+	if (MorphTargetSet)
 	{
-		// Release and flush, waiting for the release to have completed, 
-		// If we don't do this we can get an error that we destroy a render resource that is still initialized,
-		// as the release happens in another thread.
-		ReleaseResourceAndFlush(&MorphTargetSet->MorphBuffers);
+		BeginReleaseResource(&MorphTargetSet->MorphBuffers);
+		RenderCommandFence.BeginFence();
 		MorphTargetSet.Reset();
 	}
 	Super::BeginDestroy();
 }
 
+bool UMLDeformerMorphModel::IsReadyForFinishDestroy()
+{
+	// wait for associated render resources to be released
+	return Super::IsReadyForFinishDestroy() && RenderCommandFence.IsFenceComplete();
+}
+
 void UMLDeformerMorphModel::SetMorphTargetsErrorOrder(const TArray<int32>& MorphTargetOrder, const TArray<float>& ErrorValues)
 {
 	MorphTargetErrorOrder = MorphTargetOrder;
diff --git a/Engine/Plugins/Animation/MLDeformer/MLDeformerFramework/Source/MLDeformerFramework/Private/MLDeformerMorphModelInstance.cpp b/Engine/Plugins/Animation/MLDeformer/MLDeformerFramework/Source/MLDeformerFramework/Private/MLDeformerMorphModelInstance.cpp
--- a/Engine/Plugins/Animation/MLDeformer/MLDeformerFramework/Source/MLDeformerFramework/Private/MLDeformerMorphModelInstance.cpp
+++ b/Engine/Plugins/Animation/MLDeformer/MLDeformerFramework/Source/MLDeformerFramework/Private/MLDeformerMorphModelInstance.cpp
@@ -28,7 +28,7 @@ int32 UMLDeformerMorphModelInstance::GetExternalMorphSetID() const
 	return ExternalMorphSetID;
 }
 
-void UMLDeformerMorphModelInstance::Release()
+void UMLDeformerMorphModelInstance::BeginDestroy()
 {
 	// Try to unregister the morph target and morph target set.
 	if (SkeletalMeshComponent)
@@ -42,7 +42,7 @@ void UMLDeformerMorphModelInstance::Release()
 		}
 	}
 
-	Super::Release();
+	Super::BeginDestroy();
 }
 
 void UMLDeformerMorphModelInstance::PostMLDeformerComponentInit()
diff --git a/Engine/Plugins/Animation/MLDeformer/MLDeformerFramework/Source/MLDeformerFramework/Public/MLDeformerModelInstance.h b/Engine/Plugins/Animation/MLDeformer/MLDeformerFramework/Source/MLDeformerFramework/Public/MLDeformerModelInstance.h
--- a/Engine/Plugins/Animation/MLDeformer/MLDeformerFramework/Source/MLDeformerFramework/Public/MLDeformerModelInstance.h
+++ b/Engine/Plugins/Animation/MLDeformer/MLDeformerFramework/Source/MLDeformerFramework/Public/MLDeformerModelInstance.h
@@ -36,6 +36,7 @@ public:
 
 	// UObject overrides.
 	virtual void BeginDestroy() override;
+	virtual bool IsReadyForFinishDestroy() override;
 	// ~END UObject overrides.
 	
 	/**
@@ -57,7 +58,8 @@ public:
 	 * Release all built lookup tables etc.
 	 * Also call this base class method if you override this method.
 	 */
-	virtual void Release();
+	UE_DEPRECATED(5.3, "This functionality was moved to BeginDestroy() and IsReadyForFinishDestroy(). Most probably you don't need to call this method anymore. If you've been overriding it, please move the logic to BeginDestroy() and IsReadyForFinishDestroy() in a similar manner.")
+	virtual void Release() {}
 
 	/**
 	 * Update the deformer instance.
diff --git a/Engine/Plugins/Animation/MLDeformer/MLDeformerFramework/Source/MLDeformerFramework/Public/MLDeformerMorphModel.h b/Engine/Plugins/Animation/MLDeformer/MLDeformerFramework/Source/MLDeformerFramework/Public/MLDeformerMorphModel.h
--- a/Engine/Plugins/Animation/MLDeformer/MLDeformerFramework/Source/MLDeformerFramework/Public/MLDeformerMorphModel.h
+++ b/Engine/Plugins/Animation/MLDeformer/MLDeformerFramework/Source/MLDeformerFramework/Public/MLDeformerMorphModel.h
@@ -32,6 +32,7 @@ public:
 
 	// UObject overrides.
 	virtual void BeginDestroy() override;
+	virtual bool IsReadyForFinishDestroy() override;
 	virtual void PostLoad() override;
 	// ~END UObject overrides.
 
@@ -295,4 +296,7 @@ private:
 	 */
 	UPROPERTY(EditAnywhere, Category = "Morph Targets", meta = (EditCondition = "MaskChannel != EMLDeformerMaskChannel::Disabled"))
 	bool bInvertMaskChannel = false;
+
+	/** The fence that let's us wait for all render commands to finish, before this instance is destroyed. */
+	FRenderCommandFence RenderCommandFence;
 };
diff --git a/Engine/Plugins/Animation/MLDeformer/MLDeformerFramework/Source/MLDeformerFramework/Public/MLDeformerMorphModelInstance.h b/Engine/Plugins/Animation/MLDeformer/MLDeformerFramework/Source/MLDeformerFramework/Public/MLDeformerMorphModelInstance.h
--- a/Engine/Plugins/Animation/MLDeformer/MLDeformerFramework/Source/MLDeformerFramework/Public/MLDeformerMorphModelInstance.h
+++ b/Engine/Plugins/Animation/MLDeformer/MLDeformerFramework/Source/MLDeformerFramework/Public/MLDeformerMorphModelInstance.h
@@ -27,7 +27,7 @@ public:
 	// UMLDeformerModelInstance overrides.
 	virtual void Init(USkeletalMeshComponent* SkelMeshComponent) override;
 	virtual void PostMLDeformerComponentInit() override;
-	virtual void Release() override;
+	virtual void BeginDestroy() override;
 	virtual void Tick(float DeltaTime, float ModelWeight) override;
 	virtual void HandleZeroModelWeight() override;
 	virtual bool IsValidForDataProvider() const override;
diff --git a/Engine/Plugins/Animation/MLDeformer/NearestNeighborModel/Source/NearestNeighborModelEditor/Private/NearestNeighborTrainingModel.cpp b/Engine/Plugins/Animation/MLDeformer/NearestNeighborModel/Source/NearestNeighborModelEditor/Private/NearestNeighborTrainingModel.cpp
--- a/Engine/Plugins/Animation/MLDeformer/NearestNeighborModel/Source/NearestNeighborModelEditor/Private/NearestNeighborTrainingModel.cpp
+++ b/Engine/Plugins/Animation/MLDeformer/NearestNeighborModel/Source/NearestNeighborModelEditor/Private/NearestNeighborTrainingModel.cpp
@@ -123,7 +123,6 @@ UNearestNeighborModelInstance* UNearestNeighborTrainingModel::CreateModelInstanc
 
 void UNearestNeighborTrainingModel::DestroyModelInstance(UNearestNeighborModelInstance* ModelInstance)
 {
-	ModelInstance->Release();
 	ModelInstance->ConditionalBeginDestroy();
 }
 
diff --git a/Engine/Plugins/Animation/MLDeformer/VertexDeltaModel/Source/VertexDeltaModel/Private/VertexDeltaModelInstance.cpp b/Engine/Plugins/Animation/MLDeformer/VertexDeltaModel/Source/VertexDeltaModel/Private/VertexDeltaModelInstance.cpp
--- a/Engine/Plugins/Animation/MLDeformer/VertexDeltaModel/Source/VertexDeltaModel/Private/VertexDeltaModelInstance.cpp
+++ b/Engine/Plugins/Animation/MLDeformer/VertexDeltaModel/Source/VertexDeltaModel/Private/VertexDeltaModelInstance.cpp
@@ -13,17 +13,6 @@ UVertexDeltaModel* UVertexDeltaModelInstance::GetVertexDeltaModel() const
 }
 
 
-void UVertexDeltaModelInstance::Release()
-{
-	Super::Release();
-
-	// Destroy the neural network instance.
-	if (GetVertexDeltaModel())
-	{
-
-	}
-}
-
 UE::NNECore::IModelRDG* UVertexDeltaModelInstance::GetNNEModelRDG() const
 {
 	return ModelRDG.Get();
diff --git a/Engine/Plugins/Animation/MLDeformer/VertexDeltaModel/Source/VertexDeltaModel/Public/VertexDeltaModelInstance.h b/Engine/Plugins/Animation/MLDeformer/VertexDeltaModel/Source/VertexDeltaModel/Public/VertexDeltaModelInstance.h
--- a/Engine/Plugins/Animation/MLDeformer/VertexDeltaModel/Source/VertexDeltaModel/Public/VertexDeltaModelInstance.h
+++ b/Engine/Plugins/Animation/MLDeformer/VertexDeltaModel/Source/VertexDeltaModel/Public/VertexDeltaModelInstance.h
@@ -20,7 +20,6 @@ class VERTEXDELTAMODEL_API UVertexDeltaModelInstance
 
 public:
 	// UMLDeformerModelInstance overrides.
-	virtual void Release() override;
 	virtual FString CheckCompatibility(USkeletalMeshComponent* InSkelMeshComponent, bool bLogIssues) override;
 	virtual bool IsValidForDataProvider() const override;
 	virtual void Execute(float ModelWeight) override;
