[MLDeformer] Fixed a crash when opening an MLD asset.

#trivial
#rnx

diff --git a/Engine/Plugins/Animation/MLDeformer/MLDeformerFramework/Source/MLDeformerFramework/Private/MLDeformerComponent.cpp b/Engine/Plugins/Animation/MLDeformer/MLDeformerFramework/Source/MLDeformerFramework/Private/MLDeformerComponent.cpp
--- a/Engine/Plugins/Animation/MLDeformer/MLDeformerFramework/Source/MLDeformerFramework/Private/MLDeformerComponent.cpp
+++ b/Engine/Plugins/Animation/MLDeformer/MLDeformerFramework/Source/MLDeformerFramework/Private/MLDeformerComponent.cpp
@@ -108,8 +108,11 @@ void UMLDeformerComponent::SetupComponent(UMLDeformerAsset* InDeformerAsset, USk
 
 void UMLDeformerComponent::ReleaseModelInstance()
 {
-	ModelInstance->ConditionalBeginDestroy(); // force destruction immediately instead of waiting for the next GC
-	ModelInstance = nullptr;
+	if (ModelInstance)
+	{
+		ModelInstance->ConditionalBeginDestroy(); // Force destruction immediately instead of waiting for the next GC.
+		ModelInstance = nullptr;
+	}
 }
 
 USkeletalMeshComponent* UMLDeformerComponent::FindSkeletalMeshComponent(const UMLDeformerAsset* const Asset) const
