[MLDeformer] Fixed a crash on shutdown, related to releasing render resources.

#trivial
#rnx

diff --git a/Engine/Plugins/Animation/MLDeformer/MLDeformerFramework/Source/MLDeformerFramework/Private/MLDeformerMorphModel.cpp b/Engine/Plugins/Animation/MLDeformer/MLDeformerFramework/Source/MLDeformerFramework/Private/MLDeformerMorphModel.cpp
--- a/Engine/Plugins/Animation/MLDeformer/MLDeformerFramework/Source/MLDeformerFramework/Private/MLDeformerMorphModel.cpp
+++ b/Engine/Plugins/Animation/MLDeformer/MLDeformerFramework/Source/MLDeformerFramework/Private/MLDeformerMorphModel.cpp
@@ -92,14 +92,13 @@ void UMLDeformerMorphModel::BeginDestroy()
 	{
 		BeginReleaseResource(&MorphTargetSet->MorphBuffers);
 		RenderCommandFence.BeginFence();
-		MorphTargetSet.Reset();
 	}
 	Super::BeginDestroy();
 }
 
 bool UMLDeformerMorphModel::IsReadyForFinishDestroy()
 {
-	// wait for associated render resources to be released
+	// Wait for associated render resources to be released.
 	return Super::IsReadyForFinishDestroy() && RenderCommandFence.IsFenceComplete();
 }
 
