MLDeformer: more robust frame time

diff --git a/Engine/Plugins/Animation/MLDeformer/MLDeformerFramework/Source/MLDeformerFramework/Private/MLDeformerGeomCacheHelpers.cpp b/Engine/Plugins/Animation/MLDeformer/MLDeformerFramework/Source/MLDeformerFramework/Private/MLDeformerGeomCacheHelpers.cpp
--- a/Engine/Plugins/Animation/MLDeformer/MLDeformerFramework/Source/MLDeformerFramework/Private/MLDeformerGeomCacheHelpers.cpp
+++ b/Engine/Plugins/Animation/MLDeformer/MLDeformerFramework/Source/MLDeformerFramework/Private/MLDeformerGeomCacheHelpers.cpp
@@ -237,6 +237,14 @@ namespace UE::MLDeformer
 		} // For all tracks.
 	}
 
+	float SnapToNearestFrame(const UGeometryCache& Cache, float InTime)
+	{
+		const int32 NumFrames = Cache.GetEndFrame() - Cache.GetStartFrame() + 1;
+		const float FrameTime = NumFrames > 1 ? Cache.CalculateDuration() / (float)(NumFrames - 1) : 0.0f;
+		const int32 NormalizedFrame = FMath::Clamp(FMath::RoundToInt(InTime / FrameTime), 0, NumFrames - 1);
+		return NormalizedFrame * FrameTime;
+	}
+
 	void SampleGeomCachePositions(
 		int32 InLODIndex,
 		float InSampleTime,
@@ -277,7 +285,7 @@ namespace UE::MLDeformer
 			UGeometryCacheTrack* Track = InGeometryCache->Tracks[MeshMapping.TrackIndex];
 
 			FGeometryCacheMeshData GeomCacheMeshData;
-			if (!Track->GetMeshDataAtTime(InSampleTime, GeomCacheMeshData))
+			if (!Track->GetMeshDataAtTime(SnapToNearestFrame(*InGeometryCache, InSampleTime), GeomCacheMeshData))
 			{
 				continue;
 			}
diff --git a/Engine/Plugins/Animation/MLDeformer/MLDeformerFramework/Source/MLDeformerFrameworkEditor/Private/MLDeformerGeomCacheSampler.cpp b/Engine/Plugins/Animation/MLDeformer/MLDeformerFramework/Source/MLDeformerFrameworkEditor/Private/MLDeformerGeomCacheSampler.cpp
--- a/Engine/Plugins/Animation/MLDeformer/MLDeformerFramework/Source/MLDeformerFrameworkEditor/Private/MLDeformerGeomCacheSampler.cpp
+++ b/Engine/Plugins/Animation/MLDeformer/MLDeformerFramework/Source/MLDeformerFrameworkEditor/Private/MLDeformerGeomCacheSampler.cpp
@@ -77,7 +77,7 @@ namespace UE::MLDeformer
 
 				// Sample the mesh data of the geom cache.
 				FGeometryCacheMeshData& GeomCacheMeshData = GeomCacheMeshDatas[MeshMappingIndex];
-				if (!Track->GetMeshDataAtTime(SampleTime, GeomCacheMeshData))
+				if (!Track->GetMeshDataAtTime(GeometryCacheComponent->GetTimeAtFrame(InAnimFrameIndex), GeomCacheMeshData))
 				{
 					continue;
 				}
diff --git a/Engine/Plugins/Animation/MLDeformer/NearestNeighborModel/Source/NearestNeighborModelEditor/Private/NearestNeighborGeomCacheSampler.cpp b/Engine/Plugins/Animation/MLDeformer/NearestNeighborModel/Source/NearestNeighborModelEditor/Private/NearestNeighborGeomCacheSampler.cpp
--- a/Engine/Plugins/Animation/MLDeformer/NearestNeighborModel/Source/NearestNeighborModelEditor/Private/NearestNeighborGeomCacheSampler.cpp
+++ b/Engine/Plugins/Animation/MLDeformer/NearestNeighborModel/Source/NearestNeighborModelEditor/Private/NearestNeighborGeomCacheSampler.cpp
@@ -54,7 +54,7 @@ namespace UE::NearestNeighborModel
 
 				// Sample the mesh data of the geom cache.
 				FGeometryCacheMeshData& GeomCacheMeshData = GeomCacheMeshDatas[MeshMappingIndex];
-				if (!Track->GetMeshDataAtTime(SampleTime, GeomCacheMeshData))
+				if (!Track->GetMeshDataAtTime(GeometryCacheComponent->GetTimeAtFrame(InAnimFrameIndex), GeomCacheMeshData))
 				{
 					continue;
 				}
