[MLDeformer]
- Add support for automatic bone influence masks in the local mode of the Neural Morph Model. This improves convergence of the model as well as opening the door for training on 4D data.
- Improvements to the UI of the MLD asset editor's input bone/curve/bonegroups/curvegroups setup.

diff --git a/Engine/Plugins/Animation/MLDeformer/MLDeformerFramework/Config/BaseMLDeformerFramework.ini b/Engine/Plugins/Animation/MLDeformer/MLDeformerFramework/Config/BaseMLDeformerFramework.ini
--- a/Engine/Plugins/Animation/MLDeformer/MLDeformerFramework/Config/BaseMLDeformerFramework.ini
+++ b/Engine/Plugins/Animation/MLDeformer/MLDeformerFramework/Config/BaseMLDeformerFramework.ini
@@ -8,3 +8,7 @@
 +PropertyRedirects=(OldName="/Script/MLDeformerFramework.MLDeformerMorphModel.MorphTargetErrorTolerance",NewName="MorphCompressionLevel")
 +PropertyRedirects=(OldName="/Script/MLDeformerFramework.MLDeformerMorphModel.MorphTargetDeltaThreshold",NewName="MorphDeltaZeroThreshold")
 +PropertyRedirects=(OldName="/Script/MLDeformerFramework.MLDeformerMorphModelVizSettings.MorphTargetDeltaThreshold",NewName="MorphTargetDeltaThreshold_DEPRECATED")
+
+; Deprecated model properties.
++PropertyRedirects=(OldName="/Script/MLDeformerFramework.MLDeformerModel.bIncludeBones",NewName="bIncludeBones_DEPRECATED")
++PropertyRedirects=(OldName="/Script/MLDeformerFramework.MLDeformerModel.bIncludeCurves",NewName="bIncludeCurves_DEPRECATED")
diff --git a/Engine/Plugins/Animation/MLDeformer/MLDeformerFramework/Source/MLDeformerFramework/Private/MLDeformerComponent.cpp b/Engine/Plugins/Animation/MLDeformer/MLDeformerFramework/Source/MLDeformerFramework/Private/MLDeformerComponent.cpp
--- a/Engine/Plugins/Animation/MLDeformer/MLDeformerFramework/Source/MLDeformerFramework/Private/MLDeformerComponent.cpp
+++ b/Engine/Plugins/Animation/MLDeformer/MLDeformerFramework/Source/MLDeformerFramework/Private/MLDeformerComponent.cpp
@@ -49,6 +49,7 @@ void UMLDeformerComponent::Init()
 		ModelInstance->SetModel(Model);
 		ModelInstance->Init(SkelMeshComponent);
 		ModelInstance->PostMLDeformerComponentInit();
+		BindDelegates();
 	}
 	else
 	{
@@ -80,7 +81,9 @@ void UMLDeformerComponent::SetupComponent(UMLDeformerAsset* InDeformerAsset, USk
 	SkelMeshComponent = InSkelMeshComponent;
 
 	// Initialize and make sure we have a model instance.
+	UnbindDelegates();
 	Init();
+	BindDelegates();
 
 	// Verify that a mesh deformer has been setup when the used ML model requires one.
 	if (!bSuppressMeshDeformerLogWarnings && DeformerAsset && ModelInstance && ModelInstance->GetModel() && SkelMeshComponent)
@@ -106,6 +109,36 @@ void UMLDeformerComponent::SetupComponent(UMLDeformerAsset* InDeformerAsset, USk
 	#endif
 }
 
+void UMLDeformerComponent::BindDelegates()
+{
+	UMLDeformerModel* Model = DeformerAsset ? DeformerAsset->GetModel() : nullptr;
+	if (Model)
+	{
+		ReinitModelInstanceDelegateHandle = Model->GetReinitModelInstanceDelegate().AddLambda(
+			[this]()
+			{
+				Init();
+			});
+	}
+}
+
+void UMLDeformerComponent::UnbindDelegates()
+{
+	UMLDeformerModel* Model = DeformerAsset ? DeformerAsset->GetModel() : nullptr;
+	if (Model && ReinitModelInstanceDelegateHandle.IsValid())
+	{
+		Model->GetReinitModelInstanceDelegate().Remove(ReinitModelInstanceDelegateHandle);	
+	}
+
+	ReinitModelInstanceDelegateHandle = FDelegateHandle();
+}
+
+void UMLDeformerComponent::BeginDestroy()
+{
+	UnbindDelegates();
+	Super::BeginDestroy();
+}
+
 void UMLDeformerComponent::ReleaseModelInstance()
 {
 	if (ModelInstance)
@@ -168,6 +201,7 @@ void UMLDeformerComponent::Deactivate()
 		TickPerfCounter.Reset();
 	#endif
 
+	UnbindDelegates();
 	if (ModelInstance)
 	{
 		ModelInstance->ConditionalBeginDestroy();
@@ -214,6 +248,17 @@ void UMLDeformerComponent::TickComponent(float DeltaTime, enum ELevelTick TickTy
 	#endif
 }
 
+void UMLDeformerComponent::SetWeightInternal(const float NormalizedWeightValue)
+{ 
+	Weight = FMath::Clamp<float>(NormalizedWeightValue, 0.0f, 1.0f);
+}
+
+void UMLDeformerComponent::SetDeformerAssetInternal(UMLDeformerAsset* const InDeformerAsset)
+{ 
+	DeformerAsset = InDeformerAsset;
+	UpdateSkeletalMeshComponent();
+}
+
 #if WITH_EDITOR
 	void UMLDeformerComponent::PostEditChangeProperty(FPropertyChangedEvent& PropertyChangedEvent)
 	{
diff --git a/Engine/Plugins/Animation/MLDeformer/MLDeformerFramework/Source/MLDeformerFramework/Private/MLDeformerInputInfo.cpp b/Engine/Plugins/Animation/MLDeformer/MLDeformerFramework/Source/MLDeformerFramework/Private/MLDeformerInputInfo.cpp
--- a/Engine/Plugins/Animation/MLDeformer/MLDeformerFramework/Source/MLDeformerFramework/Private/MLDeformerInputInfo.cpp
+++ b/Engine/Plugins/Animation/MLDeformer/MLDeformerFramework/Source/MLDeformerFramework/Private/MLDeformerInputInfo.cpp
@@ -88,11 +88,21 @@ TArray<FName>& UMLDeformerInputInfo::GetBoneNames()
 	return BoneNames;
 }
 
+const TArray<FName>& UMLDeformerInputInfo::GetBoneNames() const
+{ 
+	return BoneNames;
+}
+
 TArray<FName>& UMLDeformerInputInfo::GetCurveNames()
 { 
 	return CurveNames;
 }
 
+const TArray<FName>& UMLDeformerInputInfo::GetCurveNames() const
+{ 
+	return CurveNames;
+}
+
 FName UMLDeformerInputInfo::GetCurveName(int32 Index) const
 { 
 	return CurveNames[Index];
@@ -125,7 +135,7 @@ bool UMLDeformerInputInfo::IsCompatible(USkeletalMesh* InSkeletalMesh) const
 		return false;
 	}
 
-	if (FSoftObjectPath(InSkeletalMesh) != SkeletalMesh)
+	if ((SkeletalMesh.IsValid() && !SkeletalMesh.IsNull()) && FSoftObjectPath(InSkeletalMesh) != SkeletalMesh)
 	{
 		return false;
 	}
diff --git a/Engine/Plugins/Animation/MLDeformer/MLDeformerFramework/Source/MLDeformerFramework/Private/MLDeformerModelInstance.cpp b/Engine/Plugins/Animation/MLDeformer/MLDeformerFramework/Source/MLDeformerFramework/Private/MLDeformerModelInstance.cpp
--- a/Engine/Plugins/Animation/MLDeformer/MLDeformerFramework/Source/MLDeformerFramework/Private/MLDeformerModelInstance.cpp
+++ b/Engine/Plugins/Animation/MLDeformer/MLDeformerFramework/Source/MLDeformerFramework/Private/MLDeformerModelInstance.cpp
@@ -134,10 +134,9 @@ FString UMLDeformerModelInstance::CheckCompatibility(USkeletalMeshComponent* InS
 	if (SkelMesh && !InputInfo->IsCompatible(SkelMesh) && Model->GetDeformerAsset())
 	{
 		ErrorText += InputInfo->GenerateCompatibilityErrorString(SkelMesh);
-		ErrorText += "\n";
 		if (bLogIssues)
 		{
-			UE_LOG(LogMLDeformer, Error, TEXT("ML Deformer '%s' isn't compatible with Skeletal Mesh '%s'.\nReason(s):\n%s"), 
+			UE_LOG(LogMLDeformer, Error, TEXT("ML Deformer '%s' isn't compatible with Skeletal Mesh '%s'.\nReason(s):\n%s\n"), 
 				*Model->GetDeformerAsset()->GetName(), 
 				*SkelMesh->GetName(), 
 				*ErrorText);
diff --git a/Engine/Plugins/Animation/MLDeformer/MLDeformerFramework/Source/MLDeformerFramework/Private/MLDeformerMorphModel.cpp b/Engine/Plugins/Animation/MLDeformer/MLDeformerFramework/Source/MLDeformerFramework/Private/MLDeformerMorphModel.cpp
--- a/Engine/Plugins/Animation/MLDeformer/MLDeformerFramework/Source/MLDeformerFramework/Private/MLDeformerMorphModel.cpp
+++ b/Engine/Plugins/Animation/MLDeformer/MLDeformerFramework/Source/MLDeformerFramework/Private/MLDeformerMorphModel.cpp
@@ -4,6 +4,7 @@
 #include "MLDeformerMorphModelInstance.h"
 #include "MLDeformerModelInstance.h"
 #include "MLDeformerComponent.h"
+#include "MLDeformerMorphModelInputInfo.h"
 #include "Components/ExternalMorphSet.h"
 #include "Components/SkeletalMeshComponent.h"
 #include "Rendering/MorphTargetVertexInfoBuffers.h"
@@ -29,8 +30,9 @@ void UMLDeformerMorphModel::Serialize(FArchive& Archive)
 	bool bHasMorphData = false;
 	if (Archive.IsSaving())
 	{
+		// Strip editor only data on cook.
 		if (Archive.IsCooking())
-		{
+		{			
 			MorphTargetDeltas.Empty();
 		}
 
@@ -49,6 +51,19 @@ void UMLDeformerMorphModel::Serialize(FArchive& Archive)
 void UMLDeformerMorphModel::PostLoad()
 {
 	Super::PostLoad();
+
+	// If we have an input info, but it isn't one inherited from the MorphInputInfo, try to create a new one.
+	// This is because we introduced a UMLDeformerMorphModelInputInfo later on, and we want to convert old assets to use this new class.
+	UMLDeformerInputInfo* CurrentInputInfo = GetInputInfo();
+	if (CurrentInputInfo && !CurrentInputInfo->IsA<UMLDeformerMorphModelInputInfo>())
+	{
+		UMLDeformerMorphModelInputInfo* MorphInputInfo = Cast<UMLDeformerMorphModelInputInfo>(CreateInputInfo());
+		MorphInputInfo->CopyMembersFrom(CurrentInputInfo);
+		CurrentInputInfo->ConditionalBeginDestroy();
+		check(MorphInputInfo); // The input info class should be inherited from the UMLDeformerMorphModelInputInfo class.
+		SetInputInfo(MorphInputInfo);
+	}
+
 	UpdateStatistics();
 
 #if WITH_EDITOR
@@ -173,10 +188,20 @@ void UMLDeformerMorphModel::UpdateMemoryUsage()
 	Super::UpdateMemoryUsage();
 
 	// We strip the deltas when cooking.
-	CookedMemUsageInBytes -= MorphTargetDeltas.Num() * sizeof(FVector3f);
+	CookedMemUsageInBytes -= MorphTargetDeltas.GetAllocatedSize();
+
+	// Strip the input item mask buffer from the cooked size.
+	if (GetInputInfo() && GetInputInfo()->IsA<UMLDeformerMorphModelInputInfo>())
+	{
+		const UMLDeformerMorphModelInputInfo* MorphInputInfo = Cast<UMLDeformerMorphModelInputInfo>(GetInputInfo());
+		check(MorphInputInfo);	// Input Info class is expected to be inherited from the UMLDeformerMorphModelInputInfo class.
+		CookedMemUsageInBytes -= MorphInputInfo->GetInputItemMaskBuffer().GetAllocatedSize();
+	}
 
-	// Add the compressed morph size.
-	GPUMemUsageInBytes += GetCompressedMorphDataSizeInBytes();
+	// Add the compressed morph target data size.
+	const uint64 GPUMorphSize = GetCompressedMorphDataSizeInBytes();
+	GPUMemUsageInBytes += GPUMorphSize;
+	CookedMemUsageInBytes += GPUMorphSize;
 }
 #endif
 
diff --git a/Engine/Plugins/Animation/MLDeformer/MLDeformerFramework/Source/MLDeformerFramework/Public/MLDeformerComponent.h b/Engine/Plugins/Animation/MLDeformer/MLDeformerFramework/Source/MLDeformerFramework/Public/MLDeformerComponent.h
--- a/Engine/Plugins/Animation/MLDeformer/MLDeformerFramework/Source/MLDeformerFramework/Public/MLDeformerComponent.h
+++ b/Engine/Plugins/Animation/MLDeformer/MLDeformerFramework/Source/MLDeformerFramework/Public/MLDeformerComponent.h
@@ -30,6 +30,7 @@ class MLDEFORMERFRAMEWORK_API UMLDeformerComponent
 
 public:
 	// UObject overrides.
+	void BeginDestroy() override;
 #if WITH_EDITOR
 	void PostEditChangeProperty(FPropertyChangedEvent& PropertyChangedEvent) override;
 #endif
@@ -177,13 +178,16 @@ protected:
 	UE_DEPRECATED(5.3, "This method will be removed.")
 	void RemoveReleaseModelInstancesDelegate();
 
+	void BindDelegates();
+	void UnbindDelegates();
+
 	void ReleaseModelInstance();
 
 	/** Set the ML Deformer weight. */
-	virtual void SetWeightInternal(const float NormalizedWeightValue) { Weight = FMath::Clamp<float>(NormalizedWeightValue, 0.0f, 1.0f); }
+	virtual void SetWeightInternal(const float NormalizedWeightValue);
 
 	/** Set the ML Deformer asset. */
-	virtual void SetDeformerAssetInternal(UMLDeformerAsset* const InDeformerAsset) { DeformerAsset = InDeformerAsset; UpdateSkeletalMeshComponent(); }
+	virtual void SetDeformerAssetInternal(UMLDeformerAsset* const InDeformerAsset);
 
 #if WITH_EDITOR
 	/** Reset the tick cycle counters. */
@@ -210,6 +214,9 @@ protected:
 	UE_DEPRECATED(5.3, "This member has been deprecated.")
 	FDelegateHandle NeuralNetworkModifyDelegateHandle_DEPRECATED;
 
+	/** The delegate handle used to bind to the reinit model instance delegate handle. */
+	FDelegateHandle ReinitModelInstanceDelegateHandle;
+
 	/** Suppress mesh deformer logging warnings? This is used by the ML Deformer editor, as we don't want to show some warnings when using that. */
 	bool bSuppressMeshDeformerLogWarnings = false;
 
diff --git a/Engine/Plugins/Animation/MLDeformer/MLDeformerFramework/Source/MLDeformerFramework/Public/MLDeformerInputInfo.h b/Engine/Plugins/Animation/MLDeformer/MLDeformerFramework/Source/MLDeformerFramework/Public/MLDeformerInputInfo.h
--- a/Engine/Plugins/Animation/MLDeformer/MLDeformerFramework/Source/MLDeformerFramework/Public/MLDeformerInputInfo.h
+++ b/Engine/Plugins/Animation/MLDeformer/MLDeformerFramework/Source/MLDeformerFramework/Public/MLDeformerInputInfo.h
@@ -104,12 +104,14 @@ public:
 	 * @return An array of FName objects, one for each bone.
 	 */
 	TArray<FName>& GetBoneNames();
+	const TArray<FName>& GetBoneNames() const;
 
 	/**
 	 * Get the list of curve names that should be included as training inputs.
 	 * @return An array of FName objects, one for each curve.
 	 */
 	TArray<FName>& GetCurveNames();
+	const TArray<FName>& GetCurveNames() const;
 
 	/**
 	 * Get the curve name as an FName, for a given curve we included during training.
diff --git a/Engine/Plugins/Animation/MLDeformer/MLDeformerFramework/Source/MLDeformerFramework/Public/MLDeformerModel.h b/Engine/Plugins/Animation/MLDeformer/MLDeformerFramework/Source/MLDeformerFramework/Public/MLDeformerModel.h
--- a/Engine/Plugins/Animation/MLDeformer/MLDeformerFramework/Source/MLDeformerFramework/Public/MLDeformerModel.h
+++ b/Engine/Plugins/Animation/MLDeformer/MLDeformerFramework/Source/MLDeformerFramework/Public/MLDeformerModel.h
@@ -81,12 +81,9 @@ namespace UE::MLDeformer
 	};
 }	// namespace UE::MLDeformer
 
-// DEPRECATED: Use FMLDeformerReleaseModelInstancesDelegate instead.
+// DEPRECATED: Use FMLDeformerReinitModelInstancesDelegate instead.
 DECLARE_EVENT_OneParam(UMLDeformerModel, FMLDeformerModelOnPostEditProperty, FPropertyChangedEvent&)
 
-/** Delegate used to signal that the UMLDeformerModelInstance's of the specified model should get released. */
-DECLARE_EVENT(UMLDeformerModel, FMLDeformerReleaseModelInstancesDelegate)
-
 /**
  * The ML Deformer runtime model base class.
  * All models should be inherited from this class.
@@ -101,6 +98,9 @@ class MLDEFORMERFRAMEWORK_API UMLDeformerModel
 public:
 	DECLARE_MULTICAST_DELEGATE(FNeuralNetworkModifyDelegate);
 
+	/** Delegate used to signal that the UMLDeformerModelInstance should be reinitialized. The ML Deformer component will connect to this to be informed about needing to reinit the instance. */
+	DECLARE_MULTICAST_DELEGATE(FMLDeformerReinitModelInstancesDelegate)
+
 	virtual ~UMLDeformerModel() = default;
 
 	/**
@@ -325,7 +325,14 @@ public:
 	 * @return A reference to the delegate.
 	 */
 	UE_DEPRECATED(5.2, "This delegate will be removed.")
-	FNeuralNetworkModifyDelegate& GetNeuralNetworkModifyDelegate()		{ return NeuralNetworkModifyDelegate_DEPRECATED; }
+	FNeuralNetworkModifyDelegate& GetNeuralNetworkModifyDelegate()				{ return NeuralNetworkModifyDelegate_DEPRECATED; }
+
+	/**
+	 * Get the delegate which will be called when to inform when the model instance needs to be reinitailized.
+	 * The ML Deformer component will hook into this for example. This is broadcast when something changes in the structure, such as when
+	 * training finished and the neural network changed.
+	 */
+	FMLDeformerReinitModelInstancesDelegate& GetReinitModelInstanceDelegate()	{ return ReinitModelInstanceDelegate; }
 
 
 #if WITH_EDITORONLY_DATA
@@ -345,28 +352,32 @@ public:
 	 * Check whether we should include bone transforms as input to the model during training or not.
 	 * @return Returns true when bone transfomations should be a part of the network inputs, during the training process.
 	 */
-	bool ShouldIncludeBonesInTraining() const					{ return bIncludeBones; }
+	UE_DEPRECATED(5.3, "This method and property has been removed and shouldn't be used anymore.")
+	bool ShouldIncludeBonesInTraining() const					{ return bIncludeBones_DEPRECATED; }
 
 	/**
 	 * Set whether we want to include bones during training or not.
 	 * This will make bone transforms part of the neural network inputs.
 	 * @param bInclude Set to true if you wish bone transforms to be included during training and at inference time.
 	 */
-	void SetShouldIncludeBonesInTraining(bool bInclude)			{ bIncludeBones = bInclude; }
+	UE_DEPRECATED(5.3, "This method and property has been removed and shouldn't be used anymore.")
+	void SetShouldIncludeBonesInTraining(bool bInclude)			{ bIncludeBones_DEPRECATED = bInclude; }
 
 	/**
 	 * Check whether we should include curve values as input to the model during training or not.
 	 * Curve values are single floats.
 	 * @return Returns true when curve values should be a part of the network inputs, during the training process.
 	 */
-	bool ShouldIncludeCurvesInTraining() const					{ return bIncludeCurves; }
+	UE_DEPRECATED(5.3, "This method and property has been removed and shouldn't be used anymore.")
+	bool ShouldIncludeCurvesInTraining() const					{ return bIncludeCurves_DEPRECATED; }
 
 	/**
 	 * Set whether we want to include curves during training.
 	 * This will make curves part of the neural network inputs.
 	 * @param bInclude Set to true to include curves during training and inference time.
 	 */
-	void SetShouldIncludeCurvesInTraining(bool bInclude)		{ bIncludeCurves = bInclude; }
+	UE_DEPRECATED(5.3, "This method and property has been removed and shouldn't be used anymore.")
+	void SetShouldIncludeCurvesInTraining(bool bInclude)		{ bIncludeCurves_DEPRECATED = bInclude; }
 
 	/**
 	 * The delegate that gets fired when a property value changes.
@@ -490,14 +501,19 @@ public:
 
 	// Get property names.
 	static FName GetSkeletalMeshPropertyName()			{ return GET_MEMBER_NAME_CHECKED(UMLDeformerModel, SkeletalMesh); }
-	static FName GetShouldIncludeBonesPropertyName()	{ return GET_MEMBER_NAME_CHECKED(UMLDeformerModel, bIncludeBones); }
-	static FName GetShouldIncludeCurvesPropertyName()	{ return GET_MEMBER_NAME_CHECKED(UMLDeformerModel, bIncludeCurves); }
 	static FName GetAnimSequencePropertyName()			{ return GET_MEMBER_NAME_CHECKED(UMLDeformerModel, AnimSequence); }
 	static FName GetAlignmentTransformPropertyName()	{ return GET_MEMBER_NAME_CHECKED(UMLDeformerModel, AlignmentTransform); }
 	static FName GetBoneIncludeListPropertyName()		{ return GET_MEMBER_NAME_CHECKED(UMLDeformerModel, BoneIncludeList); }
 	static FName GetCurveIncludeListPropertyName()		{ return GET_MEMBER_NAME_CHECKED(UMLDeformerModel, CurveIncludeList); }
 	static FName GetMaxTrainingFramesPropertyName()		{ return GET_MEMBER_NAME_CHECKED(UMLDeformerModel, MaxTrainingFrames); }
 	static FName GetDeltaCutoffLengthPropertyName()		{ return GET_MEMBER_NAME_CHECKED(UMLDeformerModel, DeltaCutoffLength); }
+
+	UE_DEPRECATED(5.3, "This property has been removed and shouldn't be used anymore.")
+	static FName GetShouldIncludeBonesPropertyName()	{ return GET_MEMBER_NAME_CHECKED(UMLDeformerModel, bIncludeBones_DEPRECATED); }
+
+	UE_DEPRECATED(5.3, "This property has been removed and shouldn't be used anymore.")
+	static FName GetShouldIncludeCurvesPropertyName()	{ return GET_MEMBER_NAME_CHECKED(UMLDeformerModel, bIncludeCurves_DEPRECATED); }
+
 #endif	// #if WITH_EDITORONLY_DATA
 
 	/** @return The number of floats per bone in network input. */
@@ -569,6 +585,9 @@ private:
 	/** Delegate that will be called immediately before the NeuralNetwork is changed. */
 	FNeuralNetworkModifyDelegate NeuralNetworkModifyDelegate_DEPRECATED;
 
+	/** Delegate used to trigger reinitialization of the model instance. */
+	FMLDeformerReinitModelInstancesDelegate ReinitModelInstanceDelegate;
+
 	/** Cached number of skeletal mesh vertices. */
 	UPROPERTY()
 	int32 NumBaseMeshVerts = 0;
@@ -602,12 +621,12 @@ private:
 	TObjectPtr<UMLDeformerVizSettings> VizSettings = nullptr;
 
 	/** Specifies whether bone transformations should be included as inputs during the training process. */
-	UPROPERTY(EditAnywhere, Category = "Inputs and Output")
-	bool bIncludeBones = true;
+	UPROPERTY()
+	bool bIncludeBones_DEPRECATED = true;
 
 	/** Specifies whether curve values (a float per curve) should be included as inputs during the training process. */
-	UPROPERTY(EditAnywhere, Category = "Inputs and Output")
-	bool bIncludeCurves = false;
+	UPROPERTY()
+	bool bIncludeCurves_DEPRECATED = false;
 
 	/**
 	 * The animation sequence to apply to the base mesh. This has to match the animation of the target mesh's geometry cache. 
@@ -621,22 +640,22 @@ private:
 	FTransform AlignmentTransform = FTransform::Identity;
 
 	/** The bones to include during training. When none are provided, all bones of the Skeleton will be included. */
-	UPROPERTY(EditAnywhere, Category = "Inputs and Output", meta = (EditCondition = "bIncludeBones"))
+	UPROPERTY(meta = (EditCondition = "bIncludeBones"))
 	TArray<FBoneReference> BoneIncludeList;
 
 	/** The curves to include during training. When none are provided, all curves of the Skeleton will be included. */
-	UPROPERTY(EditAnywhere, Category = "Inputs and Output", meta = (EditCondition = "bIncludeCurves"))
+	UPROPERTY(meta = (EditCondition = "bIncludeCurves"))
 	TArray<FMLDeformerCurveReference> CurveIncludeList;
 
 	/** The maximum numer of training frames (samples) to train on. Use this to train on a sub-section of your full training data. */
-	UPROPERTY(EditAnywhere, Category = "Inputs and Output", meta = (ClampMin = "1"))
+	UPROPERTY(EditAnywhere, Category = "Inputs", meta = (ClampMin = "1"))
 	int32 MaxTrainingFrames = 1000000;
 
 	/**
 	 * Sometimes there can be some vertices that cause some issues that cause deltas to be very long. We can ignore these deltas by setting a cutoff value. 
 	 * Deltas that are longer than the cutoff value (in units), will be ignored and set to zero length. 
 	 */
-	UPROPERTY(EditAnywhere, Category = "Inputs and Output", meta = (ClampMin = "0.01", ForceUnits="cm"))
+	UPROPERTY(EditAnywhere, Category = "Inputs", meta = (ClampMin = "0.01", ForceUnits="cm"))
 	float DeltaCutoffLength = 30.0f;
 #endif
 };
diff --git a/Engine/Plugins/Animation/MLDeformer/MLDeformerFramework/Source/MLDeformerFramework/Public/MLDeformerMorphModelVizSettings.h b/Engine/Plugins/Animation/MLDeformer/MLDeformerFramework/Source/MLDeformerFramework/Public/MLDeformerMorphModelVizSettings.h
--- a/Engine/Plugins/Animation/MLDeformer/MLDeformerFramework/Source/MLDeformerFramework/Public/MLDeformerMorphModelVizSettings.h
+++ b/Engine/Plugins/Animation/MLDeformer/MLDeformerFramework/Source/MLDeformerFramework/Public/MLDeformerMorphModelVizSettings.h
@@ -19,9 +19,6 @@ public:
 	void SetMorphTargetNumber(int32 MorphIndexToPreview)	{ MorphTargetNumber = MorphIndexToPreview; }
 	void SetDrawMorphTargets(bool bDraw)					{ bDrawMorphTargets = bDraw; }
 
-	UE_DEPRECATED(5.2, "This method will be removed, please stop using it.")
-	float GetMorphTargetDeltaThreshold() const				{ return MorphTargetDeltaThreshold_DEPRECATED; }
-
 	int32 GetMorphTargetNumber() const						{ return MorphTargetNumber; }
 	bool GetDrawMorphTargets() const						{ return bDrawMorphTargets; }
 
@@ -29,7 +26,10 @@ public:
 	static FName GetMorphTargetNumberPropertyName()			{ return GET_MEMBER_NAME_CHECKED(UMLDeformerMorphModelVizSettings, MorphTargetNumber); }
 	static FName GetDrawMorphTargetsPropertyName()			{ return GET_MEMBER_NAME_CHECKED(UMLDeformerMorphModelVizSettings, bDrawMorphTargets); }
 
-	UE_DEPRECATED(5.2, "This method will be removed, as well as the MorphTargetDeltaThreshold member, please stop using it.")
+	UE_DEPRECATED(5.2, "This method will be removed, please do not use this method anymore.")
+	float GetMorphTargetDeltaThreshold() const				{ return MorphTargetDeltaThreshold_DEPRECATED; }
+
+	UE_DEPRECATED(5.2, "This method will be removed, as well as the MorphTargetDeltaThreshold member, please do not use this method anymore.")
 	static FName GetMorphTargetDeltaThresholdPropertyName() { return GET_MEMBER_NAME_CHECKED(UMLDeformerMorphModelVizSettings, MorphTargetDeltaThreshold_DEPRECATED); }
 
 protected:
@@ -51,8 +51,7 @@ protected:
 	/**
 	 * Specify whether we want to debug draw the morph targets.
 	 * Enabling this can help you see what the sparsity of the generated morph targets is.
-	 * It also takes the morph target delta threshold into account, so you can preview how that effects the sparsity of the morphs.
-	 * This only can be used after you trained, in the same editor session directly after training.
+	 * If the model supports input masking, it will also show the mask applied to the morph target.
 	 */
 	UPROPERTY(EditAnywhere, Category = "Live Settings")
 	bool bDrawMorphTargets = false;
diff --git a/Engine/Plugins/Animation/MLDeformer/MLDeformerFramework/Source/MLDeformerFrameworkEditor/Private/MLDeformerEditorModel.cpp b/Engine/Plugins/Animation/MLDeformer/MLDeformerFramework/Source/MLDeformerFrameworkEditor/Private/MLDeformerEditorModel.cpp
--- a/Engine/Plugins/Animation/MLDeformer/MLDeformerFramework/Source/MLDeformerFrameworkEditor/Private/MLDeformerEditorModel.cpp
+++ b/Engine/Plugins/Animation/MLDeformer/MLDeformerFramework/Source/MLDeformerFrameworkEditor/Private/MLDeformerEditorModel.cpp
@@ -37,6 +37,7 @@
 #include "BoneContainer.h"
 #include "Misc/ScopedSlowTask.h"
 #include "Animation/AnimData/IAnimationDataModel.h"
+#include "SMLDeformerInputWidget.h"
 
 #define LOCTEXT_NAMESPACE "MLDeformerEditorModel"
 
@@ -579,10 +580,10 @@ namespace UE::MLDeformer
 		RefreshMLDeformerComponents();
 		UpdateIsReadyForTrainingState();
 
-		const int32 TrainingFrame = GetTrainingFrameAtTime(CalcTrainingTimelinePosition());
+		const int32 TrainingFrame = Model->GetVizSettings()->GetTrainingFrameNumber();
 		SetTrainingFrame(TrainingFrame);
 
-		const int32 TestFrame = GetTestFrameAtTime(CalcTestTimelinePosition());
+		const int32 TestFrame = Model->GetVizSettings()->GetTestingFrameNumber();
 		SetTestFrame(TestFrame);
 
 		UpdateEditorInputInfo();
@@ -731,22 +732,12 @@ namespace UE::MLDeformer
 			SetResamplingInputOutputsNeeded(true);
 		}
 		else
-		if (Property->GetFName() == UMLDeformerModel::GetShouldIncludeBonesPropertyName() ||
-			Property->GetFName() == UMLDeformerModel::GetShouldIncludeCurvesPropertyName())
-		{
-			if (PropertyChangedEvent.ChangeType == EPropertyChangeType::ValueSet)
-			{
-				SetResamplingInputOutputsNeeded(true);
-				UpdateIsReadyForTrainingState();
-				GetEditor()->GetModelDetailsView()->ForceRefresh();
-			}
-		}
-		else
 		if (Property->GetFName() == UMLDeformerModel::GetBoneIncludeListPropertyName() ||
 			Property->GetFName() == UMLDeformerModel::GetCurveIncludeListPropertyName())
 		{
 			SetResamplingInputOutputsNeeded(true);
-			UpdateEditorInputInfo();
+			UpdateIsReadyForTrainingState();
+			RefreshInputWidget();
 		}
 		else if (Property->GetFName() == UMLDeformerVizSettings::GetAnimPlaySpeedPropertyName())
 		{
@@ -1013,19 +1004,7 @@ namespace UE::MLDeformer
 		if (Model->GetSkeletalMesh() && GetEditorInputInfo()->IsEmpty())
 		{
 			FString ErrorString;
-			if (Model->DoesSupportBones() && Model->ShouldIncludeBonesInTraining())
-			{
-				ErrorString += FText(LOCTEXT("InputsEmptyBonesErrorText", "Your base mesh has no bones to train on.\n")).ToString();
-			}
-			else
-			if (Model->DoesSupportCurves() && Model->ShouldIncludeCurvesInTraining())
-			{
-				ErrorString += FText(LOCTEXT("InputsNoCurvesToTrainText", "Your base mesh has no curves to train on.\n")).ToString();
-			}
-			else
-			{
-				ErrorString += FText(LOCTEXT("InputsNeededErrorText", "The training process needs inputs.\n")).ToString();
-			}
+			ErrorString += FText(LOCTEXT("InputsNeededErrorText", "The training process needs inputs.\n")).ToString();
 
 			ErrorString.RemoveFromEnd("\n");
 			return FText::FromString(ErrorString);
@@ -1067,95 +1046,12 @@ namespace UE::MLDeformer
 		return Result;
 	}
 
-	void FMLDeformerEditorModel::InitInputInfo(UMLDeformerInputInfo* InputInfo)
-	{
-		InputInfo->Reset();
-
-		TArray<FName>& BoneNames = InputInfo->GetBoneNames();
-		TArray<FName>& CurveNames = InputInfo->GetCurveNames();
-
-		USkeletalMesh* SkeletalMesh = Model->GetSkeletalMesh();
-		InputInfo->SetSkeletalMesh(SkeletalMesh);
-
-		BoneNames.Reset();
-		CurveNames.Reset();
-
-		InputInfo->SetNumBaseVertices(Model->GetNumBaseMeshVerts());
-		InputInfo->SetNumTargetVertices(Model->GetNumTargetMeshVerts());
-
-		const bool bIncludeBones = Model->DoesSupportBones() && Model->ShouldIncludeBonesInTraining();
-		const bool bIncludeCurves = Model->DoesSupportCurves() && Model->ShouldIncludeCurvesInTraining();
-		const USkeleton* Skeleton = Model->GetSkeletalMesh() ? Model->GetSkeletalMesh()->GetSkeleton() : nullptr;
-
-		// Handle bones.
-		if (bIncludeBones && SkeletalMesh)
-		{
-			// Include all the bones when no list was provided.
-			const FReferenceSkeleton& RefSkeleton = SkeletalMesh->GetRefSkeleton();
-			if (Model->GetBoneIncludeList().IsEmpty())
-			{
-				// Grab all bone names.
-				const int32 NumBones = RefSkeleton.GetNum();
-				BoneNames.Reserve(NumBones);
-				for (int32 Index = 0; Index < NumBones; ++Index)
-				{
-					const FName BoneName = RefSkeleton.GetBoneName(Index);
-					BoneNames.Add(BoneName);
-				}
-			}
-			else // A list of bones to include was provided.
-			{
-				for (const FBoneReference& BoneReference : Model->GetBoneIncludeList())
-				{
-					if (BoneReference.BoneName.IsValid())
-					{
-						const FName BoneName = BoneReference.BoneName;
-						if (RefSkeleton.FindBoneIndex(BoneName) == INDEX_NONE)
-						{
-							UE_LOG(LogMLDeformer, Warning, TEXT("Bone '%s' in the bones include list doesn't exist, ignoring it."), *BoneName.ToString());
-							continue;
-						}
-
-						BoneNames.Add(BoneName);
-					}
-				}
-			}
-		}
-
-		// Handle curves.
-		if (bIncludeCurves && SkeletalMesh)
-		{
-			// Anim curves.
-			if(Skeleton)
-			{
-				// Include all curves when no list was provided.
-				if (Model->GetCurveIncludeList().IsEmpty())
-				{
-					Skeleton->GetCurveMetaDataNames(CurveNames);
-				}
-				else // A list of curve names was provided.
-				{
-					for (const FMLDeformerCurveReference& CurveReference : Model->GetCurveIncludeList())
-					{
-						if (CurveReference.CurveName.IsValid())
-						{
-							CurveNames.Add(CurveReference.CurveName);
-						}
-					}
-				}
-			}
-		}
-
-		PRAGMA_DISABLE_DEPRECATION_WARNINGS
-		InputInfo->UpdateNameStrings();
-		PRAGMA_ENABLE_DEPRECATION_WARNINGS
-	}
-
-	void FMLDeformerEditorModel::InitBoneIncludeListToAnimatedBonesOnly()
+	void FMLDeformerEditorModel::AddAnimatedBonesToBonesIncludeList()
 	{
 		if (!Model->GetAnimSequence())
 		{
 			UE_LOG(LogMLDeformer, Warning, TEXT("Cannot initialize bone list as no Anim Sequence has been picked."));
+			UpdateEditorInputInfo();
 			return;
 		}
 
@@ -1163,12 +1059,14 @@ namespace UE::MLDeformer
 		if (!DataModel)
 		{
 			UE_LOG(LogMLDeformer, Warning, TEXT("Anim sequence has no data model."));
+			UpdateEditorInputInfo();
 			return;
 		}
 
 		if (!Model->GetSkeletalMesh())
 		{
 			UE_LOG(LogMLDeformer, Warning, TEXT("Skeletal Mesh has not been set."));
+			UpdateEditorInputInfo();
 			return;
 		}
 
@@ -1176,6 +1074,7 @@ namespace UE::MLDeformer
 		if (!Skeleton)
 		{
 			UE_LOG(LogMLDeformer, Warning, TEXT("Skeletal Mesh has no skeleton."));
+			UpdateEditorInputInfo();
 			return;
 		}
 
@@ -1224,31 +1123,21 @@ namespace UE::MLDeformer
 		}
 
 		// Init the bone include list using the animated bones.
-		if (!AnimatedBoneList.IsEmpty())
+		TArray<FBoneReference>& BoneList = Model->GetBoneIncludeList();
+		for (const FName BoneName : AnimatedBoneList)
 		{
-			TArray<FBoneReference>& BoneList = Model->GetBoneIncludeList();
-			BoneList.Empty();
-			BoneList.Reserve(AnimatedBoneList.Num());
-			for (FName BoneName : AnimatedBoneList)
-			{
-				BoneList.AddDefaulted();
-				FBoneReference& BoneRef = BoneList.Last();
-				BoneRef.BoneName = BoneName;
-			}
-		}
-		else
-		{
-			Model->GetBoneIncludeList().Empty();
-			UE_LOG(LogMLDeformer, Warning, TEXT("There are no animated bone rotations in Anim Sequence '%s'."), *(Model->GetAnimSequence()->GetName()));
+			BoneList.AddUnique(FBoneReference(BoneName));
 		}
+
 		UpdateEditorInputInfo();
 	}
 
-	void FMLDeformerEditorModel::InitCurveIncludeListToAnimatedCurvesOnly()
+	void FMLDeformerEditorModel::AddAnimatedCurvesToCurvesIncludeList()
 	{
 		if (!Model->GetAnimSequence())
 		{
 			UE_LOG(LogMLDeformer, Warning, TEXT("Cannot initialize curve list as no Anim Sequence has been picked."));
+			UpdateEditorInputInfo();
 			return;
 		}
 
@@ -1256,12 +1145,14 @@ namespace UE::MLDeformer
 		if (!DataModel)
 		{
 			UE_LOG(LogMLDeformer, Warning, TEXT("Anim sequence has no data model."));
+			UpdateEditorInputInfo();
 			return;
 		}
 
 		if (!Model->GetSkeletalMesh())
 		{
 			UE_LOG(LogMLDeformer, Warning, TEXT("Skeletal Mesh has not been set."));
+			UpdateEditorInputInfo();
 			return;
 		}
 
@@ -1269,14 +1160,14 @@ namespace UE::MLDeformer
 		if (!Skeleton)
 		{
 			UE_LOG(LogMLDeformer, Warning, TEXT("Skeletal Mesh has no skeleton."));
+			UpdateEditorInputInfo();
 			return;
 		}
 
-		// Iterate over all curves that are both in the skeleton and the animation.
 		TArray<FName> AnimatedCurveList;
 		TArray<FName> SkeletonCurveNames;
 		Skeleton->GetCurveMetaDataNames(SkeletonCurveNames);
-		for (const FName& SkeletonCurveName : SkeletonCurveNames)
+		for (const FName SkeletonCurveName : SkeletonCurveNames)
 		{
 			const TArray<FFloatCurve>& AnimCurves = DataModel->GetFloatCurves();
 			for (const FFloatCurve& AnimCurve : AnimCurves)
@@ -1303,25 +1194,82 @@ namespace UE::MLDeformer
 			}
 		}
 
-		// Init the bone include list using the animated bones.
-		if (!AnimatedCurveList.IsEmpty())
+		TArray<FMLDeformerCurveReference>& CurveList = Model->GetCurveIncludeList();
+		for (const FName CurveName : AnimatedCurveList)
 		{
-			TArray<FMLDeformerCurveReference>& CurveList = Model->GetCurveIncludeList();
-			CurveList.Empty();
-			CurveList.Reserve(AnimatedCurveList.Num());
-			for (FName CurveName : AnimatedCurveList)
-			{
-				CurveList.AddDefaulted();
-				FMLDeformerCurveReference& CurveRef = CurveList.Last();
-				CurveRef.CurveName = CurveName;
-			}
+			CurveList.AddUnique(FMLDeformerCurveReference(CurveName));
 		}
-		else
+		UpdateEditorInputInfo();
+	}
+
+	void FMLDeformerEditorModel::InitInputInfo(UMLDeformerInputInfo* InputInfo)
+	{
+		InputInfo->Reset();
+
+		TArray<FName>& BoneNames = InputInfo->GetBoneNames();
+		TArray<FName>& CurveNames = InputInfo->GetCurveNames();
+
+		USkeletalMesh* SkeletalMesh = Model->GetSkeletalMesh();
+		InputInfo->SetSkeletalMesh(SkeletalMesh);
+
+		BoneNames.Reset();
+		CurveNames.Reset();
+
+		InputInfo->SetNumBaseVertices(Model->GetNumBaseMeshVerts());
+		InputInfo->SetNumTargetVertices(Model->GetNumTargetMeshVerts());
+
+		const USkeleton* Skeleton = Model->GetSkeletalMesh() ? Model->GetSkeletalMesh()->GetSkeleton() : nullptr;
+
+		// Handle bones.
+		if (SkeletalMesh)
 		{
-			Model->GetCurveIncludeList().Empty();
-			UE_LOG(LogMLDeformer, Warning, TEXT("There are no animated curves in Anim Sequence '%s'."), *(Model->GetAnimSequence()->GetName()));
+			// Include all the bones when no list was provided.
+			const FReferenceSkeleton& RefSkeleton = SkeletalMesh->GetRefSkeleton();
+			for (const FBoneReference& BoneReference : Model->GetBoneIncludeList())
+			{
+				if (BoneReference.BoneName.IsValid())
+				{
+					const FName BoneName = BoneReference.BoneName;
+					if (RefSkeleton.FindBoneIndex(BoneName) == INDEX_NONE)
+					{
+						UE_LOG(LogMLDeformer, Warning, TEXT("Bone '%s' in the bones include list doesn't exist, ignoring it."), *BoneName.ToString());
+						continue;
+					}
+
+					BoneNames.Add(BoneName);
+				}
+			}
+
+			// Anim curves.
+			if (Skeleton)
+			{
+				for (const FMLDeformerCurveReference& CurveReference : Model->GetCurveIncludeList())
+				{
+					if (CurveReference.CurveName.IsValid())
+					{
+						CurveNames.Add(CurveReference.CurveName);
+					}
+				}
+			}
 		}
-		UpdateEditorInputInfo();
+
+		PRAGMA_DISABLE_DEPRECATION_WARNINGS
+		InputInfo->UpdateNameStrings();
+		PRAGMA_ENABLE_DEPRECATION_WARNINGS
+	}
+
+	void FMLDeformerEditorModel::InitBoneIncludeListToAnimatedBonesOnly()
+	{
+		TArray<FBoneReference>& BoneList = Model->GetBoneIncludeList();
+		BoneList.Empty();
+		AddAnimatedBonesToBonesIncludeList();
+	}
+
+	void FMLDeformerEditorModel::InitCurveIncludeListToAnimatedCurvesOnly()
+	{
+		TArray<FMLDeformerCurveReference>& CurveList = Model->GetCurveIncludeList();
+		CurveList.Empty();
+		AddAnimatedCurvesToCurvesIncludeList();
 	}
 
 	void FMLDeformerEditorModel::Render(const FSceneView* View, FViewport* Viewport, FPrimitiveDrawInterface* PDI)
@@ -1690,7 +1638,9 @@ namespace UE::MLDeformer
 
 		// Calculate the normals of that displaced mesh.
 		TArray<FVector3f> MorphedNormals;
+		PRAGMA_DISABLE_DEPRECATION_WARNINGS
 		CalcMeshNormals(MorphedVertexPositions, IndexArray, VertexMap, MorphedNormals);
+		PRAGMA_ENABLE_DEPRECATION_WARNINGS
 
 		// Calculate and output the difference between the morphed normal and base normal.
 		OutDeltaNormals.Reset();
@@ -1705,7 +1655,15 @@ namespace UE::MLDeformer
 		}
 	}
 
-	void FMLDeformerEditorModel::CreateEngineMorphTargets(TArray<UMorphTarget*>& OutMorphTargets, const TArray<FVector3f>& Deltas, const FString& NamePrefix, int32 LOD, float DeltaThreshold, bool bIncludeNormals, EMLDeformerMaskChannel MaskChannel, bool bInvertMaskChannel)
+	void FMLDeformerEditorModel::CreateEngineMorphTargets(
+		TArray<UMorphTarget*>& OutMorphTargets, 
+		const TArray<FVector3f>& Deltas, 
+		const FString& NamePrefix, 
+		int32 LOD, 
+		float DeltaThreshold, 
+		bool bIncludeNormals, 
+		EMLDeformerMaskChannel MaskChannel,
+		bool bInvertMaskChannel)
 	{
 		OutMorphTargets.Reset();
 		if (Deltas.IsEmpty())
@@ -1748,7 +1706,9 @@ namespace UE::MLDeformer
 		TArray<FVector3f> BaseNormals;
 		if (bIncludeNormals)
 		{
+			PRAGMA_DISABLE_DEPRECATION_WARNINGS
 			CalcMeshNormals(BaseVertexPositions, IndexArray, VertexMap, BaseNormals);
+			PRAGMA_ENABLE_DEPRECATION_WARNINGS
 		}
 
 		const FColorVertexBuffer& ColorBuffer = RenderData->LODRenderData[LOD].StaticVertexBuffers.ColorVertexBuffer;
@@ -1760,7 +1720,9 @@ namespace UE::MLDeformer
 		{
 			if (bIncludeNormals)
 			{
+				PRAGMA_DISABLE_DEPRECATION_WARNINGS
 				GenerateNormalsForMorphTarget(LOD, SkelMesh, MorphTargetIndex, Deltas, BaseVertexPositions, BaseNormals, DeltaNormals);
+				PRAGMA_ENABLE_DEPRECATION_WARNINGS
 			}
 
 			const FName MorphName = *FString::Printf(TEXT("%s%.3d"), *NamePrefix, MorphTargetIndex);
@@ -1947,6 +1909,20 @@ namespace UE::MLDeformer
 		return FString(TEXT("/MLDeformerFramework/Deformers/DG_MLDeformerModel_HeatMap.DG_MLDeformerModel_HeatMap"));
 	}
 
+	void FMLDeformerEditorModel::RefreshInputWidget()
+	{
+		if (InputWidget.IsValid())
+		{
+			InputWidget->Refresh();
+		}
+	}
+
+	TSharedPtr<SMLDeformerInputWidget> FMLDeformerEditorModel::CreateInputWidget()
+	{
+		return SNew(SMLDeformerInputWidget)
+			.EditorModel(this);
+	}
+
 	const UAnimSequence* FMLDeformerEditorModel::GetAnimSequence() const
 	{
 		if (GetModel())
diff --git a/Engine/Plugins/Animation/MLDeformer/MLDeformerFramework/Source/MLDeformerFrameworkEditor/Private/MLDeformerEditorModule.cpp b/Engine/Plugins/Animation/MLDeformer/MLDeformerFramework/Source/MLDeformerFrameworkEditor/Private/MLDeformerEditorModule.cpp
--- a/Engine/Plugins/Animation/MLDeformer/MLDeformerFramework/Source/MLDeformerFrameworkEditor/Private/MLDeformerEditorModule.cpp
+++ b/Engine/Plugins/Animation/MLDeformer/MLDeformerFramework/Source/MLDeformerFrameworkEditor/Private/MLDeformerEditorModule.cpp
@@ -7,6 +7,7 @@
 #include "MLDeformerModel.h"
 #include "MLDeformerCurveReferenceCustomization.h"
 #include "MLDeformerMorphModelQualityLevelCustomization.h"
+#include "SMLDeformerInputWidget.h"
 #include "Modules/ModuleManager.h"
 #include "EditorModeRegistry.h"
 #include "PropertyEditorDelegates.h"
@@ -26,6 +27,8 @@ namespace UE::MLDeformer
 		PropertyModule.RegisterCustomPropertyTypeLayout("MLDeformerCurveReference", FOnGetPropertyTypeCustomizationInstance::CreateStatic(&FMLDeformerCurveReferenceCustomization::MakeInstance) );
 		PropertyModule.RegisterCustomPropertyTypeLayout("MLDeformerMorphModelQualityLevel", FOnGetPropertyTypeCustomizationInstance::CreateStatic(&FMLDeformerMorphModelQualityLevelCustomization::MakeInstance) );
 		PropertyModule.NotifyCustomizationModuleChanged();
+
+		SMLDeformerInputWidget::RegisterCommands();
 	}
 
 	void FMLDeformerEditorModule::ShutdownModule()
diff --git a/Engine/Plugins/Animation/MLDeformer/MLDeformerFramework/Source/MLDeformerFrameworkEditor/Private/MLDeformerEditorStyle.cpp b/Engine/Plugins/Animation/MLDeformer/MLDeformerFramework/Source/MLDeformerFrameworkEditor/Private/MLDeformerEditorStyle.cpp
--- a/Engine/Plugins/Animation/MLDeformer/MLDeformerFramework/Source/MLDeformerFrameworkEditor/Private/MLDeformerEditorStyle.cpp
+++ b/Engine/Plugins/Animation/MLDeformer/MLDeformerFramework/Source/MLDeformerFrameworkEditor/Private/MLDeformerEditorStyle.cpp
@@ -51,9 +51,15 @@ namespace UE::MLDeformer
 		Set("MLDeformer.Morphs.IncludedVertexColor", FLinearColor(1.0f, 0.0f, 1.0f));
 		Set("MLDeformer.Morphs.ExcludedVertexColor", FLinearColor(0.1f, 0.1f, 0.1f));
 
+		Set("MLDeformer.Morphs.MaskIncludedVertexColor", FLinearColor(0.0f, 1.0f, 1.0f));
+		Set("MLDeformer.Morphs.MaskExcludedVertexColor", FLinearColor(0.1f, 0.1f, 0.1f));
+
 		Set("MLDeformer.Statistics.Performance", FLinearColor(0.0f, 1.0f, 0.0f));
 		Set("MLDeformer.DefaultLabelScale", 0.65f);
 
+		Set("MLDeformer.InputsWidget.HighlightColor", FLinearColor(0.0f, 1.0f, 0.0f));
+		Set("MLDeformer.InputsWidget.ErrorColor", FLinearColor(1.0f, 0.5f, 0.0f));
+
 		FSlateStyleRegistry::RegisterSlateStyle(*this);
 	}
 
diff --git a/Engine/Plugins/Animation/MLDeformer/MLDeformerFramework/Source/MLDeformerFrameworkEditor/Private/MLDeformerModelDetails.cpp b/Engine/Plugins/Animation/MLDeformer/MLDeformerFramework/Source/MLDeformerFrameworkEditor/Private/MLDeformerModelDetails.cpp
--- a/Engine/Plugins/Animation/MLDeformer/MLDeformerFramework/Source/MLDeformerFrameworkEditor/Private/MLDeformerModelDetails.cpp
+++ b/Engine/Plugins/Animation/MLDeformer/MLDeformerFramework/Source/MLDeformerFrameworkEditor/Private/MLDeformerModelDetails.cpp
@@ -27,6 +27,7 @@
 #include "DetailLayoutBuilder.h"
 #include "SWarningOrErrorBox.h"
 #include "Widgets/Text/STextBlock.h"
+#include "SMLDeformerInputWidget.h"
 
 #define LOCTEXT_NAMESPACE "MLDeformerModelDetails"
 
@@ -55,7 +56,7 @@ namespace UE::MLDeformer
 	{
 		BaseMeshCategoryBuilder = &DetailLayoutBuilder->EditCategory("Base Mesh", FText::GetEmpty(), ECategoryPriority::Important);
 		TargetMeshCategoryBuilder = &DetailLayoutBuilder->EditCategory("Target Mesh", FText::GetEmpty(), ECategoryPriority::Important);
-		InputOutputCategoryBuilder = &DetailLayoutBuilder->EditCategory("Inputs and Output", FText::GetEmpty(), ECategoryPriority::Important);
+		InputOutputCategoryBuilder = &DetailLayoutBuilder->EditCategory("Inputs", FText::GetEmpty(), ECategoryPriority::Important);
 		TrainingSettingsCategoryBuilder = &DetailLayoutBuilder->EditCategory("Training Settings", FText::GetEmpty(), ECategoryPriority::Important);
 	}
 
@@ -93,6 +94,7 @@ namespace UE::MLDeformer
 				]
 			];
 
+
 		// Check if our skeletal mesh's imported model contains a list of mesh infos. If not, we need to reimport it as it is an older asset.
 		const FText NeedsReimportErrorText = EditorModel->GetSkeletalMeshNeedsReimportErrorText();
 		FDetailWidgetRow& NeedsReimportErrorRow = BaseMeshCategoryBuilder->AddCustomRow(FText::FromString("BaseMeshNeedsReimportError"))
@@ -164,18 +166,15 @@ namespace UE::MLDeformer
 
 		TargetMeshCategoryBuilder->AddProperty(UMLDeformerModel::GetAlignmentTransformPropertyName(), UMLDeformerModel::StaticClass());
 
-		InputOutputCategoryBuilder->AddProperty(UMLDeformerModel::GetShouldIncludeBonesPropertyName(), UMLDeformerModel::StaticClass())
-			.Visibility(IsBonesFlagVisible() ? EVisibility::Visible : EVisibility::Collapsed);
-
-		InputOutputCategoryBuilder->AddProperty(UMLDeformerModel::GetShouldIncludeCurvesPropertyName(), UMLDeformerModel::StaticClass())
-			.Visibility(IsCurvesFlagVisible() ? EVisibility::Visible : EVisibility::Collapsed);
-
 		AddTrainingInputFlags();
 		AddTrainingInputErrors();
 
-		const FText ErrorText = EditorModel->GetInputsErrorText();
 		FDetailWidgetRow& ErrorRow = InputOutputCategoryBuilder->AddCustomRow(FText::FromString("InputsError"))
-			.Visibility(!ErrorText.IsEmpty() ? EVisibility::Visible : EVisibility::Collapsed)
+			.Visibility(TAttribute<EVisibility>::CreateLambda(
+				[this]()
+				{
+					return !EditorModel->GetInputsErrorText().IsEmpty() ? EVisibility::Visible : EVisibility::Collapsed;
+				}))
 			.WholeRowContent()
 			[
 				SNew(SBox)
@@ -183,153 +182,20 @@ namespace UE::MLDeformer
 				[
 					SNew(SWarningOrErrorBox)
 					.MessageStyle(EMessageStyle::Error)
-					.Message(ErrorText)
+					.Message_Lambda([this] { return EditorModel->GetInputsErrorText(); })
 				]
 			];
 
-		InputOutputCategoryBuilder->AddProperty(UMLDeformerModel::GetMaxTrainingFramesPropertyName(), UMLDeformerModel::StaticClass());
-		InputOutputCategoryBuilder->AddProperty(UMLDeformerModel::GetDeltaCutoffLengthPropertyName(), UMLDeformerModel::StaticClass());
+		// Create the inputs widget.
+		TSharedPtr<SMLDeformerInputWidget> InputWidget = EditorModel->CreateInputWidget();
+		EditorModel->SetInputWidget(InputWidget);
+		InputOutputCategoryBuilder->AddCustomRow(FText::FromString("Inputs")).WholeRowContent().Widget = InputWidget.ToSharedRef();
 
-		// Bone include list group.
-		if (Model->DoesSupportBones())
-		{
-			IDetailGroup& BoneIncludeGroup = InputOutputCategoryBuilder->AddGroup("BoneIncludeGroup", FText(), false, false);
-			BoneIncludeGroup.HeaderRow()
-			.NameContent()
-			[
-				SNew(SHorizontalBox)
-				+ SHorizontalBox::Slot()
-				.HAlign(HAlign_Right)
-				[
-					SNew(STextBlock)
-					.Text_Lambda
-					(
-						[this]()
-						{
-							const int NumBonesIncluded = EditorModel->GetEditorInputInfo() ? EditorModel->GetEditorInputInfo()->GetNumBones() : 0;
-							return FText::Format(FTextFormat(LOCTEXT("BonesGroupName", "Bones ({0})")), NumBonesIncluded);
-						}
-					)
-					.Font(IDetailLayoutBuilder::GetDetailFont())
-				]
-			]
-			.ValueContent()
-			[
-				SNew(SHorizontalBox)
-				+ SHorizontalBox::Slot()
-				.HAlign(HAlign_Left)
-				[
-					SNew(STextBlock)
-					.Text_Lambda
-					(
-						[this]()
-						{
-							const int NumBonesIncluded = EditorModel->GetEditorInputInfo() ? EditorModel->GetEditorInputInfo()->GetNumBones() : 0;
-							FText AllBonesText;
-							if (Model->GetSkeletalMesh() && (Model->GetSkeletalMesh()->GetRefSkeleton().GetNum() == NumBonesIncluded || NumBonesIncluded == 0))
-							{
-								AllBonesText = LOCTEXT("BoneGroupValue", "All Bones Included");
-							}
-							return AllBonesText;
-						}
-					)
-					.Font(IDetailLayoutBuilder::GetDetailFontItalic())
-				]
-			];
-			BoneIncludeGroup.AddWidgetRow()
-				.ValueContent()
-				[
-					SNew(SButton)
-					.HAlign(HAlign_Center)
-					.VAlign(VAlign_Center)
-					.Text(LOCTEXT("AnimatedBonesButton", "Animated Bones Only"))
-					.ToolTipText(LOCTEXT("AnimatedBonesButtonTooltip", "Initialize the bone include list to bones that are animated."))
-					.OnClicked(FOnClicked::CreateSP(this, &FMLDeformerModelDetails::OnFilterAnimatedBonesOnly))
-					.IsEnabled_Lambda([this](){ return Model->ShouldIncludeBonesInTraining(); })
-				];
-			BoneIncludeGroup.AddPropertyRow(DetailBuilder.GetProperty(UMLDeformerModel::GetBoneIncludeListPropertyName(), UMLDeformerModel::StaticClass()));
-
-			AddBoneInputErrors();
-		}
-		else
-		{
-			InputOutputCategoryBuilder->AddProperty(UMLDeformerModel::GetBoneIncludeListPropertyName(), UMLDeformerModel::StaticClass())
-				.Visibility(EVisibility::Collapsed);
-		}
+		AddBoneInputErrors();
+		AddCurveInputErrors();
 
-		// Curve include list group.
-		if (Model->DoesSupportCurves())
-		{
-			IDetailGroup& CurvesIncludeGroup = InputOutputCategoryBuilder->AddGroup("CurveIncludeGroup", LOCTEXT("CurveIncludeGroup", "Curves"), false, false);
-			CurvesIncludeGroup.HeaderRow()
-			.NameContent()
-			[
-				SNew(SHorizontalBox)
-				+ SHorizontalBox::Slot()
-				.HAlign(HAlign_Right)
-				[
-					SNew(STextBlock)
-					.Text_Lambda
-					(
-						[this]()
-						{
-							const int NumCurves = EditorModel->GetEditorInputInfo() ? EditorModel->GetEditorInputInfo()->GetNumCurves() : 0;
-							return FText::Format(FTextFormat(LOCTEXT("CurvesGroupName", "Curves ({0})")), NumCurves);
-						}
-					)
-					.Font(IDetailLayoutBuilder::GetDetailFont())
-				]
-			]
-			.ValueContent()
-			[
-				SNew(SHorizontalBox)
-				+ SHorizontalBox::Slot()
-				.HAlign(HAlign_Left)
-				[
-					SNew(STextBlock)
-					.Text_Lambda
-					(
-						[this]()
-						{
-							const int NumCurves = EditorModel->GetEditorInputInfo() ? EditorModel->GetEditorInputInfo()->GetNumCurves() : 0;
-							const int32 NumCurvesOnSkelMesh = EditorModel->GetNumCurvesOnSkeletalMesh(Model->GetSkeletalMesh());
-							if (NumCurvesOnSkelMesh == 0)
-							{
-								return LOCTEXT("CurvesGroupValueNoCurves", "No Curves Found");
-							}
-							else if (NumCurves == NumCurvesOnSkelMesh)
-							{
-								return LOCTEXT("CurvesGroupValue", "All Curves Included");
-							}
-							return FText();
-						}
-					)
-					.Font(IDetailLayoutBuilder::GetDetailFontItalic())
-				]
-			];
-
-			CurvesIncludeGroup.AddWidgetRow()
-				.ValueContent()
-				[
-					SNew(SButton)
-					.HAlign(HAlign_Center)
-					.VAlign(VAlign_Center)
-					.Text(LOCTEXT("AnimatedCurvesButton", "Animated Curves Only"))
-					.ToolTipText(LOCTEXT("AnimatedCurvesButtonTooltip", "Initialize the curve include list to curves that are animated."))
-					.OnClicked(FOnClicked::CreateSP(this, &FMLDeformerModelDetails::OnFilterAnimatedCurvesOnly))
-					.IsEnabled_Lambda([this](){ return Model->ShouldIncludeCurvesInTraining() && EditorModel->GetNumCurvesOnSkeletalMesh(Model->GetSkeletalMesh()) > 0; })
-				];
-			CurvesIncludeGroup.AddPropertyRow(DetailBuilder.GetProperty(UMLDeformerModel::GetCurveIncludeListPropertyName(), UMLDeformerModel::StaticClass()));
-
-			AddCurveInputErrors();
-		}
-		else
-		{
-			InputOutputCategoryBuilder->AddProperty(UMLDeformerModel::GetCurveIncludeListPropertyName(), UMLDeformerModel::StaticClass())
-				.Visibility(EVisibility::Collapsed);
-		}
-
-		AddTrainingInputFilters();
+		InputOutputCategoryBuilder->AddProperty(UMLDeformerModel::GetMaxTrainingFramesPropertyName(), UMLDeformerModel::StaticClass());
+		InputOutputCategoryBuilder->AddProperty(UMLDeformerModel::GetDeltaCutoffLengthPropertyName(), UMLDeformerModel::StaticClass());
 
 		// Show a warning when no neural network has been set.
 		{		
diff --git a/Engine/Plugins/Animation/MLDeformer/MLDeformerFramework/Source/MLDeformerFrameworkEditor/Private/MLDeformerMorphModelEditorModel.cpp b/Engine/Plugins/Animation/MLDeformer/MLDeformerFramework/Source/MLDeformerFrameworkEditor/Private/MLDeformerMorphModelEditorModel.cpp
--- a/Engine/Plugins/Animation/MLDeformer/MLDeformerFramework/Source/MLDeformerFrameworkEditor/Private/MLDeformerMorphModelEditorModel.cpp
+++ b/Engine/Plugins/Animation/MLDeformer/MLDeformerFramework/Source/MLDeformerFrameworkEditor/Private/MLDeformerMorphModelEditorModel.cpp
@@ -2,11 +2,12 @@
 
 #include "MLDeformerMorphModelEditorModel.h"
 #include "MLDeformerMorphModel.h"
+#include "MLDeformerMorphModelInputInfo.h"
 #include "MLDeformerMorphModelVizSettings.h"
 #include "MLDeformerMorphModelInstance.h"
-#include "MLDeformerMorphModelInstance.h"
 #include "MLDeformerGeomCacheActor.h"
 #include "MLDeformerEditorToolkit.h"
+#include "MLDeformerEditorStyle.h"
 #include "MLDeformerInputInfo.h"
 #include "MLDeformerComponent.h"
 #include "MLDeformerGeomCacheHelpers.h"
@@ -19,6 +20,8 @@
 #include "Rendering/SkeletalMeshLODModel.h"
 #include "Rendering/SkeletalMeshRenderData.h"
 #include "Engine/SkeletalMesh.h"
+#include "Misc/ScopedSlowTask.h"
+#include "BoneWeights.h"
 
 #define LOCTEXT_NAMESPACE "MLDeformerMorphModelEditorModel"
 
@@ -29,6 +32,11 @@ namespace UE::MLDeformer
 		return new FMLDeformerMorphModelEditorModel();
 	}
 
+	bool FMLDeformerMorphModelEditorModel::IsInputMaskingSupported() const
+	{
+		return false;
+	}
+
 	void FMLDeformerMorphModelEditorModel::OnPropertyChanged(FPropertyChangedEvent& PropertyChangedEvent)
 	{
 		const FProperty* Property = PropertyChangedEvent.Property;
@@ -154,15 +162,15 @@ namespace UE::MLDeformer
 
 	void FMLDeformerMorphModelEditorModel::UpdateMorphErrorValues(TArrayView<UMorphTarget*> MorphTargets)
 	{
-		// If we don't have morph targets yet, there is nothing to do.
-		UMLDeformerMorphModel* MorphModel = GetMorphModel();
-		if (MorphModel->GetNumMorphTargets() == 0)
+		if (MorphTargets.IsEmpty())
 		{
+			
 			return;
 		}
 
 		// Check if we have max morph weight information.
 		// If we do not have this yet, we have to initialize the weights to 1.
+		UMLDeformerMorphModel* MorphModel = GetMorphModel();
 		TArrayView<const float> MaxMorphWeights = MorphModel->GetMorphTargetMaxWeights();
 
 		// Preallocate space for the standard deviation of each morph target.
@@ -214,6 +222,387 @@ namespace UE::MLDeformer
 		MorphModel->SetMorphTargetsErrorOrder(SortedIndices, ErrorValues);
 	}
 
+	const TArrayView<const float> FMLDeformerMorphModelEditorModel::GetMaskForMorphTarget(int32 MorphTargetIndex) const
+	{
+		// Return an empty array view object on default, which essentially disables the masking.
+		// EditorModels can override this.
+		return TArrayView<const float>();
+	}
+
+	void FMLDeformerMorphModelEditorModel::ZeroDeltasByLengthThreshold(TArray<FVector3f>& Deltas, float Threshold)
+	{
+		for (int32 VertexIndex = 0; VertexIndex < Deltas.Num(); ++VertexIndex)
+		{
+			if (Deltas[VertexIndex].Length() <= Threshold)
+			{
+				Deltas[VertexIndex] = FVector3f::ZeroVector;
+			}
+		}
+	}
+
+	void FMLDeformerMorphModelEditorModel::CalcVertexNormals(TArrayView<const FVector3f> VertexPositions, TArrayView<const uint32> IndexArray, TArrayView<const int32> VertexMap, TArray<FVector3f>& OutNormals) const
+	{
+		const int32 NumVertices = VertexPositions.Num();
+		OutNormals.Reset();
+		OutNormals.SetNumZeroed(NumVertices);
+
+		checkf(IndexArray.Num() % 3 == 0, TEXT("Expecting a triangle mesh!"));
+		const int32 NumTriangles = IndexArray.Num() / 3;
+		for (int32 TriangleIndex = 0; TriangleIndex < NumTriangles; ++TriangleIndex)
+		{
+			const int32 ImportedIndices[3]	{ VertexMap[IndexArray[TriangleIndex * 3]], VertexMap[IndexArray[TriangleIndex * 3 + 1]], VertexMap[IndexArray[TriangleIndex * 3 + 2]] };
+			const FVector3f Positions[3]	{ VertexPositions[ImportedIndices[0]], VertexPositions[ImportedIndices[1]], VertexPositions[ImportedIndices[2]] };
+
+			const FVector3f EdgeA = (Positions[1] - Positions[0]).GetSafeNormal();
+			const FVector3f EdgeB = (Positions[2] - Positions[0]).GetSafeNormal();
+			if (EdgeA.SquaredLength() > 0.00001f && EdgeB.SquaredLength() > 0.00001f)
+			{	
+				const FVector3f FaceNormal = EdgeB.Cross(EdgeA);
+				OutNormals[ImportedIndices[0]] += FaceNormal;
+				OutNormals[ImportedIndices[1]] += FaceNormal;
+				OutNormals[ImportedIndices[2]] += FaceNormal;
+			}
+		}
+
+		// Renormalize.
+		for (int32 Index = 0; Index < NumVertices; ++Index)
+		{
+			OutNormals[Index] = OutNormals[Index].GetSafeNormal();
+		}
+	}
+
+	void FMLDeformerMorphModelEditorModel::CalcMorphTargetNormals(
+		int32 LOD,
+		USkeletalMesh* SkelMesh,
+		int32 MorphTargetIndex,
+		TArrayView<const FVector3f> Deltas,
+		TArrayView<const FVector3f> BaseVertexPositions,
+		TArrayView<FVector3f> BaseNormals,
+		TArray<FVector3f>& OutDeltaNormals)
+	{
+		FSkeletalMeshRenderData* RenderData = SkelMesh->GetResourceForRendering();
+		check(RenderData);
+		check(!RenderData->LODRenderData.IsEmpty());
+		const int32 NumRenderVertices = RenderData->LODRenderData[LOD].GetNumVertices();
+		const FColorVertexBuffer& ColorBuffer = RenderData->LODRenderData[LOD].StaticVertexBuffers.ColorVertexBuffer;
+
+		CalcMorphTargetNormals(
+			LOD,
+			SkelMesh,
+			MorphTargetIndex,
+			Deltas,
+			BaseVertexPositions,
+			BaseNormals,
+			TArray<int32>(),
+			ColorBuffer,
+			EMLDeformerMaskChannel::Disabled,
+			false,
+			OutDeltaNormals);
+	}
+
+	void FMLDeformerMorphModelEditorModel::CalcMorphTargetNormals(
+		int32 LOD,
+		const USkeletalMesh* SkelMesh,
+		int32 MorphTargetIndex,
+		const TArrayView<const FVector3f> Deltas,
+		const TArrayView<const FVector3f> BaseVertexPositions,
+		const TArrayView<const FVector3f> BaseNormals,
+		const TArrayView<const int32> ImportedVertexToRenderVertexMapping,
+		const FColorVertexBuffer& ColorBuffer,
+		EMLDeformerMaskChannel MaskChannel,
+		bool bInvertGlobalMaskChannel,
+		TArray<FVector3f>& OutDeltaNormals)
+	{
+		const FSkeletalMeshModel* ImportedModel = SkelMesh->GetImportedModel();
+		if (ImportedModel == nullptr || !ImportedModel->LODModels.IsValidIndex(LOD))
+		{
+			OutDeltaNormals.Reset();
+			OutDeltaNormals.SetNumZeroed(Model->GetNumBaseMeshVerts());
+			return;
+		}
+		const TArray<uint32>& IndexArray = ImportedModel->LODModels[LOD].IndexBuffer;
+		const TArray<int32>& VertexMap = ImportedModel->LODModels[LOD].MeshToImportVertexMap;
+
+		// Get the optional input mask for this morph target.
+		const TArrayView<const float> MorphMask = (MorphTargetIndex > 0) ? GetMaskForMorphTarget(MorphTargetIndex - 1) : TArrayView<const float>();
+
+		// Build the array of displaced vertex positions.
+		TArray<FVector3f> MorphedVertexPositions;
+		const int32 NumBaseMeshVerts = Model->GetNumBaseMeshVerts();
+		MorphedVertexPositions.SetNumUninitialized(NumBaseMeshVerts);
+		for (int32 VertexIndex = 0; VertexIndex < NumBaseMeshVerts; ++VertexIndex)
+		{
+			const int32 RenderVertexIndex = !ImportedVertexToRenderVertexMapping.IsEmpty() ? ImportedVertexToRenderVertexMapping[VertexIndex] : INDEX_NONE;
+			const int32 DeltaIndex = (MorphTargetIndex * NumBaseMeshVerts) + VertexIndex;
+			const FVector3f RawDelta = Deltas[DeltaIndex];
+			const float MorphMaskWeight = MorphMask.IsEmpty() ? 1.0f : MorphMask[VertexIndex];
+			const float GlobalMaskWeight = (RenderVertexIndex != INDEX_NONE) ? CalcGlobalMaskWeight(RenderVertexIndex, ColorBuffer, MaskChannel, bInvertGlobalMaskChannel) : 1.0f;
+			FVector3f ScaledDelta;
+			FVector3f DummyNormal;
+			ProcessVertexDelta(ScaledDelta, DummyNormal, RawDelta, FVector3f::ZeroVector, 0.0f, MorphMaskWeight, GlobalMaskWeight);
+			MorphedVertexPositions[VertexIndex] = BaseVertexPositions[VertexIndex] + ScaledDelta;
+		}
+
+		// Calculate the normals of that displaced mesh.
+		TArray<FVector3f> MorphedNormals;
+		CalcVertexNormals(MorphedVertexPositions, IndexArray, VertexMap, MorphedNormals);
+
+		// Calculate and output the difference between the morphed normal and base normal.
+		OutDeltaNormals.Reset();
+		OutDeltaNormals.SetNumUninitialized(NumBaseMeshVerts);
+		for (int32 VertexIndex = 0; VertexIndex < NumBaseMeshVerts; ++VertexIndex)
+		{
+			OutDeltaNormals[VertexIndex] = MorphedNormals[VertexIndex] - BaseNormals[VertexIndex];
+			if (OutDeltaNormals[VertexIndex].SquaredLength() <= 0.00001f)
+			{
+				OutDeltaNormals[VertexIndex] = FVector3f::ZeroVector;
+			}
+		}
+	}
+
+	bool FMLDeformerMorphModelEditorModel::ProcessVertexDelta(FVector3f& OutScaledDelta, FVector3f& OutScaledDeltaNormal, const FVector3f RawDelta, const FVector3f RawDeltaNormal, float DeltaThreshold, float MorphMaskWeight, float GlobalMaskWeight) const
+	{
+		OutScaledDelta = RawDelta * MorphMaskWeight * GlobalMaskWeight;
+		OutScaledDeltaNormal = RawDeltaNormal * GlobalMaskWeight;
+		return (OutScaledDelta.Length() >= DeltaThreshold);
+	}
+
+	float FMLDeformerMorphModelEditorModel::CalcGlobalMaskWeight(int32 RenderVertexIndex, const FColorVertexBuffer& ColorBuffer, EMLDeformerMaskChannel MaskChannel, bool bInvertMaskChannel) const
+	{
+		float VertexWeight = 1.0f;
+		if (ColorBuffer.GetNumVertices() != 0 && MaskChannel != EMLDeformerMaskChannel::Disabled && RenderVertexIndex != INDEX_NONE)
+		{
+			const FLinearColor& VertexColor = ColorBuffer.VertexColor(RenderVertexIndex);
+			switch (MaskChannel)
+			{
+				case EMLDeformerMaskChannel::VertexColorRed:	{ VertexWeight = VertexColor.R; break; }
+				case EMLDeformerMaskChannel::VertexColorGreen:	{ VertexWeight = VertexColor.G; break; }
+				case EMLDeformerMaskChannel::VertexColorBlue:	{ VertexWeight = VertexColor.B; break; }
+				case EMLDeformerMaskChannel::VertexColorAlpha:	{ VertexWeight = VertexColor.A; break; }
+				default: 
+					checkf(false, TEXT("Unexpected mask channel value."));
+					break;
+			};
+
+			if (bInvertMaskChannel)
+			{
+				VertexWeight = FMath::Clamp<float>(1.0f - VertexWeight, 0.0f, 1.0f);
+			}
+		}
+		return VertexWeight;
+	}
+
+	void FMLDeformerMorphModelEditorModel::CreateMorphTargets(
+		TArray<UMorphTarget*>& OutMorphTargets, 
+		const TArray<FVector3f>& Deltas, 
+		const FString& NamePrefix, 
+		int32 LOD, 
+		float DeltaThreshold, 
+		bool bIncludeNormals, 
+		EMLDeformerMaskChannel MaskChannel,
+		bool bInvertMaskChannel)
+	{
+		OutMorphTargets.Reset();
+		if (Deltas.IsEmpty())
+		{
+			return;
+		}
+
+		const int32 NumBaseMeshVerts = Model->GetNumBaseMeshVerts();
+		check(Deltas.Num() % NumBaseMeshVerts == 0);
+		const int32 NumMorphTargets = Deltas.Num() / NumBaseMeshVerts;
+		check((Deltas.Num() / NumMorphTargets) == NumBaseMeshVerts);
+		check(!Model->GetVertexMap().IsEmpty());
+
+		FScopedSlowTask Task(NumMorphTargets, LOCTEXT("CreateMorphTargetProgress", "Creating morph targets"));
+		Task.MakeDialog(false);	
+
+		USkeletalMesh* SkelMesh = Model->GetSkeletalMesh();
+		FSkeletalMeshRenderData* RenderData = SkelMesh->GetResourceForRendering();
+		check(RenderData);
+		check(!RenderData->LODRenderData.IsEmpty());
+		const int32 NumRenderVertices = RenderData->LODRenderData[LOD].GetNumVertices();
+
+		// Calculate the normals for the base mesh.
+		const FSkeletalMeshModel* ImportedModel = SkelMesh->GetImportedModel();
+		const TArrayView<const uint32> IndexArray = ImportedModel->LODModels[LOD].IndexBuffer;
+		const TArrayView<const int32> VertexMap = ImportedModel->LODModels[LOD].MeshToImportVertexMap;
+		const TArrayView<const FVector3f> BaseVertexPositions = Sampler->GetUnskinnedVertexPositions();
+		const FColorVertexBuffer& ColorBuffer = RenderData->LODRenderData[LOD].StaticVertexBuffers.ColorVertexBuffer;
+
+		// Precalc an array that maps imported vertices to a render vertex.
+		TArray<int32> ImportedVertexToRenderVertexMapping;
+		ImportedVertexToRenderVertexMapping.SetNumUninitialized(NumBaseMeshVerts);
+		for (int32 Index = 0; Index < NumBaseMeshVerts; ++Index)
+		{
+			ImportedVertexToRenderVertexMapping[Index] = VertexMap.Find(Index);
+		}
+
+		TArray<FVector3f> BaseNormals;
+		if (bIncludeNormals)
+		{
+			CalcVertexNormals(BaseVertexPositions, IndexArray, VertexMap, BaseNormals);
+		}
+
+		// Initialize an engine morph target for each model morph target.
+		UE_LOG(LogMLDeformer, Display, TEXT("Initializing %d engine morph targets of %d vertices each"), NumMorphTargets, Deltas.Num() / NumMorphTargets);
+		TArray<FVector3f> DeltaNormals;
+		for (int32 MorphTargetIndex = 0; MorphTargetIndex < NumMorphTargets; ++MorphTargetIndex)
+		{
+			if (bIncludeNormals)
+			{
+				CalcMorphTargetNormals(LOD, SkelMesh, MorphTargetIndex, Deltas, BaseVertexPositions, BaseNormals, ImportedVertexToRenderVertexMapping, ColorBuffer, MaskChannel, bInvertMaskChannel, DeltaNormals);
+			}
+
+			const FName MorphName = *FString::Printf(TEXT("%s%.3d"), *NamePrefix, MorphTargetIndex);
+			UMorphTarget* MorphTarget = NewObject<UMorphTarget>(SkelMesh, MorphName);
+			MorphTarget->BaseSkelMesh = SkelMesh;
+			OutMorphTargets.Add(MorphTarget);
+
+			// Create a new LOD model for this morph.
+			TArray<FMorphTargetLODModel>& MorphLODs = MorphTarget->GetMorphLODModels();
+			MorphLODs.AddDefaulted();
+			FMorphTargetLODModel& MorphLODModel = MorphLODs.Last();
+
+			// Initialize the morph target LOD level.
+			MorphLODModel.Reset();
+			MorphLODModel.bGeneratedByEngine = true;
+			MorphLODModel.NumBaseMeshVerts = NumRenderVertices;
+			MorphLODModel.NumVertices = NumRenderVertices;
+
+			// Get the optional input mask for this morph target.
+			const TArrayView<const float> MorphMask = (MorphTargetIndex > 0) ? GetMaskForMorphTarget(MorphTargetIndex - 1) : TArrayView<const float>();
+
+			// Init deltas for this morph target.
+			MorphLODModel.Vertices.Reserve(NumRenderVertices);
+			for (int32 VertexIndex = 0; VertexIndex < NumRenderVertices; ++VertexIndex)
+			{
+				const int32 ImportedVertexNumber = VertexMap[VertexIndex];
+				if (ImportedVertexNumber != INDEX_NONE)
+				{
+					const float GlobalMaskWeight = CalcGlobalMaskWeight(VertexIndex, ColorBuffer, MaskChannel, bInvertMaskChannel);
+					const float MorphMaskWeight = MorphMask.IsEmpty() ? 1.0f : MorphMask[ImportedVertexNumber];
+					const FVector3f RawDelta = Deltas[ImportedVertexNumber + MorphTargetIndex * NumBaseMeshVerts];
+					const FVector3f RawDeltaNormal = !DeltaNormals.IsEmpty() ? DeltaNormals[ImportedVertexNumber] : FVector3f::ZeroVector;
+					FVector3f ScaledDelta;
+					FVector3f ScaledDeltaNormal;				
+					if (ProcessVertexDelta(ScaledDelta, ScaledDeltaNormal, RawDelta, RawDeltaNormal, DeltaThreshold, MorphMaskWeight, GlobalMaskWeight))
+					{
+						MorphLODModel.Vertices.AddDefaulted();
+						FMorphTargetDelta& MorphTargetDelta = MorphLODModel.Vertices.Last();
+						MorphTargetDelta.PositionDelta = ScaledDelta;
+						MorphTargetDelta.SourceIdx = VertexIndex;
+						MorphTargetDelta.TangentZDelta = bIncludeNormals ? ScaledDeltaNormal : FVector3f::ZeroVector;
+
+						// Make sure we update the list of sections that we touch.
+						int32 RenderSection = INDEX_NONE;
+						int32 TempVertexIndex = INDEX_NONE;
+						RenderData->LODRenderData[0].GetSectionFromVertexIndex(VertexIndex, RenderSection, TempVertexIndex);
+						if (RenderSection != INDEX_NONE)
+						{
+							MorphLODModel.SectionIndices.AddUnique(RenderSection);
+						}
+					}
+				}
+			}	// for all morph targets
+
+			MorphLODModel.Vertices.Shrink();
+			Task.EnterProgressFrame();
+		}
+	}
+
+	void FMLDeformerMorphModelEditorModel::CompressMorphTargets(FMorphTargetVertexInfoBuffers& OutMorphBuffers, const TArray<UMorphTarget*>& MorphTargets, int32 LOD, float MorphErrorTolerance)
+	{
+		FScopedSlowTask Task(1, LOCTEXT("CompressMorphTargetsProgress", "Compressing morph targets"));
+		Task.MakeDialog(false);	
+
+		USkeletalMesh* SkelMesh = Model->GetSkeletalMesh();
+		FSkeletalMeshRenderData* RenderData = SkelMesh->GetResourceForRendering();
+		check(RenderData);
+		check(!RenderData->LODRenderData.IsEmpty());
+		const int32 NumRenderVertices = RenderData->LODRenderData[LOD].GetNumVertices();
+
+		// Release any existing morph buffer data.
+		if (OutMorphBuffers.IsRHIIntialized() && OutMorphBuffers.IsInitialized())
+		{
+			ReleaseResourceAndFlush(&OutMorphBuffers);
+		}
+
+		// Don't empty the array of morph target data when we init the RHI buffers, as we need them to serialize later on.
+		OutMorphBuffers = FMorphTargetVertexInfoBuffers();
+		OutMorphBuffers.SetEmptyMorphCPUDataOnInitRHI(false);
+
+		// Initialize the compressed morph target buffers.
+		OutMorphBuffers.InitMorphResources
+		(
+			GMaxRHIShaderPlatform,
+			RenderData->LODRenderData[LOD].RenderSections,
+			MorphTargets,
+			NumRenderVertices,
+			LOD,
+			MorphErrorTolerance
+		);
+
+		// Reinit the render resources.
+		if (OutMorphBuffers.IsMorphCPUDataValid() && OutMorphBuffers.GetNumMorphs() > 0 && OutMorphBuffers.GetNumBatches() > 0)
+		{
+			BeginInitResource(&OutMorphBuffers);
+		}
+
+		// Update the editor actor skel mesh components for all the ones that also have an ML Deformer on it.
+		for (FMLDeformerEditorActor* EditorActor : EditorActors)
+		{
+			const UMLDeformerComponent* MLDeformerComponent = EditorActor->GetMLDeformerComponent();
+			USkeletalMeshComponent* SkelMeshComponent = EditorActor->GetSkeletalMeshComponent();
+			if (SkelMeshComponent && MLDeformerComponent)
+			{
+				SkelMeshComponent->RefreshExternalMorphTargetWeights();
+			}
+		}
+
+		UpdateMemoryUsage();
+		Task.EnterProgressFrame();
+	}
+
+	void FMLDeformerMorphModelEditorModel::DebugDrawMorphTarget(FPrimitiveDrawInterface* PDI, const TArray<FVector3f>& MorphDeltas, float DeltaThreshold, int32 MorphTargetIndex, const FVector& DrawOffset)
+	{
+		UMLDeformerVizSettings* VizSettings = Model->GetVizSettings();
+		const int32 NumVerts = Model->GetNumBaseMeshVerts();
+		const TArray<FVector3f>& UnskinnedPositions = Sampler->GetUnskinnedVertexPositions();
+		if (!MorphDeltas.IsEmpty() && 
+			(MorphDeltas.Num() % NumVerts == 0) &&
+			NumVerts == Model->GetInputInfo()->GetNumBaseMeshVertices() &&
+			UnskinnedPositions.Num() == NumVerts)
+		{
+			// Get the optional input mask for this morph target.
+			const TArrayView<const float> MorphMask = (MorphTargetIndex > 0) ? GetMaskForMorphTarget(MorphTargetIndex - 1) : TArrayView<const float>();
+
+			const int32 NumMorphTargets = MorphDeltas.Num() / NumVerts;
+			const int32 FinalMorphTargetIndex = FMath::Clamp<int32>(MorphTargetIndex, 0, NumMorphTargets - 1);
+			const FLinearColor IncludedColor = FMLDeformerEditorStyle::Get().GetColor("MLDeformer.Morphs.IncludedVertexColor");
+			const FLinearColor ExcludedColor = FMLDeformerEditorStyle::Get().GetColor("MLDeformer.Morphs.ExcludedVertexColor");
+			for (int32 VertexIndex = 0; VertexIndex < NumVerts; ++VertexIndex)
+			{
+				const FVector StartPoint = FVector(UnskinnedPositions[VertexIndex]) + DrawOffset;
+				const int32 DeltaArrayOffset = NumVerts * FinalMorphTargetIndex + VertexIndex;
+				const FVector3f RawDelta(MorphDeltas[DeltaArrayOffset]);
+				const float MorphMaskValue = MorphMask.IsEmpty() ? 1.0f : MorphMask[VertexIndex];
+				const float VertexWeight = 1.0f;
+				FVector3f ScaledDelta;
+				FVector3f ScaledDeltaNormal;
+				if (ProcessVertexDelta(ScaledDelta, ScaledDeltaNormal, RawDelta, FVector3f::ZeroVector, DeltaThreshold, MorphMaskValue, VertexWeight))
+				{
+					PDI->DrawPoint(StartPoint, IncludedColor, 1.0f, 0);
+					PDI->DrawLine(StartPoint, StartPoint + FVector(ScaledDelta), IncludedColor, 0);
+				}
+				else
+				{
+					PDI->DrawPoint(StartPoint + FVector(ScaledDelta), ExcludedColor, 0.75f, 0);
+				}
+			}
+		}
+	}
+
 	void FMLDeformerMorphModelEditorModel::InitEngineMorphTargets(const TArray<FVector3f>& Deltas)
 	{
 		UMLDeformerMorphModel* MorphModel = GetMorphModel();
@@ -224,7 +613,7 @@ namespace UE::MLDeformer
 
 		// Zero out small deltas.
 		TArray<FVector3f> MorphTargetDeltas = Deltas;
-		ZeroDeltasByThreshold(MorphTargetDeltas, MorphModel->GetMorphDeltaZeroThreshold());
+		ZeroDeltasByLengthThreshold(MorphTargetDeltas, MorphModel->GetMorphDeltaZeroThreshold());
 
 		// Turn the delta buffer in a set of engine morph targets.
 		const int32 LOD = 0;
@@ -234,7 +623,7 @@ namespace UE::MLDeformer
 
 		// Create the engine morph targets.
 		TArray<UMorphTarget*> MorphTargets;
-		CreateEngineMorphTargets(
+		CreateMorphTargets(
 			MorphTargets,
 			Deltas, 
 			TEXT("MLDeformerMorph_"), 
@@ -250,7 +639,7 @@ namespace UE::MLDeformer
 		// Now compress the morph targets to GPU friendly buffers.
 		check(MorphModel->GetMorphTargetSet().IsValid());
 		FMorphTargetVertexInfoBuffers& MorphBuffers = MorphModel->GetMorphTargetSet()->MorphBuffers;
-		CompressEngineMorphTargets(MorphBuffers, MorphTargets, LOD, MorphModel->GetMorphCompressionLevel());
+		CompressMorphTargets(MorphBuffers, MorphTargets, LOD, MorphModel->GetMorphCompressionLevel());
 
 		if (MorphBuffers.GetNumBatches() == 0 || MorphBuffers.GetNumMorphs() == 0)
 		{
@@ -276,9 +665,168 @@ namespace UE::MLDeformer
 		if (VizSettings->GetDrawMorphTargets() && VizSettings->GetVisualizationMode() == EMLDeformerVizMode::TestData)
 		{
 			const FVector DrawOffset = -VizSettings->GetMeshSpacingOffsetVector();
-			DrawMorphTarget(PDI, GetMorphModel()->GetMorphTargetDeltas(), MorphModel->GetMorphDeltaZeroThreshold(), VizSettings->GetMorphTargetNumber(), DrawOffset);
+			DebugDrawMorphTarget(PDI, GetMorphModel()->GetMorphTargetDeltas(), MorphModel->GetMorphDeltaZeroThreshold(), VizSettings->GetMorphTargetNumber(), DrawOffset);
 		}
 	}
+
+	void FMLDeformerMorphModelEditorModel::FillMaskValues(TArrayView<float> ItemMaskBuffer, float Value) const
+	{
+		const int32 NumBaseMeshVerts = Model->GetNumBaseMeshVerts();
+		check(ItemMaskBuffer.Num() == NumBaseMeshVerts);
+		for (int32 Index = 0; Index < NumBaseMeshVerts; ++Index)
+		{
+			ItemMaskBuffer[Index] = Value;
+		}
+	}
+
+	void FMLDeformerMorphModelEditorModel::ApplyBoneToMask(int32 SkeletonBoneIndex, TArrayView<float> MaskBuffer)
+	{
+		const int32 LOD = 0;
+
+		const int32 NumVerts = Model->GetNumBaseMeshVerts();
+		check(MaskBuffer.Num() == NumVerts);
+
+		const USkeletalMesh* SkeletalMesh = Model->GetSkeletalMesh();
+		FSkeletalMeshModel* ImportedModel = SkeletalMesh->GetImportedModel();
+		const FMLDeformerEditorActor* EditorActor = FindEditorActor(ActorID_Train_Base);
+		if (ImportedModel == nullptr || EditorActor == nullptr || !ImportedModel->LODModels.IsValidIndex(LOD))
+		{
+			return;
+		}
+
+		const TArray<int32>& ImportedVertexNumbers = ImportedModel->LODModels[LOD].MeshToImportVertexMap;
+		
+		const UDebugSkelMeshComponent* SkeletalMeshComponent = EditorActor->GetSkeletalMeshComponent();
+		FSkinWeightVertexBuffer& SkinWeightBuffer = *SkeletalMeshComponent->GetSkinWeightBuffer(LOD);
+		const FSkeletalMeshLODRenderData& LODData = SkeletalMesh->GetResourceForRendering()->LODRenderData[LOD];
+		const int32 NumRenderVerts = LODData.GetNumVertices();
+		for (int32 VertexIndex = 0; VertexIndex < NumRenderVerts; ++VertexIndex)
+		{
+			const int32 ImportedVertexNumber = ImportedVertexNumbers[VertexIndex];
+			if (ImportedVertexNumber == INDEX_NONE)
+			{
+				continue;
+			}
+
+			// Find the render section, which we need to find the right bone index.
+			int32 SectionIndex = INDEX_NONE;
+			int32 SectionVertexIndex = INDEX_NONE;
+			LODData.GetSectionFromVertexIndex(VertexIndex, SectionIndex, SectionVertexIndex);
+
+			// Iterate over all skinning influences.
+			const int32 NumInfluences = SkinWeightBuffer.GetMaxBoneInfluences();
+			for (int32 InfluenceIndex = 0; InfluenceIndex < NumInfluences; ++InfluenceIndex)
+			{
+				const int32 BoneIndex = SkinWeightBuffer.GetBoneIndex(VertexIndex, InfluenceIndex);
+				const uint16 WeightByte = SkinWeightBuffer.GetBoneWeight(VertexIndex, InfluenceIndex);
+				if (WeightByte > 0)
+				{
+					const int32 RealBoneIndex = LODData.RenderSections[SectionIndex].BoneMap[BoneIndex];
+					if (RealBoneIndex == SkeletonBoneIndex)
+					{
+						const float Weight = static_cast<float>(WeightByte) * UE::AnimationCore::InvMaxRawBoneWeightFloat;
+						MaskBuffer[ImportedVertexNumber] += Weight;
+					}
+				}
+			}
+		}		
+	}
+
+	void FMLDeformerMorphModelEditorModel::AddRequiredBones(const FReferenceSkeleton& RefSkel, int32 SkeletonBoneIndex, const TArray<int32>& VirtualParentTable, TArray<int32>& OutBonesAdded)
+	{
+		// Add all bones virtually parented to this bone to the mask if we haven't already.
+		for (int32 Index = 0; Index < VirtualParentTable.Num(); ++Index)
+		{
+			const int32 VirtualParent = VirtualParentTable[Index];
+			if (VirtualParent == SkeletonBoneIndex && !OutBonesAdded.Contains(Index))
+			{
+				OutBonesAdded.Add(Index);
+			}
+		}
+	}
+
+	void FMLDeformerMorphModelEditorModel::RecursiveAddBoneToMaskUpwards(const FReferenceSkeleton& RefSkel, int32 SkeletonBoneIndex, int32 MaxHierarchyDepth, const TArray<int32>& VirtualParentTable, TArray<int32>& OutBonesAdded, int32 CurHierarchyDepth)
+	{
+		if (CurHierarchyDepth > MaxHierarchyDepth)
+		{
+			return;
+		}
+
+		// Apply the current bone to the mask.
+		if (!OutBonesAdded.Contains(SkeletonBoneIndex))
+		{
+			OutBonesAdded.Add(SkeletonBoneIndex);
+		}
+
+		// Apply the parent bone.
+		const int32 ParentSkeletonBoneIndex = RefSkel.GetParentIndex(SkeletonBoneIndex);
+		if (ParentSkeletonBoneIndex != INDEX_NONE)
+		{
+			RecursiveAddBoneToMaskUpwards(RefSkel, ParentSkeletonBoneIndex, MaxHierarchyDepth, VirtualParentTable, OutBonesAdded, CurHierarchyDepth + 1);
+		}
+	}
+
+	void FMLDeformerMorphModelEditorModel::RecursiveAddBoneToMaskDownwards(const FReferenceSkeleton& RefSkel, int32 SkeletonBoneIndex, int32 MaxHierarchyDepth, const TArray<int32>& VirtualParentTable, TArray<int32>& OutBonesAdded, int32 CurHierarchyDepth)
+	{
+		if (CurHierarchyDepth > MaxHierarchyDepth)
+		{
+			return;
+		}
+
+		// Apply the current bone to the mask.
+		if (!OutBonesAdded.Contains(SkeletonBoneIndex))
+		{
+			OutBonesAdded.Add(SkeletonBoneIndex);
+		}
+
+		// Find all child bones.
+		TArray<int32> ChildBones;
+		ChildBones.Reserve(8);
+		RefSkel.GetDirectChildBones(SkeletonBoneIndex, ChildBones);
+
+		// Now recursively add the child bones.
+		for (const int32 ChildIndex : ChildBones)
+		{
+			RecursiveAddBoneToMaskDownwards(RefSkel, ChildIndex, MaxHierarchyDepth, VirtualParentTable, OutBonesAdded, CurHierarchyDepth + 1);
+		}
+	}
+
+	int32 FMLDeformerMorphModelEditorModel::FindVirtualParentIndex(const FReferenceSkeleton& RefSkel, int32 BoneIndex, const TArray<FName>& IncludedBoneNames) const
+	{
+		int32 CurBoneIndex = BoneIndex;
+		while (CurBoneIndex != INDEX_NONE)
+		{
+			const int32 ParentIndex = RefSkel.GetParentIndex(CurBoneIndex);
+			if (ParentIndex == INDEX_NONE)
+			{
+				break;
+			}
+
+			const FName ParentName = RefSkel.GetBoneName(ParentIndex);
+			if (IncludedBoneNames.Contains(ParentName))
+			{
+				return ParentIndex;
+			}
+
+			CurBoneIndex = ParentIndex;
+		}
+
+		return BoneIndex;
+	}
+
+	TArray<int32> FMLDeformerMorphModelEditorModel::BuildVirtualParentTable(const FReferenceSkeleton& RefSkel, const TArray<FName>& IncludedBoneNames) const
+	{
+		TArray<int32> VirtualParentTable;
+		VirtualParentTable.SetNumUninitialized(RefSkel.GetNum());
+
+		for (int32 BoneIndex = 0; BoneIndex < RefSkel.GetNum(); ++BoneIndex)
+		{
+			VirtualParentTable[BoneIndex] = FindVirtualParentIndex(RefSkel, BoneIndex, IncludedBoneNames);
+		}
+
+		return MoveTemp(VirtualParentTable);
+	}
+
 }	// namespace UE::MLDeformer
 
 #undef LOCTEXT_NAMESPACE
diff --git a/Engine/Plugins/Animation/MLDeformer/MLDeformerFramework/Source/MLDeformerFrameworkEditor/Public/MLDeformerEditorModel.h b/Engine/Plugins/Animation/MLDeformer/MLDeformerFramework/Source/MLDeformerFrameworkEditor/Public/MLDeformerEditorModel.h
--- a/Engine/Plugins/Animation/MLDeformer/MLDeformerFramework/Source/MLDeformerFrameworkEditor/Public/MLDeformerEditorModel.h
+++ b/Engine/Plugins/Animation/MLDeformer/MLDeformerFramework/Source/MLDeformerFrameworkEditor/Public/MLDeformerEditorModel.h
@@ -57,6 +57,7 @@ namespace UE::MLDeformer
 	class FMLDeformerEditorToolkit;
 	class FMLDeformerSampler;
 	class FMLDeformerEditorModel;
+	class SMLDeformerInputWidget;
 
 	/**
 	 * The base class for the editor side of an UMLDeformerModel.
@@ -563,6 +564,21 @@ namespace UE::MLDeformer
 		 */
 		virtual FString GetHeatMapDeformerGraphPath() const;
 
+		/**
+		 * Refresh the contents of the input widget.
+		 * This includes the bone and curve list contents.
+		 */
+		virtual void RefreshInputWidget();
+
+		/**
+		 * Create the input widget.
+		 * The default widget that is created contains the bone and curve list.
+		 * You can inherit from the SMLDeformerInputWidget base class and override this method to return your custom widget.
+		 * This can be useful when you want to add new types of inputs to your model.
+		 * @return A shared pointer to the input widget that will appear in the UI.
+		 */
+		virtual TSharedPtr<SMLDeformerInputWidget> CreateInputWidget();
+
 		/**
 		 * Generate the normals for a given morph target.
 		 * @param LOD The LOD level.
@@ -572,6 +588,7 @@ namespace UE::MLDeformer
 		 * @param BaseNormals The normals of the base mesh.
 		 * @param OutDeltaNormals The array that we will write the generated normals to. This will automatically be resized by this method.
 		 */
+		UE_DEPRECATED(5.3, "Please call FMLDeformerMorphModelEditorModel::CalcMorphTargetNormals instead.")
 		virtual void GenerateNormalsForMorphTarget(int32 LOD, USkeletalMesh* SkelMesh, int32 MorphTargetIndex, TArrayView<const FVector3f> Deltas, TArrayView<const FVector3f> BaseVertexPositions, TArrayView<FVector3f> BaseNormals, TArray<FVector3f>& OutDeltaNormals);
 
 		/** Invalidate the memory usage, so it gets updated in the UI again. */
@@ -738,6 +755,16 @@ namespace UE::MLDeformer
 		 */
 		void InitCurveIncludeListToAnimatedCurvesOnly();
 
+		/**
+		 * Add all animated bones to the bone include list.
+		 */
+		void AddAnimatedBonesToBonesIncludeList();
+
+		/**
+		 * Add all animated curves to the curve include list.
+		 */
+		void AddAnimatedCurvesToCurvesIncludeList();
+
 		/**
 		 * Get the number of curves on a specific skeletal mesh.
 		 * @param SkelMesh The skeletal mesh to get the curves count from.
@@ -768,6 +795,7 @@ namespace UE::MLDeformer
 		 * @param MorphTargetIndex The morph target number to visualize.
 		 * @param DrawOffset An offset to perform the debug draw at.
 		 */
+		UE_DEPRECATED(5.3, "Please call FMLDeformerMorphModelEditorModel::DebugDrawMorphTargets instead.")
 		void DrawMorphTarget(FPrimitiveDrawInterface* PDI, const TArray<FVector3f>& MorphDeltas, float DeltaThreshold, int32 MorphTargetIndex, const FVector& DrawOffset);
 
 		/**
@@ -786,6 +814,10 @@ namespace UE::MLDeformer
 		 * @return A corrected frame time, which is equal to FrameTimeForFrameNumber unless a floating point error means the floor is not equal to the FrameNumber
 		 */
 		static float CorrectedFrameTime(int32 FrameNumber, float FrameTimeForFrameNumber, FFrameRate FrameRate);
+
+		TSharedPtr<SMLDeformerInputWidget> GetInputWidget() const			{ return InputWidget; }
+		void SetInputWidget(TSharedPtr<SMLDeformerInputWidget> Widget)		{ InputWidget = Widget; }
+
 	protected:
 		/**
 		 * Executed when a property changes by a change in the UI.
@@ -818,6 +850,7 @@ namespace UE::MLDeformer
 		bool IsEditorReadyForTrainingBasicChecks();
 
 		/** Zero all deltas with a length equal to, or smaller than the threshold value. */
+		UE_DEPRECATED(5.3, "Please call FMLDeformerMorphModelEditorModel::ZeroDeltasByLengthThreshold instead.")
 		void ZeroDeltasByThreshold(TArray<FVector3f>& Deltas, float Threshold);
 
 		/**
@@ -831,7 +864,10 @@ namespace UE::MLDeformer
 		 *        recomputing normals can lead to lower quality results, in trade for faster performance.
 		 * @param MaskChannel The weight mask mode, which specifies what channel to get the weight data from. Such channel allows the user to define what areas the deformer should for example not be active in.
 		 * @param bInvertMaskChannel Specifies whether the weight mask should be inverted or not.
+		 * @param MaskBuffer An optional mask buffer that contains 'Model->GetNumBaseMeshVerts() * (MorphModel->GetNumMorphTargets() - 1)' number of floats. Deltas will be multiplied by this value. 
+		 *        When the mask buffer is an empty array, it will be ignored.
 		 */
+		UE_DEPRECATED(5.3, "Please call FMLDeformerMorphModelEditorModel::CreateMorphTargets instead.")
 		void CreateEngineMorphTargets(
 			TArray<UMorphTarget*>& OutMorphTargets,
 			const TArray<FVector3f>& Deltas,
@@ -849,8 +885,20 @@ namespace UE::MLDeformer
 		 * @param LOD The LOD index to generate the morphs for.
 		 * @param MorphErrorTolerance The error tolerance for the delta compression, in cm. Higher values compress better but can result in artifacts.
 		 */
+		UE_DEPRECATED(5.3, "Please call FMLDeformerMorphModelEditorModel::CompressMorphTargets instead.")
 		void CompressEngineMorphTargets(FMorphTargetVertexInfoBuffers& OutMorphBuffers, const TArray<UMorphTarget*>& MorphTargets, int32 LOD = 0, float MorphErrorTolerance = 0.01f);
 
+		/**
+		 * Calculate the normals for each vertex, given the triangle data and positions.
+		 * It computes this by summing up the face normals for each vertex using that face, and normalizing them at the end.
+		 * @param VertexPositions The buffer with vertex positions. This is the size of the number of imported vertices.
+		 * @param IndexArray The index buffer, which contains NumTriangles * 3 number of integers.
+		 * @param VertexMap For each render vertex, an imported vertex number. For example, for a cube these indices go from 0..7.
+		 * @param OutNormals The array that will contain the normals. This will automatically be resized internally by this method.
+		 */
+		UE_DEPRECATED(5.3, "Please call FMLDeformerMorphModelEditorModel::CalcVertexNormals instead.")
+		void CalcMeshNormals(TArrayView<const FVector3f> VertexPositions, TArrayView<const uint32> IndexArray, TArrayView<const int32> VertexMap, TArray<FVector3f>& OutNormals) const;
+
 		/**
 		 * Get the base actor depending on the currently active mode (train or test mode).
 		 * The base actor is the linear skinned one.
@@ -877,16 +925,6 @@ namespace UE::MLDeformer
 		 */
 		void UpdateRanges();
 
-		/**
-		 * Calculate the normals for each vertex, given the triangle data and positions.
-		 * It computes this by summing up the face normals for each vertex using that face, and normalizing them at the end.
-		 * @param VertexPositions The buffer with vertex positions. This is the size of the number of imported vertices.
-		 * @param IndexArray The index buffer, which contains NumTriangles * 3 number of integers.
-		 * @param VertexMap For each render vertex, an imported vertex number. For example, for a cube these indices go from 0..7.
-		 * @param OutNormals The array that will contain the normals. This will automatically be resized internally by this method.
-		 */
-		void CalcMeshNormals(TArrayView<const FVector3f> VertexPositions, TArrayView<const uint32> IndexArray, TArrayView<const int32> VertexMap, TArray<FVector3f>& OutNormals) const;
-
 	protected:
 		/** The runtime model associated with this editor model. */
 		TObjectPtr<UMLDeformerModel> Model = nullptr;
@@ -900,6 +938,9 @@ namespace UE::MLDeformer
 		/** A pointer to the sampler, which can sample target meshes to calculate deltas. */
 		FMLDeformerSampler* Sampler = nullptr;
 
+		/** The inputs widget. */
+		TSharedPtr<SMLDeformerInputWidget> InputWidget;
+
 		/**
 		 * The input info as currently setup in the editor.
 		 * This is different from the runtime model's input info, as that is the one that was used to train with.
