[MLDeformer]
- Show an error when the ML Deformer component works on an invalid skeletal mesh.
- Stop spamming the log when enabling heatmaps and having geom cache track's that cannot be mapped to the skeletal mesh.

#trivial
#rnx

diff --git a/Engine/Plugins/Animation/MLDeformer/MLDeformerFramework/Source/MLDeformerFramework/Private/MLDeformerGeomCacheModel.cpp b/Engine/Plugins/Animation/MLDeformer/MLDeformerFramework/Source/MLDeformerFramework/Private/MLDeformerGeomCacheModel.cpp
--- a/Engine/Plugins/Animation/MLDeformer/MLDeformerFramework/Source/MLDeformerFramework/Private/MLDeformerGeomCacheModel.cpp
+++ b/Engine/Plugins/Animation/MLDeformer/MLDeformerFramework/Source/MLDeformerFramework/Private/MLDeformerGeomCacheModel.cpp
@@ -57,7 +57,7 @@ void UMLDeformerGeomCacheModel::SampleGroundTruthPositions(float SampleTime, TAr
 	{
 		TArray<FString> FailedImportedMeshNames;
 		TArray<FString> VertexMisMatchNames;
-		UE::MLDeformer::GenerateGeomCacheMeshMappings(GetSkeletalMesh(), GeomCache, MeshMappings, FailedImportedMeshNames, VertexMisMatchNames);
+		UE::MLDeformer::GenerateGeomCacheMeshMappings(GetSkeletalMesh(), GeomCache, MeshMappings, FailedImportedMeshNames, VertexMisMatchNames, /**bSuppressLog*/true);
 	}
 
 	UE::MLDeformer::SampleGeomCachePositions(
diff --git a/Engine/Plugins/Animation/MLDeformer/MLDeformerFramework/Source/MLDeformerFramework/Private/MLDeformerModelInstance.cpp b/Engine/Plugins/Animation/MLDeformer/MLDeformerFramework/Source/MLDeformerFramework/Private/MLDeformerModelInstance.cpp
--- a/Engine/Plugins/Animation/MLDeformer/MLDeformerFramework/Source/MLDeformerFramework/Private/MLDeformerModelInstance.cpp
+++ b/Engine/Plugins/Animation/MLDeformer/MLDeformerFramework/Source/MLDeformerFramework/Private/MLDeformerModelInstance.cpp
@@ -90,6 +90,17 @@ void UMLDeformerModelInstance::Init(USkeletalMeshComponent* SkelMeshComponent)
 	if (SkelMeshComponent == nullptr)
 	{
 		AssetBonesToSkelMeshMappings.Empty();
+
+		UMLDeformerAsset* Asset = GetModel()->GetDeformerAsset();
+		if (Asset && Asset->GetModel()->GetInputInfo())
+		{
+			const FSoftObjectPath& SkelMeshPath = Asset->GetModel()->GetInputInfo()->GetSkeletalMesh();
+			if (SkelMeshPath.IsValid())
+			{
+				UE_LOG(LogMLDeformer, Error, TEXT("ML Deformer '%s' is trained on skeletal mesh '%s', which cannot be found on actor!"), *Asset->GetName(), *SkelMeshPath.ToString());
+			}
+		}
+
 		return;
 	}
 
