[MLDeformer] Added automated tests. Also moved the tests to the EngineTest project. Includes some small optimizations and improvements and required changes to the API.

diff --git a/Engine/Plugins/Animation/MLDeformer/MLDeformerFramework/Source/MLDeformerFramework/Private/Tests/MLDeformerTestModel.cpp b/Engine/Plugins/Animation/MLDeformer/MLDeformerFramework/Source/MLDeformerFramework/Private/Tests/MLDeformerTestModel.cpp
deleted file mode 100644
diff --git a/Engine/Plugins/Animation/MLDeformer/MLDeformerFramework/Source/MLDeformerFramework/Private/Tests/MLDeformerTestModel.h b/Engine/Plugins/Animation/MLDeformer/MLDeformerFramework/Source/MLDeformerFramework/Private/Tests/MLDeformerTestModel.h
deleted file mode 100644
diff --git a/Engine/Plugins/Animation/MLDeformer/MLDeformerFramework/Source/MLDeformerFramework/Private/Tests/ModelTests.cpp b/Engine/Plugins/Animation/MLDeformer/MLDeformerFramework/Source/MLDeformerFramework/Private/Tests/ModelTests.cpp
deleted file mode 100644
diff --git a/Engine/Plugins/Animation/MLDeformer/MLDeformerFramework/Source/MLDeformerFrameworkEditor/Private/MLDeformerEditorToolkit.cpp b/Engine/Plugins/Animation/MLDeformer/MLDeformerFramework/Source/MLDeformerFrameworkEditor/Private/MLDeformerEditorToolkit.cpp
--- a/Engine/Plugins/Animation/MLDeformer/MLDeformerFramework/Source/MLDeformerFrameworkEditor/Private/MLDeformerEditorToolkit.cpp
+++ b/Engine/Plugins/Animation/MLDeformer/MLDeformerFramework/Source/MLDeformerFrameworkEditor/Private/MLDeformerEditorToolkit.cpp
@@ -36,6 +36,11 @@ namespace UE::MLDeformer
 	const FName MLDeformerEditorModes::Editor("MLDeformerEditorMode");
 	const FName MLDeformerEditorAppName = FName(TEXT("MLDeformerEditorApp"));
 
+	FMLDeformerEditorToolkit::~FMLDeformerEditorToolkit()
+	{
+		ActiveModel.Reset();
+	}
+
 	void FMLDeformerEditorToolkit::InitAssetEditor(
 		const EToolkitMode::Type Mode,
 		const TSharedPtr<IToolkitHost>& InitToolkitHost,
@@ -57,7 +62,7 @@ namespace UE::MLDeformer
 			if (ModelRegistry.GetNumRegisteredModels() > 0)
 			{
 				const int32 HighestPriorityIndex = ModelRegistry.GetHighestPriorityModelIndex();
-				OnModelChanged(HighestPriorityIndex);
+				OnModelChanged(HighestPriorityIndex, true);
 				DeformerModel = DeformerAsset->GetModel();
 			}
 		}
@@ -131,7 +136,7 @@ namespace UE::MLDeformer
 		{
 			const EAppReturnType::Type ReturnType = FMessageDialog::Open(
 				EAppMsgType::Ok, 
-				LOCTEXT("NoModelsFoundWarningMessage", "No ML Deformer models have been registered.\nPlease load a model plugin.\n\nThere is nothing to do inside this editor, until you load a model plugin."),
+				FText(LOCTEXT("NoModelsFoundWarningMessage", "No ML Deformer models have been registered.\nPlease load a model plugin.\n\nThere is nothing to do inside this editor, until you load a model plugin.")),
 				LOCTEXT("NoModelsFoundWarningTitle", "No ML Deformer Models Found"));
 		}
 	}
@@ -173,10 +178,10 @@ namespace UE::MLDeformer
 		const bool bShouldCloseWindowAfterMenuSelection = true;
 		FMenuBuilder MenuBuilder(bShouldCloseWindowAfterMenuSelection, InCommandList);
 
-		FUIAction TrainItemAction(FExecuteAction::CreateSP(this, &FMLDeformerEditorToolkit::OnVizModeChanged, EMLDeformerVizMode::TrainingData));
+		FUIAction TrainItemAction(FExecuteAction::CreateSP(this, &FMLDeformerEditorToolkit::SwitchVizMode, EMLDeformerVizMode::TrainingData));
 		MenuBuilder.AddMenuEntry(GetVizModeName(EMLDeformerVizMode::TrainingData), TAttribute<FText>(), FSlateIcon(), TrainItemAction);
 
-		FUIAction TestItemAction(FExecuteAction::CreateSP(this, &FMLDeformerEditorToolkit::OnVizModeChanged, EMLDeformerVizMode::TestData));
+		FUIAction TestItemAction(FExecuteAction::CreateSP(this, &FMLDeformerEditorToolkit::SwitchVizMode, EMLDeformerVizMode::TestData));
 		MenuBuilder.AddMenuEntry(GetVizModeName(EMLDeformerVizMode::TestData), TAttribute<FText>(), FSlateIcon(), TestItemAction);
 
 		return MenuBuilder.MakeWidget();
@@ -221,14 +226,14 @@ namespace UE::MLDeformer
 				}
 			}
 
-			FUIAction ItemAction(FExecuteAction::CreateSP(this, &FMLDeformerEditorToolkit::OnModelChanged, UnsortedModelIndex));
+			FUIAction ItemAction(FExecuteAction::CreateSP(this, &FMLDeformerEditorToolkit::OnModelChanged, UnsortedModelIndex, false));
 			MenuBuilder.AddMenuEntry(FText::FromString(ModelNames[ModelIndex]), TAttribute<FText>(), FSlateIcon(), ItemAction);
 		}
 
 		return MenuBuilder.MakeWidget();
 	}
 
-	void FMLDeformerEditorToolkit::OnVizModeChanged(EMLDeformerVizMode Mode)
+	void FMLDeformerEditorToolkit::SwitchVizMode(EMLDeformerVizMode Mode)
 	{
 		if (!ActiveModel.IsValid())
 		{
@@ -239,9 +244,21 @@ namespace UE::MLDeformer
 		OnSwitchedVisualizationMode();
 	}
 
+	// Deprecated, use SwitchVizMode.
+	void FMLDeformerEditorToolkit::OnVizModeChanged(EMLDeformerVizMode Mode)
+	{
+		SwitchVizMode(Mode);
+	}
+
+	// Deprecated, use OnModelChanged with two parameters.
 	void FMLDeformerEditorToolkit::OnModelChanged(int Index)
 	{
-		if (ActiveModel)
+		OnModelChanged(Index, true);
+	}
+
+	void FMLDeformerEditorToolkit::OnModelChanged(int Index, bool bForceChange)
+	{
+		if (ActiveModel && !bForceChange)
 		{
 			const EAppReturnType::Type ReturnType = FMessageDialog::Open(
 				EAppMsgType::YesNo, 
@@ -349,74 +366,96 @@ namespace UE::MLDeformer
 		return LOCTEXT("SelectModel", "<Select a Model>");
 	}
 
-	void FMLDeformerEditorToolkit::FillToolbar(FToolBarBuilder& ToolbarBuilder)
+	bool FMLDeformerEditorToolkit::IsTraining() const
 	{
-		ToolbarBuilder.BeginSection("Training");
+		return bIsTraining;
+	}
+
+	bool FMLDeformerEditorToolkit::Train(bool bSuppressDialogs)
+	{
+		check(ActiveModel);
+		bIsTraining = true;
+
+		// Ask if we want to retrain the network if we already have something trained.
+		UMLDeformerModel* Model = ActiveModel->GetModel();
+		if (ActiveModel->IsTrained() && !bSuppressDialogs)
 		{
-			ToolbarBuilder.AddToolBarButton(
-				FUIAction(
-					FExecuteAction::CreateLambda([this]()
-					{
-						check(ActiveModel);
+			const EAppReturnType::Type ConfirmReturnType = FMessageDialog::Open(
+				EAppMsgType::YesNo, 
+				LOCTEXT("RetrainConfirmationMessage", "This asset already has been trained.\n\nAre you sure you would like to re-train the network with your current settings?"),
+				LOCTEXT("RetrainConfirmationTitle", "Re-train the network?"));
 
-						// Ask if we want to retrain the network if we already have something trained.
-						UMLDeformerModel* Model = ActiveModel->GetModel();
-						if (ActiveModel->IsTrained())
-						{
-							const EAppReturnType::Type ConfirmReturnType = FMessageDialog::Open(
-								EAppMsgType::YesNo, 
-								LOCTEXT("RetrainConfirmationMessage", "This asset already has been trained.\n\nAre you sure you would like to re-train the network with your current settings?"),
-								LOCTEXT("RetrainConfirmationTitle", "Re-train the network?"));
-
-							if (ConfirmReturnType == EAppReturnType::No || ConfirmReturnType == EAppReturnType::Cancel)
-							{
-								return;
-							}
-						}
+			if (ConfirmReturnType == EAppReturnType::No || ConfirmReturnType == EAppReturnType::Cancel)
+			{
+				bIsTraining = false;
+				return true;
+			}
+		}
 
-						ShowNotification(LOCTEXT("StartTraining", "Starting training process"), SNotificationItem::ECompletionState::CS_Pending, true);
+		ShowNotification(LOCTEXT("StartTraining", "Starting training process"), SNotificationItem::ECompletionState::CS_Pending, true);
 
-						ActiveModel->OnPreTraining();
+		ActiveModel->OnPreTraining();
 
-						// Change the interpolation type for the training sequence to step.
-						if (UAnimSequence* AnimSequence = Model->GetAnimSequence())
-						{
-							AnimSequence->Interpolation = EAnimInterpolationType::Step;
-						}
+		// Change the interpolation type for the training sequence to step.
+		if (UAnimSequence* AnimSequence = Model->GetAnimSequence())
+		{
+			AnimSequence->Interpolation = EAnimInterpolationType::Step;
+		}
 
-						// Initialize the training inputs.
-						ActiveModel->UpdateEditorInputInfo();
+		// Initialize the training inputs.
+		ActiveModel->UpdateEditorInputInfo();
 
-						// Make sure we have something to train on.
-						// If this triggers, the train button most likely was enabled while it shouldn't be.
-						check(!ActiveModel->GetEditorInputInfo()->IsEmpty());
+		// Make sure we have something to train on.
+		// If this triggers, the train button most likely was enabled while it shouldn't be.
+		check(!ActiveModel->GetEditorInputInfo()->IsEmpty());
 
-						// Train the model, which executes the Python code's "train" function.
-						const double StartTime = FPlatformTime::Seconds();
-						const ETrainingResult TrainingResult = ActiveModel->Train();
-						const double TrainingDuration = FPlatformTime::Seconds() - StartTime;
-						bool bUsePartiallyTrained = true;	// Will get modified by the HandleTrainingResult function below.
-						const bool bMarkDirty = HandleTrainingResult(TrainingResult, TrainingDuration, bUsePartiallyTrained);
-						ActiveModel->OnPostTraining(TrainingResult, bUsePartiallyTrained);
+		// Train the model, which executes the Python code's "train" function.
+		const double StartTime = FPlatformTime::Seconds();
+		const ETrainingResult TrainingResult = ActiveModel->Train();
+		const double TrainingDuration = FPlatformTime::Seconds() - StartTime;
+		bool bUsePartiallyTrained = true;	// Will get modified by the HandleTrainingResult function below.
+		bool bSuccess = true;
+		const bool bMarkDirty = HandleTrainingResult(TrainingResult, TrainingDuration, bUsePartiallyTrained, bSuppressDialogs, bSuccess);
+		ActiveModel->OnPostTraining(TrainingResult, bUsePartiallyTrained);
 
-						const bool NeedsResamplingBackup = ActiveModel->GetResamplingInputOutputsNeeded();
-						ActiveModel->SetResamplingInputOutputsNeeded(NeedsResamplingBackup);
-						ActiveModel->RefreshMLDeformerComponents();
-						ActiveModel->SetHeatMapMaterialEnabled(Model->GetVizSettings()->GetShowHeatMap());
+		const bool NeedsResamplingBackup = ActiveModel->GetResamplingInputOutputsNeeded();
+		ActiveModel->SetResamplingInputOutputsNeeded(NeedsResamplingBackup);
+		ActiveModel->RefreshMLDeformerComponents();
+		ActiveModel->SetHeatMapMaterialEnabled(Model->GetVizSettings()->GetShowHeatMap());
 
-						if (bMarkDirty)
-						{
-							DeformerAsset->Modify();
-						}
+		if (bMarkDirty)
+		{
+			DeformerAsset->Modify();
+		}
+
+		GetModelDetailsView()->ForceRefresh();
+		GetVizSettingsDetailsView()->ForceRefresh();
+
+		bIsTraining = false;
+		return bSuccess;
+	}
+
+	bool FMLDeformerEditorToolkit::IsTrainButtonEnabled() const
+	{
+		return ActiveModel.Get() ? ActiveModel->IsReadyForTraining() : false;
+	}
 
-						GetModelDetailsView()->ForceRefresh();
-						GetVizSettingsDetailsView()->ForceRefresh();
-					}),
-					FCanExecuteAction::CreateLambda(
-						[this]() 
+	void FMLDeformerEditorToolkit::FillToolbar(FToolBarBuilder& ToolbarBuilder)
+	{
+		ToolbarBuilder.BeginSection("Training");
+		{
+			ToolbarBuilder.AddToolBarButton(
+				FUIAction
+				(
+					//FExecuteAction::CreateRaw(this, &FMLDeformerEditorToolkit::Train, false),
+					FExecuteAction::CreateLambda
+					(
+						[this]()
 						{
-							return ActiveModel.Get() ? ActiveModel->IsReadyForTraining() : false;
-						})
+							Train(/*bSuppressDialogs*/false);
+						}
+					),
+					FCanExecuteAction::CreateRaw(this, &FMLDeformerEditorToolkit::IsTrainButtonEnabled)
 				),
 				NAME_None,
 				LOCTEXT("TrainModel", "Train Model"),
@@ -445,11 +484,11 @@ namespace UE::MLDeformer
 		ToolbarBuilder.EndSection();
 	}
 
-	bool FMLDeformerEditorToolkit::HandleTrainingResult(ETrainingResult TrainingResult, double TrainingDuration, bool& bOutUsePartiallyTrained)
+	bool FMLDeformerEditorToolkit::HandleTrainingResult(ETrainingResult TrainingResult, double TrainingDuration, bool& bOutUsePartiallyTrained, bool bSuppressDialogs, bool& bOutSuccess)
 	{
+		bOutSuccess = false;
 		bOutUsePartiallyTrained = true;
 
-		FText WindowTitle(LOCTEXT("TrainingResultsWindowTitle", "Training Results"));
 		FText WindowMessage;
 
 		// Calculate hours, minutes and seconds.
@@ -491,6 +530,7 @@ namespace UE::MLDeformer
 					WindowMessage = FText::Format(LOCTEXT("TrainingSuccess", "Training completed successfully!\n\nTraining time: {Duration}"), SuccessArgs);
 					ActiveModel->InitInputInfo(ActiveModel->GetModel()->GetInputInfo());
 					bMarkDirty = true;
+					bOutSuccess = true;
 				}
 			}
 			break;
@@ -499,10 +539,15 @@ namespace UE::MLDeformer
 			case ETrainingResult::Aborted:
 			{
 				ActiveModel->SetResamplingInputOutputsNeeded(false);
-				const EAppReturnType::Type ReturnType = FMessageDialog::Open(
-					EAppMsgType::YesNo, 
-					LOCTEXT("TrainingAbortedMessage", "Training has been aborted.\nThe neural network has only been partially trained.\nWould you like to use this partially trained network?"),
-					LOCTEXT("TrainingAbortedMessageTitle", "Use partially trained network?"));
+
+				EAppReturnType::Type ReturnType = EAppReturnType::No;	// When we suppress dialogs, do not use the the partially trained network on abort.
+				if (!bSuppressDialogs)
+				{
+					ReturnType = FMessageDialog::Open(
+						EAppMsgType::YesNo, 
+						LOCTEXT("TrainingAbortedMessage", "Training has been aborted.\nThe neural network has only been partially trained.\nWould you like to use this partially trained network?"),
+						LOCTEXT("TrainingAbortedMessageTitle", "Use partially trained network?"));
+				}
 
 				if (ReturnType == EAppReturnType::Yes)
 				{
@@ -522,6 +567,8 @@ namespace UE::MLDeformer
 					bOutUsePartiallyTrained = false;
 					ShowNotification(LOCTEXT("TrainingAborted", "Training aborted!"), SNotificationItem::ECompletionState::CS_None, true);
 				}
+
+				bOutSuccess = true;
 			}
 			break;
 
@@ -555,9 +602,9 @@ namespace UE::MLDeformer
 		}
 
 		// Show a message window.
-		if (!WindowMessage.IsEmpty())
+		if (!WindowMessage.IsEmpty() && !bSuppressDialogs)
 		{
-			FMessageDialog::Open(EAppMsgType::Ok, WindowMessage, WindowTitle);
+			FMessageDialog::Open(EAppMsgType::Ok, WindowMessage, LOCTEXT("TrainingResultsWindowTitle", "Training Results"));
 		}
 
 		ActiveModel->UpdateDeformerGraph();
@@ -837,6 +884,26 @@ namespace UE::MLDeformer
 			ActiveModel->SetViewRange(ViewRange);
 		}
 	}
+
+	bool FMLDeformerEditorToolkit::SwitchModelType(UClass* ModelType, bool bForceChange)
+	{
+		FMLDeformerEditorModule& EditorModule = FModuleManager::GetModuleChecked<FMLDeformerEditorModule>("MLDeformerFrameworkEditor");
+		FMLDeformerEditorModelRegistry& ModelRegistry = EditorModule.GetModelRegistry();
+
+		// Get the list of all registered model types.
+		TArray<UClass*> ModelTypes;
+		ModelRegistry.GetRegisteredModels().GenerateKeyArray(ModelTypes);
+
+		// Find the model index for this type.
+		const int32 Index = ModelTypes.Find(ModelType);
+		if (Index == INDEX_NONE)
+		{
+			return false;
+		}
+
+		OnModelChanged(Index, bForceChange);
+		return true;
+	}
 }	// namespace UE::MLDeformer
 
 #undef LOCTEXT_NAMESPACE
diff --git a/Engine/Plugins/Animation/MLDeformer/MLDeformerFramework/Source/MLDeformerFrameworkEditor/Private/MLDeformerTestHelpers.cpp b/Engine/Plugins/Animation/MLDeformer/MLDeformerFramework/Source/MLDeformerFrameworkEditor/Private/MLDeformerTestHelpers.cpp
new file mode 100644
--- /dev/null
+++ b/Engine/Plugins/Animation/MLDeformer/MLDeformerFramework/Source/MLDeformerFrameworkEditor/Private/MLDeformerTestHelpers.cpp
@@ -0,0 +1,116 @@
+// Copyright Epic Games, Inc. All Rights Reserved.
+
+#include "MLDeformerTestHelpers.h"
+#include "MLDeformerModule.h"
+#include "MLDeformerEditorToolkit.h"
+#include "MLDeformerEditorModule.h"
+#include "MLDeformerModelRegistry.h"
+#include "MLDeformerAsset.h"
+
+namespace UE::MLDeformer
+{	
+	//--------------------------------------------------------------------------------------
+	// FMLDeformerScopedEditor
+	//--------------------------------------------------------------------------------------
+	FMLDeformerScopedEditor::FMLDeformerScopedEditor(FMLDeformerEditorToolkit* InEditor)
+	{
+		Editor = InEditor;
+	}
+
+	FMLDeformerScopedEditor::FMLDeformerScopedEditor(UMLDeformerAsset* Asset)
+	{
+		Editor = FMLDeformerTestHelpers::OpenAssetEditor(Asset);		
+	}
+
+	FMLDeformerScopedEditor::FMLDeformerScopedEditor(const FString& AssetName)
+	{
+		Editor = FMLDeformerTestHelpers::OpenAssetEditor(AssetName);
+	}
+
+	FMLDeformerScopedEditor::~FMLDeformerScopedEditor()
+	{
+		if (Editor && bCloseEditor)
+		{
+			Editor->CloseWindow(EAssetEditorCloseReason::AssetEditorHostClosed);
+		}
+	}
+
+	FMLDeformerEditorToolkit* FMLDeformerScopedEditor::Get() const
+	{
+		return Editor;
+	}
+
+	void FMLDeformerScopedEditor::SetCloseEditor(bool bCloseOnDestroy)
+	{
+		bCloseEditor = bCloseOnDestroy;
+	}
+
+	bool FMLDeformerScopedEditor::IsValid() const
+	{
+		return (Editor != nullptr);
+	}
+
+	FMLDeformerEditorToolkit* FMLDeformerScopedEditor::operator->() const
+	{ 
+		return Editor;
+	}
+
+	//--------------------------------------------------------------------------------------
+	// FMLDeformerTestHelpers
+	//--------------------------------------------------------------------------------------
+	FMLDeformerEditorToolkit* FMLDeformerTestHelpers::OpenAssetEditor(UMLDeformerAsset* Asset)
+	{
+		if (Asset == nullptr)
+		{
+			UE_LOG(LogMLDeformer, Error, TEXT("Asset is a nullptr."));
+			return nullptr;
+		}
+
+		if (!GEditor->GetEditorSubsystem<UAssetEditorSubsystem>()->OpenEditorForAsset(Asset))
+		{
+			UE_LOG(LogMLDeformer, Error, TEXT("Failed to open the asset editor for ML Deformer asset '%s'"), *Asset->GetName());
+			return nullptr;
+		}
+
+		// This checks to see if the asset sub editor is loaded.
+		FMLDeformerEditorToolkit* Editor = static_cast<FMLDeformerEditorToolkit*>(GEditor->GetEditorSubsystem<UAssetEditorSubsystem>()->FindEditorForAsset(Asset, true));
+		if (Editor == nullptr)
+		{
+			UE_LOG(LogMLDeformer, Error, TEXT("Failed to open asset editor for ML Deformer asset %s, as we couldn't find the editor."), *Asset->GetName());
+			return nullptr;
+		}
+
+		return Editor;
+	}
+
+	FMLDeformerEditorToolkit* FMLDeformerTestHelpers::OpenAssetEditor(const FString& AssetName)
+	{
+		// AssetName is something like "/MLDeformerFramework/Tests/MLD_Rampage.MLD_Rampage"
+
+		// Try to load the object from the package.
+		UObject* Object = StaticLoadObject(UMLDeformerAsset::StaticClass(), NULL, *AssetName);
+		UMLDeformerAsset* Asset = Cast<UMLDeformerAsset>(Object);
+		if (Asset == nullptr)
+		{
+			UE_LOG(LogMLDeformer, Error, TEXT("Failed to load ML Deformer asset from '%s'"), *AssetName);
+			return nullptr;
+		}
+
+		return OpenAssetEditor(Asset);
+	}
+
+	bool FMLDeformerTestHelpers::PostModelPropertyChanged(UMLDeformerModel* Model, const FName PropertyName)
+	{
+		FProperty* Property = Model->GetClass()->FindPropertyByName(PropertyName);
+		if (Property == nullptr)
+		{
+			UE_LOG(LogMLDeformer, Error, TEXT("Failed to find property '%s' in class '%s'"), *PropertyName.ToString(), *Model->GetName());
+			return false;
+		}
+
+		FPropertyChangedEvent Event(Property, EPropertyChangeType::ValueSet);
+		Model->PostEditChangeProperty(Event);
+		return true;
+	}
+
+}	// namespace UE::MLDeformer
diff --git a/Engine/Plugins/Animation/MLDeformer/MLDeformerFramework/Source/MLDeformerFrameworkEditor/Private/SMLDeformerTimeline.cpp b/Engine/Plugins/Animation/MLDeformer/MLDeformerFramework/Source/MLDeformerFrameworkEditor/Private/SMLDeformerTimeline.cpp
--- a/Engine/Plugins/Animation/MLDeformer/MLDeformerFramework/Source/MLDeformerFrameworkEditor/Private/SMLDeformerTimeline.cpp
+++ b/Engine/Plugins/Animation/MLDeformer/MLDeformerFramework/Source/MLDeformerFrameworkEditor/Private/SMLDeformerTimeline.cpp
@@ -1346,8 +1346,8 @@ namespace UE::MLDeformer
 		Model = WeakModel;
 		OnReceivedFocus = InArgs._OnReceivedFocus;
 
-		const int32 TickResolutionValue = Model.Pin()->GetTickResolution();
-		const int32 SequenceFrameRate = FMath::RoundToInt(Model.Pin()->GetFrameRate());
+		const int32 TickResolutionValue = Model.IsValid() ? Model.Pin()->GetTickResolution() : 1000 * 30;	// 1000*30 is ticks per frame multiplied by 30 fps.
+		const int32 SequenceFrameRate = FMath::RoundToInt(Model.IsValid() ? Model.Pin()->GetFrameRate() : 30);
 
 		ViewRange = MakeAttributeLambda([this]()
 		{
@@ -1407,7 +1407,12 @@ namespace UE::MLDeformer
 			TimeSliderArgs.DisplayRate = DisplayRate;
 			TimeSliderArgs.TickResolution = TickResolution;
 			TimeSliderArgs.OnViewRangeChanged = FOnViewRangeChanged::CreateSP(this, &SMLDeformerTimeline::HandleViewRangeChanged);
-			TimeSliderArgs.OnClampRangeChanged = FOnTimeRangeChanged::CreateSP(Model.Pin().Get(), &FMLDeformerEditorModel::HandleWorkingRangeChanged);
+
+			if (Model.IsValid())
+			{
+				TimeSliderArgs.OnClampRangeChanged = FOnTimeRangeChanged::CreateSP(Model.Pin().Get(), &FMLDeformerEditorModel::HandleWorkingRangeChanged);
+			}
+
 			TimeSliderArgs.IsPlaybackRangeLocked = true;
 			TimeSliderArgs.PlaybackStatus = EMovieScenePlayerStatus::Stopped;
 			TimeSliderArgs.NumericTypeInterface = NumericTypeInterface;
diff --git a/Engine/Plugins/Animation/MLDeformer/MLDeformerFramework/Source/MLDeformerFrameworkEditor/Public/MLDeformerEditorModule.h b/Engine/Plugins/Animation/MLDeformer/MLDeformerFramework/Source/MLDeformerFrameworkEditor/Public/MLDeformerEditorModule.h
--- a/Engine/Plugins/Animation/MLDeformer/MLDeformerFramework/Source/MLDeformerFrameworkEditor/Public/MLDeformerEditorModule.h
+++ b/Engine/Plugins/Animation/MLDeformer/MLDeformerFramework/Source/MLDeformerFrameworkEditor/Public/MLDeformerEditorModule.h
@@ -27,8 +27,15 @@ namespace UE::MLDeformer
 		 * You use this to see what kind of model types ther are, register new models, or create editor models for specific types of runtime models.
 		 * @return A reference to the ML Deformer model registry.
 		 */
-		FMLDeformerEditorModelRegistry& GetModelRegistry()				{ return ModelRegistry; }
-		const FMLDeformerEditorModelRegistry& GetModelRegistry() const	{ return ModelRegistry; }
+		FMLDeformerEditorModelRegistry& GetModelRegistry()						{ return ModelRegistry; }
+		const FMLDeformerEditorModelRegistry& GetModelRegistry() const			{ return ModelRegistry; }
+
+		/**
+		 * Manually assign a new model registry.
+		 * This can be used in some automated tests where we want to test with different registries, for example an empty one.
+		 * @param Registry The model registry to use.
+		 */
+		void SetModelRegistry(const FMLDeformerEditorModelRegistry& Registry)	{ ModelRegistry = Registry; }
 
 	private:
 		/** The model registry, which keeps track of all model types and instances created of these models, and how to create the editor models for specific runtime model types. */
diff --git a/Engine/Plugins/Animation/MLDeformer/MLDeformerFramework/Source/MLDeformerFrameworkEditor/Public/MLDeformerEditorToolkit.h b/Engine/Plugins/Animation/MLDeformer/MLDeformerFramework/Source/MLDeformerFrameworkEditor/Public/MLDeformerEditorToolkit.h
--- a/Engine/Plugins/Animation/MLDeformer/MLDeformerFramework/Source/MLDeformerFrameworkEditor/Public/MLDeformerEditorToolkit.h
+++ b/Engine/Plugins/Animation/MLDeformer/MLDeformerFramework/Source/MLDeformerFrameworkEditor/Public/MLDeformerEditorToolkit.h
@@ -43,6 +43,9 @@ namespace UE::MLDeformer
 		friend class FMLDeformerApplicationMode;
 		friend struct FMLDeformerVizSettingsTabSummoner;
 
+
+		~FMLDeformerEditorToolkit();
+
 		/** Initialize the asset editor. This will register the application mode, init the preview scene, etc. */
 		void InitAssetEditor(
 			const EToolkitMode::Type Mode,
@@ -60,45 +63,68 @@ namespace UE::MLDeformer
 		// ~END FAssetEditorToolkit overrides.
 
 		// FGCObject overrides.
-		virtual FString GetReferencerName() const override { return TEXT("FMLDeformerEditorToolkit"); }
+		virtual FString GetReferencerName() const override							{ return TEXT("FMLDeformerEditorToolkit"); }
 		virtual void AddReferencedObjects(FReferenceCollector& Collector) override;
 		// ~END FGCObject overrides.
 
 		// FTickableEditorObject overrides.
 		virtual void Tick(float DeltaTime) override {};
-		virtual ETickableTickType GetTickableTickType() const override { return ETickableTickType::Always; }
+		virtual ETickableTickType GetTickableTickType() const override				{ return ETickableTickType::Always; }
 		virtual TStatId GetStatId() const override;
 		// ~END FTickableEditorObject overrides.
 
 		// IHasPersonaToolkit overrides.
 		virtual TSharedRef<IPersonaToolkit> GetPersonaToolkit() const override;
-		IPersonaToolkit* GetPersonaToolkitPointer() const { return PersonaToolkit.Get(); }
+		IPersonaToolkit* GetPersonaToolkitPointer() const							{ return PersonaToolkit.Get(); }
 		// ~END IHasPersonaToolkit overrides.
 
-		void SetVizSettingsDetailsView(TSharedPtr<IDetailsView> InDetailsView) { VizSettingsDetailsView = InDetailsView; }
+		void SetVizSettingsDetailsView(TSharedPtr<IDetailsView> InDetailsView)		{ VizSettingsDetailsView = InDetailsView; }
+		IDetailsView* GetVizSettingsDetailsView() const;
 
 		IDetailsView* GetModelDetailsView() const;
-		IDetailsView* GetVizSettingsDetailsView() const;
 
 		void SetTimeSlider(TSharedPtr<SMLDeformerTimeline> InTimeSlider);
 		SMLDeformerTimeline* GetTimeSlider() const;
 
-		UMLDeformerAsset* GetDeformerAsset() const { return DeformerAsset.Get(); }
-		FMLDeformerEditorModel* GetActiveModel() { return ActiveModel.Get(); }
-
-		TWeakPtr<FMLDeformerEditorModel> GetActiveModelPointer(){ return TWeakPtr<FMLDeformerEditorModel>(ActiveModel); }
-		const FMLDeformerEditorModel* GetActiveModel() const { return ActiveModel.Get(); }
-
-		FMLDeformerApplicationMode* GetApplicationMode() const { return ApplicationMode; }
+		UMLDeformerAsset* GetDeformerAsset() const									{ return DeformerAsset.Get(); }
+		FMLDeformerEditorModel* GetActiveModel()									{ return ActiveModel.Get(); }
+		TWeakPtr<FMLDeformerEditorModel> GetActiveModelPointer()					{ return TWeakPtr<FMLDeformerEditorModel>(ActiveModel); }
+		const FMLDeformerEditorModel* GetActiveModel() const						{ return ActiveModel.Get(); }
+		FMLDeformerApplicationMode* GetApplicationMode() const						{ return ApplicationMode; }
+		TSharedPtr<IPersonaViewport> GetViewport() const							{ return PersonaViewport; }
 
 		double CalcTimelinePosition() const;
 		void OnTimeSliderScrubPositionChanged(double NewScrubTime, bool bIsScrubbing);
 		void UpdateTimeSliderRange();
 		void SetTimeSliderRange(double StartTime, double EndTime);
 
-		TSharedPtr<IPersonaViewport> GetViewport() const { return PersonaViewport; }
+		/**
+		 * Switch the editor to a given model type.
+		 * @param ModelType The model type you want to switch to, for example something like: UNeuralMorphModel::StaticClass().
+		 * @param bForceChange Force changing to this model? This will suppress any UI popups.
+		 * @return Returns true in case we successfully switched model types, or otherwise false is returned.
+		 */
+		bool SwitchModelType(UClass* ModelType, bool bForceChange);
+
+		/**
+		 * Switch the editor's visualization mode.
+		 * This essentially allows you to switch the UI between testing and training modes.
+		 * @param Mode The mode to switch to.
+		 */
+		void SwitchVizMode(EMLDeformerVizMode Mode);
+
+		bool Train(bool bSuppressDialogs);
+		bool IsTrainButtonEnabled() const;
+		bool IsTraining() const;
 
 	private:
+		UE_DEPRECATED(5.3, "Please use the OnModelChanged that takes two parameters instead.")
+		void OnModelChanged(int Index);
+		void OnModelChanged(int Index, bool bForceChange);
+
+		UE_DEPRECATED(5.3, "Please use the SwitchVizMode instead.")
+		void OnVizModeChanged(EMLDeformerVizMode Mode);
+
 		/* Toolbar related. */
 		void ExtendToolbar();
 		void FillToolbar(FToolBarBuilder& ToolbarBuilder);
@@ -113,9 +139,7 @@ namespace UE::MLDeformer
 		void ShowNotification(const FText& Message, SNotificationItem::ECompletionState State, bool PlaySound) const;
 		FText GetOverlayText() const;
 		void OnSwitchedVisualizationMode();
-		bool HandleTrainingResult(ETrainingResult TrainingResult, double TrainingDuration, bool& bOutUsePartiallyTrained);
-		void OnModelChanged(int Index);
-		void OnVizModeChanged(EMLDeformerVizMode Mode);
+		bool HandleTrainingResult(ETrainingResult TrainingResult, double TrainingDuration, bool& bOutUsePartiallyTrained, bool bSuppressDialogs, bool& bOutSuccess);
 		FText GetActiveModelName() const;
 		FText GetCurrentVizModeName() const;
 		FText GetVizModeName(EMLDeformerVizMode Mode) const;
@@ -151,5 +175,8 @@ namespace UE::MLDeformer
 
 		/** Has the asset editor been initialized? */
 		bool bIsInitialized = false;
+
+		/** Are we currently in a training process? */
+		bool bIsTraining = false;
 	};
 }	// namespace UE::MLDeformer
diff --git a/Engine/Plugins/Animation/MLDeformer/MLDeformerFramework/Source/MLDeformerFrameworkEditor/Public/MLDeformerTestHelpers.h b/Engine/Plugins/Animation/MLDeformer/MLDeformerFramework/Source/MLDeformerFrameworkEditor/Public/MLDeformerTestHelpers.h
new file mode 100644
--- /dev/null
+++ b/Engine/Plugins/Animation/MLDeformer/MLDeformerFramework/Source/MLDeformerFrameworkEditor/Public/MLDeformerTestHelpers.h
@@ -0,0 +1,68 @@
+// Copyright Epic Games, Inc. All Rights Reserved.
+#pragma once
+
+#include "CoreTypes.h"
+#include "Containers/UnrealString.h"
+#include "Misc/AutomationTest.h"
+
+class UMLDeformerAsset;
+class UMLDeformerModel;
+
+namespace UE::MLDeformer
+{
+	class FMLDeformerEditorToolkit;
+	class FMLDeformerEditorModelRegistry;
+
+	/** A scoped ML Deformer editor, which auto closes the window on destruction. */
+	class MLDEFORMERFRAMEWORKEDITOR_API FMLDeformerScopedEditor
+	{
+	public:
+		FMLDeformerScopedEditor() = delete;
+		FMLDeformerScopedEditor(const FMLDeformerScopedEditor&) = delete;
+		FMLDeformerScopedEditor(FMLDeformerScopedEditor&&) = delete;
+		FMLDeformerScopedEditor(FMLDeformerEditorToolkit* InEditor);
+		FMLDeformerScopedEditor(UMLDeformerAsset* Asset);
+		FMLDeformerScopedEditor(const FString& AssetName);
+		~FMLDeformerScopedEditor();	// Auto close the editor window.
+		bool IsValid() const;
+		void SetCloseEditor(bool bCloseOnDestroy);
+		FMLDeformerEditorToolkit* Get() const;
+		FMLDeformerEditorToolkit* operator = (const FMLDeformerScopedEditor&) = delete;
+		FMLDeformerEditorToolkit* operator = (const FMLDeformerScopedEditor&&) = delete;
+		FMLDeformerEditorToolkit* operator ->() const;
+
+	private:
+		FMLDeformerEditorToolkit* Editor = nullptr;
+		bool bCloseEditor = true;
+	};
+
+	/** Some helper functions for testing ML Deformer related things. */
+	class MLDEFORMERFRAMEWORKEDITOR_API FMLDeformerTestHelpers
+	{
+	public:
+		/**
+		 * Open the asset editor for a given ML Deformer asset.
+		 * This will log an error when something doesn't succeed.
+		 * @param Asset A pointer to the asset to open an editor for.
+		 * @return A pointer to the editor toolkit, or nullptr when it didn't succeed.
+		 */
+		static FMLDeformerEditorToolkit* OpenAssetEditor(UMLDeformerAsset* Asset);
+
+		/**
+		 * Open the asset editor for a given ML Deformer asset, by its package name.
+		 * This will log an error when something doesn't succeed.
+		 * Examples of failure cases are when the package cannot be loaded, or if it doesn't contain a valid ML Deformer asset.
+		 * @param Asset The package name, for example something like: "/MLDeformerFramework/Tests/MLD_Rampage.MLD_Rampage".
+		 * @return A pointer to the editor toolkit, or nullptr when it didn't succeed.
+		 */
+		static FMLDeformerEditorToolkit* OpenAssetEditor(const FString& AssetName);
+
+		/**
+		 * Post an edit change event for a given property in the Model class.
+		 * @param Model The model that the property has been changed in.
+		 * @param PropertyName The name of the property (the class member name).
+		 * @return Returns true if successfully emitted the PostEditChangeProperty call, or false if the property cannot be found.
+		 */
+		static bool PostModelPropertyChanged(UMLDeformerModel* Model, const FName PropertyName);
+	};
+}	// namespace UE::MLDeformer
diff --git a/Engine/Source/Editor/Persona/Private/AnimationEditorPreviewScene.cpp b/Engine/Source/Editor/Persona/Private/AnimationEditorPreviewScene.cpp
--- a/Engine/Source/Editor/Persona/Private/AnimationEditorPreviewScene.cpp
+++ b/Engine/Source/Editor/Persona/Private/AnimationEditorPreviewScene.cpp
@@ -138,7 +138,7 @@ TArray<UDebugSkelMeshComponent*> FAnimationEditorPreviewScene::GetAllPreviewMesh
 {
 	TArray<UDebugSkelMeshComponent*> PreviewMeshComponents;
 	const AActor* MyActor = GetActor();
-	if (ensure(MyActor))
+	if (MyActor)
 	{
 		MyActor->GetComponents(PreviewMeshComponents, true);
 	}
@@ -149,7 +149,7 @@ void FAnimationEditorPreviewScene::ForEachPreviewMesh(TFunction<void (UDebugSkel
 {
 	TArray<UDebugSkelMeshComponent*> PreviewMeshComponents;
 	const AActor* MyActor = GetActor();
-	if (ensure(MyActor))
+	if (MyActor)
 	{
 		MyActor->GetComponents(PreviewMeshComponents, true);
 		for (UDebugSkelMeshComponent* PreviewMesh : PreviewMeshComponents)
