Fix for missing module in ML Deformer

diff --git a/Engine/Plugins/Animation/MLDeformer/MLDeformerFramework/Source/MLDeformerFramework/MLDeformerFramework.Build.cs b/Engine/Plugins/Animation/MLDeformer/MLDeformerFramework/Source/MLDeformerFramework/MLDeformerFramework.Build.cs
--- a/Engine/Plugins/Animation/MLDeformer/MLDeformerFramework/Source/MLDeformerFramework/MLDeformerFramework.Build.cs
+++ b/Engine/Plugins/Animation/MLDeformer/MLDeformerFramework/Source/MLDeformerFramework/MLDeformerFramework.Build.cs
@@ -40,6 +40,7 @@ public class MLDeformerFramework : ModuleRules
 				"RHI",
 				"GeometryCache",
 				"OptimusCore",
+				"RenderCore",
 				"ComputeFramework",
 				"Projects"
 			}
