[MLDeformer] Fixed rendering of bones inside the MLD asset editor.

diff --git a/Engine/Plugins/Animation/MLDeformer/MLDeformerFramework/Source/MLDeformerFrameworkEditor/Private/MLDeformerApplicationMode.cpp b/Engine/Plugins/Animation/MLDeformer/MLDeformerFramework/Source/MLDeformerFrameworkEditor/Private/MLDeformerApplicationMode.cpp
--- a/Engine/Plugins/Animation/MLDeformer/MLDeformerFramework/Source/MLDeformerFrameworkEditor/Private/MLDeformerApplicationMode.cpp
+++ b/Engine/Plugins/Animation/MLDeformer/MLDeformerFramework/Source/MLDeformerFrameworkEditor/Private/MLDeformerApplicationMode.cpp
@@ -42,7 +42,7 @@ namespace UE::MLDeformer
 		TabFactories.RegisterFactory(MakeShared<FMLDeformerTimelineTabSummoner>(MLDeformerEditor));
 
 		// Create tab layout.
-		TabLayout = FTabManager::NewLayout("Standalone_MLDeformerEditor_Layout_v0.032")
+		TabLayout = FTabManager::NewLayout("Standalone_MLDeformerEditor_Layout_v5.3.001")
 			->AddArea
 			(
 				FTabManager::NewPrimaryArea()
diff --git a/Engine/Plugins/Animation/MLDeformer/MLDeformerFramework/Source/MLDeformerFrameworkEditor/Private/MLDeformerEditorModel.cpp b/Engine/Plugins/Animation/MLDeformer/MLDeformerFramework/Source/MLDeformerFrameworkEditor/Private/MLDeformerEditorModel.cpp
--- a/Engine/Plugins/Animation/MLDeformer/MLDeformerFramework/Source/MLDeformerFrameworkEditor/Private/MLDeformerEditorModel.cpp
+++ b/Engine/Plugins/Animation/MLDeformer/MLDeformerFramework/Source/MLDeformerFrameworkEditor/Private/MLDeformerEditorModel.cpp
@@ -1961,6 +1961,43 @@ namespace UE::MLDeformer
 	void FMLDeformerEditorModel::HandleVizModeChanged(EMLDeformerVizMode Mode)
 	{
 		UpdateRanges();
+
+		// Update the preview scene so that the debug rendering draws the skeleton for the correct actor, depending on whether we are in training or testing mode.
+		// On default we render the bones etc for the ML Deformed character.
+		// Additionally trigger the input assets changed event, as that updates the animation play offsets etc. This is needed when we switch from say training back to testing mode.
+		UpdatePreviewScene();
+		TriggerInputAssetChanged();
+	}
+
+	void FMLDeformerEditorModel::UpdatePreviewScene()
+	{
+		FMLDeformerEditorActor* EditorActor = nullptr;
+		UAnimationAsset* AnimAsset = nullptr;
+		IPersonaPreviewScene& Scene = GetEditor()->GetPersonaToolkit()->GetPreviewScene().Get();
+		const EMLDeformerVizMode Mode = Model->GetVizSettings()->GetVisualizationMode();
+		if (Mode == EMLDeformerVizMode::TrainingData)
+		{
+			EditorActor = FindEditorActor(ActorID_Train_Base);
+			AnimAsset = Model->GetAnimSequence();
+		}
+		else if (Mode == EMLDeformerVizMode::TestData)
+		{
+			EditorActor = FindEditorActor(ActorID_Test_MLDeformed);
+			AnimAsset = Model->GetVizSettings()->GetTestAnimSequence();
+		}
+
+		if (EditorActor)
+		{
+			Scene.SetActor(EditorActor->GetActor());
+			Scene.SetSelectedActor(EditorActor->GetActor());
+			UDebugSkelMeshComponent* SkelMeshComponent = EditorActor->GetSkeletalMeshComponent();
+			if (SkelMeshComponent)
+			{
+				Scene.SetPreviewMeshComponent(SkelMeshComponent);
+				Scene.SetPreviewMesh(SkelMeshComponent->GetSkeletalMeshAsset());
+			}
+			Scene.SetPreviewAnimationAsset(AnimAsset);
+		}
 	}
 
 	/** Get the current view range */
diff --git a/Engine/Plugins/Animation/MLDeformer/MLDeformerFramework/Source/MLDeformerFrameworkEditor/Public/MLDeformerEditorModel.h b/Engine/Plugins/Animation/MLDeformer/MLDeformerFramework/Source/MLDeformerFrameworkEditor/Public/MLDeformerEditorModel.h
--- a/Engine/Plugins/Animation/MLDeformer/MLDeformerFramework/Source/MLDeformerFrameworkEditor/Public/MLDeformerEditorModel.h
+++ b/Engine/Plugins/Animation/MLDeformer/MLDeformerFramework/Source/MLDeformerFrameworkEditor/Public/MLDeformerEditorModel.h
@@ -579,6 +579,9 @@ namespace UE::MLDeformer
 		 */
 		virtual TSharedPtr<SMLDeformerInputWidget> CreateInputWidget();
 
+		/** Update the persona preview scene's actor, skeletal mesh component, mesh and anim asset. These are based on the current visualization mode. */
+		virtual void UpdatePreviewScene();
+
 		/**
 		 * Generate the normals for a given morph target.
 		 * @param LOD The LOD level.
