[MLDeformer] Fixed some crash when creating a new asset.

#rnx

diff --git a/Engine/Plugins/Animation/MLDeformer/MLDeformerFramework/Source/MLDeformerFrameworkEditor/Private/MLDeformerEditorModel.cpp b/Engine/Plugins/Animation/MLDeformer/MLDeformerFramework/Source/MLDeformerFrameworkEditor/Private/MLDeformerEditorModel.cpp
--- a/Engine/Plugins/Animation/MLDeformer/MLDeformerFramework/Source/MLDeformerFrameworkEditor/Private/MLDeformerEditorModel.cpp
+++ b/Engine/Plugins/Animation/MLDeformer/MLDeformerFramework/Source/MLDeformerFrameworkEditor/Private/MLDeformerEditorModel.cpp
@@ -521,12 +521,14 @@ namespace UE::MLDeformer
 		}
 
 		// Update the training base actor.
-		UDebugSkelMeshComponent* SkeletalMeshComponent = FindEditorActor(ActorID_Train_Base)->GetSkeletalMeshComponent();
 		UMLDeformerVizSettings* VizSettings = Model->GetVizSettings();
 		UAnimSequence* TestAnimSequence = VizSettings->GetTestAnimSequence();
 		const float TestAnimSpeed = VizSettings->GetAnimPlaySpeed();
+
+		const FMLDeformerEditorActor* TrainBaseActor = FindEditorActor(ActorID_Train_Base);
+		UDebugSkelMeshComponent* SkeletalMeshComponent = TrainBaseActor ? TrainBaseActor->GetSkeletalMeshComponent() : nullptr;
+		if (SkeletalMeshComponent)
 		{
-			check(SkeletalMeshComponent);
 			SkeletalMeshComponent->SetSkeletalMesh(Model->GetSkeletalMesh());
 			if (GetEditor()->GetPersonaToolkitPointer())
 			{
@@ -545,7 +547,8 @@ namespace UE::MLDeformer
 		}
 
 		// Update the test base model.
-		SkeletalMeshComponent = FindEditorActor(ActorID_Test_Base)->GetSkeletalMeshComponent();
+		const FMLDeformerEditorActor* TestBaseActor = FindEditorActor(ActorID_Test_Base);
+		SkeletalMeshComponent = TestBaseActor ? TestBaseActor->GetSkeletalMeshComponent() : nullptr;
 		if (SkeletalMeshComponent)
 		{
 			SkeletalMeshComponent->SetSkeletalMesh(Model->GetSkeletalMesh());
@@ -558,7 +561,8 @@ namespace UE::MLDeformer
 		}
 
 		// Update the test ML Deformed skeletal mesh component.
-		SkeletalMeshComponent = FindEditorActor(ActorID_Test_MLDeformed)->GetSkeletalMeshComponent();
+		const FMLDeformerEditorActor* TestMLActor = FindEditorActor(ActorID_Test_MLDeformed);
+		SkeletalMeshComponent = TestMLActor ? TestMLActor->GetSkeletalMeshComponent() : nullptr;
 		if (SkeletalMeshComponent)
 		{
 			SkeletalMeshComponent->SetSkeletalMesh(Model->GetSkeletalMesh());
@@ -1973,7 +1977,6 @@ namespace UE::MLDeformer
 	{
 		FMLDeformerEditorActor* EditorActor = nullptr;
 		UAnimationAsset* AnimAsset = nullptr;
-		IPersonaPreviewScene& Scene = GetEditor()->GetPersonaToolkit()->GetPreviewScene().Get();
 		const EMLDeformerVizMode Mode = Model->GetVizSettings()->GetVisualizationMode();
 		if (Mode == EMLDeformerVizMode::TrainingData)
 		{
@@ -1986,8 +1989,10 @@ namespace UE::MLDeformer
 			AnimAsset = Model->GetVizSettings()->GetTestAnimSequence();
 		}
 
-		if (EditorActor)
+		IPersonaToolkit* ToolkitPtr = GetEditor()->GetPersonaToolkitPointer();
+		if (EditorActor && ToolkitPtr)
 		{
+			IPersonaPreviewScene& Scene = ToolkitPtr->GetPreviewScene().Get();
 			Scene.SetActor(EditorActor->GetActor());
 			Scene.SetSelectedActor(EditorActor->GetActor());
 			UDebugSkelMeshComponent* SkelMeshComponent = EditorActor->GetSkeletalMeshComponent();
diff --git a/Engine/Plugins/Animation/MLDeformer/MLDeformerFramework/Source/MLDeformerFrameworkEditor/Private/MLDeformerGeomCacheEditorModel.cpp b/Engine/Plugins/Animation/MLDeformer/MLDeformerFramework/Source/MLDeformerFrameworkEditor/Private/MLDeformerGeomCacheEditorModel.cpp
--- a/Engine/Plugins/Animation/MLDeformer/MLDeformerFramework/Source/MLDeformerFrameworkEditor/Private/MLDeformerGeomCacheEditorModel.cpp
+++ b/Engine/Plugins/Animation/MLDeformer/MLDeformerFramework/Source/MLDeformerFrameworkEditor/Private/MLDeformerGeomCacheEditorModel.cpp
@@ -230,22 +230,28 @@ namespace UE::MLDeformer
 		UAnimSequence* TestAnimSequence = VizSettings->GetTestAnimSequence();
 
 		// Update the training geometry cache.
-		UGeometryCacheComponent* GeometryCacheComponent = FindGeomCacheEditorActor(ActorID_Train_GroundTruth)->GetGeometryCacheComponent();
-		check(GeometryCacheComponent);
-		GeometryCacheComponent->SetGeometryCache(GetGeomCacheModel()->GetGeometryCache());
-		GeometryCacheComponent->SetLooping(false);
-		GeometryCacheComponent->SetManualTick(true);
-		GeometryCacheComponent->SetPlaybackSpeed(TestAnimSpeed);
-		GeometryCacheComponent->Play();
+		const FMLDeformerGeomCacheActor* TrainGroundTruthActor = FindGeomCacheEditorActor(ActorID_Train_GroundTruth);
+		UGeometryCacheComponent* GeometryCacheComponent = TrainGroundTruthActor ? TrainGroundTruthActor->GetGeometryCacheComponent() : nullptr;
+		if (GeometryCacheComponent)
+		{
+			GeometryCacheComponent->SetGeometryCache(GetGeomCacheModel()->GetGeometryCache());
+			GeometryCacheComponent->SetLooping(false);
+			GeometryCacheComponent->SetManualTick(true);
+			GeometryCacheComponent->SetPlaybackSpeed(TestAnimSpeed);
+			GeometryCacheComponent->Play();
+		}
 
 		// Update the test geometry cache (ground truth) component.
-		GeometryCacheComponent = FindGeomCacheEditorActor(ActorID_Test_GroundTruth)->GetGeometryCacheComponent();
-		check(GeometryCacheComponent);
-		GeometryCacheComponent->SetGeometryCache(VizSettings->GetTestGroundTruth());
-		GeometryCacheComponent->SetLooping(true);
-		GeometryCacheComponent->SetManualTick(true);
-		GeometryCacheComponent->SetPlaybackSpeed(TestAnimSpeed);
-		GeometryCacheComponent->Play();
+		const FMLDeformerGeomCacheActor* TestGroundTruthActor = FindGeomCacheEditorActor(ActorID_Test_GroundTruth);
+		GeometryCacheComponent = TestGroundTruthActor ? TestGroundTruthActor->GetGeometryCacheComponent() : nullptr;
+		if (GeometryCacheComponent)
+		{
+			GeometryCacheComponent->SetGeometryCache(VizSettings->GetTestGroundTruth());
+			GeometryCacheComponent->SetLooping(true);
+			GeometryCacheComponent->SetManualTick(true);
+			GeometryCacheComponent->SetPlaybackSpeed(TestAnimSpeed);
+			GeometryCacheComponent->Play();
+		}
 
 		GetGeomCacheModel()->GetGeomCacheMeshMappings().Reset();
 	}
