[MLDeformer] Show the frame counts in geom cache vs anim sequence duration mismatch warnings.

diff --git a/Engine/Plugins/Animation/MLDeformer/MLDeformerFramework/Source/MLDeformerFramework/Private/MLDeformerGeomCacheHelpers.cpp b/Engine/Plugins/Animation/MLDeformer/MLDeformerFramework/Source/MLDeformerFramework/Private/MLDeformerGeomCacheHelpers.cpp
--- a/Engine/Plugins/Animation/MLDeformer/MLDeformerFramework/Source/MLDeformerFramework/Private/MLDeformerGeomCacheHelpers.cpp
+++ b/Engine/Plugins/Animation/MLDeformer/MLDeformerFramework/Source/MLDeformerFramework/Private/MLDeformerGeomCacheHelpers.cpp
@@ -337,13 +337,18 @@ namespace UE::MLDeformer
 			const float GeomCacheDuration = InGeomCache->CalculateDuration();
 			if (FMath::Abs(AnimSeqDuration - GeomCacheDuration) > 0.001f)
 			{
+				const int32 NumGeomCacheFrames = InGeomCache ? (InGeomCache->GetEndFrame() - InGeomCache->GetStartFrame()) + 1 : 0;
+				const int32 NumAnimSeqFrames = InAnimSequence ? InAnimSequence->GetNumberOfSampledKeys() : 0;
+
 				FNumberFormattingOptions Options;
 				Options.SetUseGrouping(false);
 				Options.SetMaximumFractionalDigits(4);
 				Result = FText::Format(
-					LOCTEXT("AnimSeqNumFramesMismatch", "Anim sequence and Geometry Cache durations don't match!\n\nAnimSeq has a duration of {0} seconds, while GeomCache has a duration of {1} seconds.\n\nThis can produce incorrect results."),
+					LOCTEXT("AnimSeqNumFramesMismatch", "Anim sequence and Geometry Cache durations don't match!\n\nAnimSeq has a duration of {0} seconds ({1} frames}, while GeomCache has a duration of {2} seconds ({3} frames).\n\nThis can produce incorrect results."),
 					FText::AsNumber(AnimSeqDuration, &Options),
-					FText::AsNumber(GeomCacheDuration, &Options));
+					FText::AsNumber(NumAnimSeqFrames),
+					FText::AsNumber(GeomCacheDuration, &Options),
+					FText::AsNumber(NumGeomCacheFrames));
 			}
 		}
 
