[MLDeformer] Fix for incorrect mismatched vertices error in MLDeformer

PR #10554: Fix for incorrect mismatched vertices error in MLDeformer

diff --git a/Engine/Plugins/Animation/MLDeformer/MLDeformerFramework/Source/MLDeformerFramework/Private/MLDeformerGeomCacheHelpers.cpp b/Engine/Plugins/Animation/MLDeformer/MLDeformerFramework/Source/MLDeformerFramework/Private/MLDeformerGeomCacheHelpers.cpp
--- a/Engine/Plugins/Animation/MLDeformer/MLDeformerFramework/Source/MLDeformerFramework/Private/MLDeformerGeomCacheHelpers.cpp
+++ b/Engine/Plugins/Animation/MLDeformer/MLDeformerFramework/Source/MLDeformerFramework/Private/MLDeformerGeomCacheHelpers.cpp
@@ -146,6 +146,13 @@ namespace UE::MLDeformer
 			return;
 		}
 
+		// Add all the names to the OutVertexMisMatchNames ahead of time.  They'll be removed as matches are found.
+		for (int32 TrackIndex = 0; TrackIndex < GeomCache->Tracks.Num(); ++TrackIndex)
+		{
+			UGeometryCacheTrack* Track = GeomCache->Tracks[TrackIndex];
+			OutVertexMisMatchNames.Add(Track->GetName());
+		}
+
 		// For all meshes in the skeletal mesh.
 		const float SampleTime = 0.0f;
 		FString SkelMeshName;
@@ -190,10 +197,12 @@ namespace UE::MLDeformer
 					const int32 NumSkelMeshVerts = MeshInfo.NumVertices;
 					if (NumSkelMeshVerts != NumGeomMeshVerts)
 					{
-						OutVertexMisMatchNames.Add(Track->GetName());
 						continue;
 					}
 
+					// Since it is a match, remove it from the list.
+					OutVertexMisMatchNames.Remove(Track->GetName());
+
 					// Create a new mesh mapping entry.
 					OutMeshMappings.AddDefaulted();
 					UE::MLDeformer::FMLDeformerGeomCacheMeshMapping& Mapping = OutMeshMappings.Last();
