Fix Linux only crashes initializing the ML Deformer

diff --git a/Engine/Plugins/Animation/MLDeformer/MLDeformerFramework/Source/MLDeformerFrameworkEditor/Private/MLDeformerModelDetails.cpp b/Engine/Plugins/Animation/MLDeformer/MLDeformerFramework/Source/MLDeformerFrameworkEditor/Private/MLDeformerModelDetails.cpp
--- a/Engine/Plugins/Animation/MLDeformer/MLDeformerFramework/Source/MLDeformerFrameworkEditor/Private/MLDeformerModelDetails.cpp
+++ b/Engine/Plugins/Animation/MLDeformer/MLDeformerFramework/Source/MLDeformerFrameworkEditor/Private/MLDeformerModelDetails.cpp
@@ -136,7 +136,7 @@ namespace UE::MLDeformer
 			SNew(SObjectPropertyEntryBox)
 			.PropertyHandle(AnimRow.GetPropertyHandle())
 			.AllowedClass(UAnimSequence::StaticClass())
-			.ObjectPath(Model ? Model->GetAnimSequence()->GetPathName() : FString())
+			.ObjectPath( (Model && Model->GetAnimSequence()) ? Model->GetAnimSequence()->GetPathName() : FString())
 			.ThumbnailPool(DetailBuilder.GetThumbnailPool())
 			.OnShouldFilterAsset(
 				this, 
diff --git a/Engine/Plugins/Animation/MLDeformer/MLDeformerFramework/Source/MLDeformerFrameworkEditor/Private/MLDeformerVizSettingsDetails.cpp b/Engine/Plugins/Animation/MLDeformer/MLDeformerFramework/Source/MLDeformerFrameworkEditor/Private/MLDeformerVizSettingsDetails.cpp
--- a/Engine/Plugins/Animation/MLDeformer/MLDeformerFramework/Source/MLDeformerFrameworkEditor/Private/MLDeformerVizSettingsDetails.cpp
+++ b/Engine/Plugins/Animation/MLDeformer/MLDeformerFramework/Source/MLDeformerFrameworkEditor/Private/MLDeformerVizSettingsDetails.cpp
@@ -105,7 +105,7 @@ namespace UE::MLDeformer
 			SNew(SObjectPropertyEntryBox)
 			.PropertyHandle(TestAnimRow.GetPropertyHandle())
 			.AllowedClass(UAnimSequence::StaticClass())
-			.ObjectPath(VizSettings ? VizSettings->GetTestAnimSequence()->GetPathName() : FString())
+			.ObjectPath( (VizSettings && VizSettings->GetTestAnimSequence() )? VizSettings->GetTestAnimSequence()->GetPathName() : FString())
 			.ThumbnailPool(DetailBuilder.GetThumbnailPool())
 			.OnShouldFilterAsset(
 				this, 
