[MLDeformer] Fixed a bug that would cause a crash or would not detect selected input bones after you create a new MLD asset, select a skeletal mesh, and then pick input bones. It only worked when you had an animation sequence and target mesh selected first as well.

#rnx

diff --git a/Engine/Plugins/Animation/MLDeformer/MLDeformerFramework/Source/MLDeformerFrameworkEditor/Private/MLDeformerEditorModel.cpp b/Engine/Plugins/Animation/MLDeformer/MLDeformerFramework/Source/MLDeformerFrameworkEditor/Private/MLDeformerEditorModel.cpp
--- a/Engine/Plugins/Animation/MLDeformer/MLDeformerFramework/Source/MLDeformerFrameworkEditor/Private/MLDeformerEditorModel.cpp
+++ b/Engine/Plugins/Animation/MLDeformer/MLDeformerFramework/Source/MLDeformerFrameworkEditor/Private/MLDeformerEditorModel.cpp
@@ -1538,6 +1538,7 @@ namespace UE::MLDeformer
 			Model->GetSkeletalMesh() == nullptr ||
 			GetNumTrainingFrames() == 0)
 		{
+			UpdateEditorInputInfo();
 			return false;
 		}
 
diff --git a/Engine/Plugins/Animation/MLDeformer/MLDeformerFramework/Source/MLDeformerFrameworkEditor/Private/SMLDeformerInputBonesWidget.cpp b/Engine/Plugins/Animation/MLDeformer/MLDeformerFramework/Source/MLDeformerFrameworkEditor/Private/SMLDeformerInputBonesWidget.cpp
--- a/Engine/Plugins/Animation/MLDeformer/MLDeformerFramework/Source/MLDeformerFrameworkEditor/Private/SMLDeformerInputBonesWidget.cpp
+++ b/Engine/Plugins/Animation/MLDeformer/MLDeformerFramework/Source/MLDeformerFrameworkEditor/Private/SMLDeformerInputBonesWidget.cpp
@@ -357,7 +357,7 @@ namespace UE::MLDeformer
 				RefreshTree();
 
 				// Trigger the input widgets events.
-				// This is done AFTER the RefreshList call, because that updates the editor input info.
+				// This is done AFTER the RefreshTree call, because that updates the editor input info.
 				// Some handler code might depend on that to be updated first.
 				InputWidget->OnAddInputBones(BonesAdded);
 			}
