[MLDeformer] Fixed an issue where Metahuman face materials wouldn't update correctly when switching between two skeletal meshes, for example with different skin colors.

diff --git a/Engine/Plugins/Animation/MLDeformer/MLDeformerFramework/Source/MLDeformerFrameworkEditor/Private/MLDeformerEditorModel.cpp b/Engine/Plugins/Animation/MLDeformer/MLDeformerFramework/Source/MLDeformerFrameworkEditor/Private/MLDeformerEditorModel.cpp
--- a/Engine/Plugins/Animation/MLDeformer/MLDeformerFramework/Source/MLDeformerFrameworkEditor/Private/MLDeformerEditorModel.cpp
+++ b/Engine/Plugins/Animation/MLDeformer/MLDeformerFramework/Source/MLDeformerFrameworkEditor/Private/MLDeformerEditorModel.cpp
@@ -37,6 +37,7 @@
 #include "BoneContainer.h"
 #include "Misc/ScopedSlowTask.h"
 #include "Animation/AnimData/IAnimationDataModel.h"
+#include "Engine/SkinnedAssetCommon.h"
 #include "SMLDeformerInputWidget.h"
 
 #define LOCTEXT_NAMESPACE "MLDeformerEditorModel"
@@ -511,6 +512,28 @@ namespace UE::MLDeformer
 		}
 	}
 
+	void FMLDeformerEditorModel::ChangeSkeletalMeshOnComponent(USkeletalMeshComponent* SkelMeshComponent, USkeletalMesh* Mesh)
+	{
+		if (Mesh != SkelMeshComponent->GetSkeletalMeshAsset())
+		{
+			// Inform the skeletal mesh component about changes to the mesh.
+			SkelMeshComponent->SetSkeletalMesh(Mesh);
+			FPropertyChangedEvent SkelMeshChangedEvent(USkeletalMeshComponent::StaticClass()->FindPropertyByName(TEXT("SkeletalMeshAsset")));
+			SkelMeshComponent->PostEditChangeProperty(SkelMeshChangedEvent);
+
+			if (Mesh)
+			{
+				// Force all materials in the Skeletal Mesh Component to be the ones coming from the Mesh to avoid issues with material slots not updating.
+				const TArray<FSkeletalMaterial>& Materials = Mesh->GetMaterials();
+				for (int32 MaterialIndex = 0; MaterialIndex < Materials.Num(); ++MaterialIndex)
+				{
+					SkelMeshComponent->SetMaterial(MaterialIndex, Materials[MaterialIndex].MaterialInterface);
+				}
+				SkelMeshComponent->MarkRenderStateDirty();
+			}
+		}
+	}
+
 	void FMLDeformerEditorModel::OnInputAssetsChanged()
 	{
 		// Force the training sequence to use Step interpolation and sample raw animation data.
@@ -529,7 +552,7 @@ namespace UE::MLDeformer
 		UDebugSkelMeshComponent* SkeletalMeshComponent = TrainBaseActor ? TrainBaseActor->GetSkeletalMeshComponent() : nullptr;
 		if (SkeletalMeshComponent)
 		{
-			SkeletalMeshComponent->SetSkeletalMesh(Model->GetSkeletalMesh());
+			ChangeSkeletalMeshOnComponent(SkeletalMeshComponent, Model->GetSkeletalMesh());
 			if (GetEditor()->GetPersonaToolkitPointer())
 			{
 				GetEditor()->GetPersonaToolkit()->GetPreviewScene()->SetPreviewMesh(Model->GetSkeletalMesh());
@@ -551,7 +574,7 @@ namespace UE::MLDeformer
 		SkeletalMeshComponent = TestBaseActor ? TestBaseActor->GetSkeletalMeshComponent() : nullptr;
 		if (SkeletalMeshComponent)
 		{
-			SkeletalMeshComponent->SetSkeletalMesh(Model->GetSkeletalMesh());
+			ChangeSkeletalMeshOnComponent(SkeletalMeshComponent, Model->GetSkeletalMesh());
 			SkeletalMeshComponent->SetAnimationMode(EAnimationMode::AnimationSingleNode);
 			const float CurrentPlayTime = SkeletalMeshComponent->GetPosition();
 			SkeletalMeshComponent->SetAnimation(TestAnimSequence);
@@ -565,7 +588,7 @@ namespace UE::MLDeformer
 		SkeletalMeshComponent = TestMLActor ? TestMLActor->GetSkeletalMeshComponent() : nullptr;
 		if (SkeletalMeshComponent)
 		{
-			SkeletalMeshComponent->SetSkeletalMesh(Model->GetSkeletalMesh());
+			ChangeSkeletalMeshOnComponent(SkeletalMeshComponent, Model->GetSkeletalMesh());
 			SkeletalMeshComponent->SetAnimationMode(EAnimationMode::AnimationSingleNode);
 			const float CurrentPlayTime = SkeletalMeshComponent->GetPosition();
 			SkeletalMeshComponent->SetAnimation(TestAnimSequence);
diff --git a/Engine/Plugins/Animation/MLDeformer/MLDeformerFramework/Source/MLDeformerFrameworkEditor/Private/MLDeformerSampler.cpp b/Engine/Plugins/Animation/MLDeformer/MLDeformerFramework/Source/MLDeformerFrameworkEditor/Private/MLDeformerSampler.cpp
--- a/Engine/Plugins/Animation/MLDeformer/MLDeformerFramework/Source/MLDeformerFrameworkEditor/Private/MLDeformerSampler.cpp
+++ b/Engine/Plugins/Animation/MLDeformer/MLDeformerFramework/Source/MLDeformerFrameworkEditor/Private/MLDeformerSampler.cpp
@@ -252,7 +252,7 @@ namespace UE::MLDeformer
 			SkeletalMeshComponent->RegisterComponent();
 			SkelMeshActor->SetRootComponent(SkeletalMeshComponent);
 		}
-		SkeletalMeshComponent->SetSkeletalMesh(SkeletalMesh);
+		FMLDeformerEditorModel::ChangeSkeletalMeshOnComponent(SkeletalMeshComponent, SkeletalMesh);
 		SkeletalMeshComponent->SetAnimationMode(EAnimationMode::AnimationSingleNode);
 		SkeletalMeshComponent->SetAnimation(TrainingAnimSequence);
 		SkeletalMeshComponent->SetPosition(0.0f);
diff --git a/Engine/Plugins/Animation/MLDeformer/MLDeformerFramework/Source/MLDeformerFrameworkEditor/Public/MLDeformerEditorModel.h b/Engine/Plugins/Animation/MLDeformer/MLDeformerFramework/Source/MLDeformerFrameworkEditor/Public/MLDeformerEditorModel.h
--- a/Engine/Plugins/Animation/MLDeformer/MLDeformerFramework/Source/MLDeformerFrameworkEditor/Public/MLDeformerEditorModel.h
+++ b/Engine/Plugins/Animation/MLDeformer/MLDeformerFramework/Source/MLDeformerFrameworkEditor/Public/MLDeformerEditorModel.h
@@ -92,6 +92,19 @@ namespace UE::MLDeformer
 		virtual FString GetReferencerName() const override { return TEXT("FMLDeformerEditorModel"); }
 		// ~END FGCObject overrides.
 
+		/**
+		 * Change the skeletal mesh on a specific skeletal mesh component.
+		 * This internally will not only call SkelMeshComponent->SetSkeletalMesh(Mesh), but also will reassign the materials.
+		 * The reason for this is because if there are material instances set already, the skeletal mesh component will not override them to preserve changes done by the user.
+		 * This helper function will force these to be updated as well. This was needed in case of Metahumans.
+		 * Metahumans use a PostProcessAnimBP with a ControlRig node in it, which drives material curves.
+		 * In order to drive this, material instances are needed instead of regular materials. But as mentioned before, when switching the mesh the 
+		 * code internally will not override the material instances. So we simply have to force this manually, hence this function.
+		 * @param SkelMeshComponent The skeletal mesh component to change the mesh for.
+		 * @param Mesh The new mesh to use.
+		 */
+		static void ChangeSkeletalMeshOnComponent(USkeletalMeshComponent* SkelMeshComponent, USkeletalMesh* Mesh);
+
 		// Required overrides.
 		/**
 		 * Get the number of training frames.
