[MLDeformer] Fix some bug that reported crazy amount of memory usage in the UI and directly improved it by separating memory from asset sizes.

diff --git a/Engine/Plugins/Animation/MLDeformer/MLDeformerFramework/Source/MLDeformerFramework/Private/MLDeformerModel.cpp b/Engine/Plugins/Animation/MLDeformer/MLDeformerFramework/Source/MLDeformerFramework/Private/MLDeformerModel.cpp
--- a/Engine/Plugins/Animation/MLDeformer/MLDeformerFramework/Source/MLDeformerFramework/Private/MLDeformerModel.cpp
+++ b/Engine/Plugins/Animation/MLDeformer/MLDeformerFramework/Source/MLDeformerFramework/Private/MLDeformerModel.cpp
@@ -214,9 +214,14 @@ void UMLDeformerModel::FloatArrayToVector3Array(const TArray<float>& FloatArray,
 		bInvalidateMemUsage = true;
 	}
 
-	uint64 UMLDeformerModel::GetMemUsageInBytes(UE::MLDeformer::EMemUsageRequestFlags Flags) const
+	uint64 UMLDeformerModel::GetCookedAssetSizeInBytes() const
 	{
-		return (Flags == UE::MLDeformer::EMemUsageRequestFlags::Cooked) ? CookedMemUsageInBytes : MemUsageInBytes;
+		return CookedAssetSizeInBytes;
+	}
+
+	uint64 UMLDeformerModel::GetMainMemUsageInBytes() const
+	{
+		return MemUsageInBytes;
 	}
 
 	uint64 UMLDeformerModel::GetGPUMemUsageInBytes() const
@@ -229,20 +234,31 @@ void UMLDeformerModel::FloatArrayToVector3Array(const TArray<float>& FloatArray,
 		return bInvalidateMemUsage;
 	}
 
+	uint64 UMLDeformerModel::GetEditorAssetSizeInBytes() const
+	{
+		return EditorAssetSizeInBytes;
+	}
+
 	void UMLDeformerModel::UpdateMemoryUsage()
 	{
 		bInvalidateMemUsage = false;
+
+		// Start everything at 0 bytes.
 		MemUsageInBytes = 0;
 		GPUMemUsageInBytes = 0;
+		EditorAssetSizeInBytes = 0;
+		CookedAssetSizeInBytes = 0;
 
-		// Get the size of the model to get an approximate size.
+		// Get the resource size of the ML Deformer model.
 		UMLDeformerModel* Model = const_cast<UMLDeformerModel*>(this);
-		MemUsageInBytes += static_cast<uint64>(Model->GetResourceSizeBytes(EResourceSizeMode::Type::EstimatedTotal));
+		EditorAssetSizeInBytes += static_cast<uint64>(Model->GetResourceSizeBytes(EResourceSizeMode::Type::EstimatedTotal));
 
-		// Init the cooked size. We can subtract bytes from this to simulate a cook.
-		CookedMemUsageInBytes = MemUsageInBytes;
+		// Set the main mem usage and cooked sizes also to this.
+		// We are going to subtract from this later, to simulate a cook, as we know which data we strip at cook time for example.
+		CookedAssetSizeInBytes += EditorAssetSizeInBytes;
+		MemUsageInBytes += EditorAssetSizeInBytes;
 
-		// Add GPU memory.
+		// Add the VertexMap buffer to the GPU memory usage.
 		if (VertexMapBuffer.VertexBufferRHI.IsValid())
 		{
 			GPUMemUsageInBytes += VertexMapBuffer.VertexBufferRHI->GetSize();
diff --git a/Engine/Plugins/Animation/MLDeformer/MLDeformerFramework/Source/MLDeformerFramework/Private/MLDeformerMorphModel.cpp b/Engine/Plugins/Animation/MLDeformer/MLDeformerFramework/Source/MLDeformerFramework/Private/MLDeformerMorphModel.cpp
--- a/Engine/Plugins/Animation/MLDeformer/MLDeformerFramework/Source/MLDeformerFramework/Private/MLDeformerMorphModel.cpp
+++ b/Engine/Plugins/Animation/MLDeformer/MLDeformerFramework/Source/MLDeformerFramework/Private/MLDeformerMorphModel.cpp
@@ -127,7 +127,7 @@ void UMLDeformerMorphModel::UpdateStatistics()
 {
 	NumMorphTargets = MorphTargetSet.IsValid() ? MorphTargetSet->MorphBuffers.GetNumMorphs() : 0;
 	CompressedMorphDataSizeInBytes = MorphTargetSet.IsValid() ? MorphTargetSet->MorphBuffers.GetMorphDataSizeInBytes() : 0;
-	UncompressedMorphDataSizeInBytes = GetMorphTargetDeltas().Num() * sizeof(FVector3f);
+	UncompressedMorphDataSizeInBytes = GetMorphTargetDeltas().Num() * GetMorphTargetDeltas().GetTypeSize();
 }
 
 void UMLDeformerMorphModel::SetMorphTargetsMaxWeights(const TArray<float>& MaxWeights)
@@ -187,22 +187,19 @@ void UMLDeformerMorphModel::UpdateMemoryUsage()
 {
 	Super::UpdateMemoryUsage();
 
-	// We strip the deltas when cooking.
-	CookedMemUsageInBytes -= MorphTargetDeltas.GetAllocatedSize();
-
-	// Strip the input item mask buffer from the cooked size.
-	if (GetInputInfo() && GetInputInfo()->IsA<UMLDeformerMorphModelInputInfo>())
-	{
-		const UMLDeformerMorphModelInputInfo* MorphInputInfo = Cast<UMLDeformerMorphModelInputInfo>(GetInputInfo());
-		check(MorphInputInfo);	// Input Info class is expected to be inherited from the UMLDeformerMorphModelInputInfo class.
-		CookedMemUsageInBytes -= MorphInputInfo->GetInputItemMaskBuffer().GetAllocatedSize();
-	}
+	// Remove the raw uncompressed deltas from the cooked size and memory usage, as they are stripped during cook.
+	// This means the game itself won't have this data in the asset or memory.
+	CookedAssetSizeInBytes -= UncompressedMorphDataSizeInBytes;
+	MemUsageInBytes -= UncompressedMorphDataSizeInBytes;
 
 	// Add the compressed morph target data size.
+	// We add this to both the GPU memory, and the cooked asset size.
+	// The morph targets are stored in a compressed way inside the asset.
 	const uint64 GPUMorphSize = GetCompressedMorphDataSizeInBytes();
 	GPUMemUsageInBytes += GPUMorphSize;
-	CookedMemUsageInBytes += GPUMorphSize;
+	CookedAssetSizeInBytes += GPUMorphSize;
 }
 #endif
 
 #undef LOCTEXT_NAMESPACE
+
diff --git a/Engine/Plugins/Animation/MLDeformer/MLDeformerFramework/Source/MLDeformerFramework/Public/MLDeformerModel.h b/Engine/Plugins/Animation/MLDeformer/MLDeformerFramework/Source/MLDeformerFramework/Public/MLDeformerModel.h
--- a/Engine/Plugins/Animation/MLDeformer/MLDeformerFramework/Source/MLDeformerFramework/Public/MLDeformerModel.h
+++ b/Engine/Plugins/Animation/MLDeformer/MLDeformerFramework/Source/MLDeformerFramework/Public/MLDeformerModel.h
@@ -259,10 +259,30 @@ public:
 	 * @param Flags The memory request flags.
 	 * @return The number of bytes that this model uses.
 	 */
-	uint64 GetMemUsageInBytes(UE::MLDeformer::EMemUsageRequestFlags Flags) const;
+	UE_DEPRECATED(5.4, "This method will be removed.")
+	uint64 GetMemUsageInBytes(UE::MLDeformer::EMemUsageRequestFlags Flags) const { return 0; }
 
 	/**
-	 * Get the GPU memory usage for this model.
+	 * Get the estimated size of the asset on disk. This is the uncooked asset, which is larger than the cooked one.
+	 * @return The size in bytes.
+	 */
+	uint64 GetEditorAssetSizeInBytes() const;
+
+	/**
+	 * Get the estimated size of this asset on disk, when cooked.
+	 * So this is the estimated size of the asset that will be packaged inside your project.
+	 * @return The size in bytes.
+	 */
+	uint64 GetCookedAssetSizeInBytes() const;
+
+	/**
+	 * Get the estimated main memory usage for this model.
+	 * @return The number of bytes that this model uses.
+	 */
+	uint64 GetMainMemUsageInBytes() const;
+
+	/**
+	 * Get the estimated GPU memory usage for this model.
 	 * @return The number of bytes that this model uses.
 	 */
 	uint64 GetGPUMemUsageInBytes() const;
@@ -572,13 +592,20 @@ protected:
 	/** Should we recalculate the memory usage? */
 	bool bInvalidateMemUsage = true;
 
-	/** The computed memory usage. */
+	/** Estimated main memory usage. */
 	uint64 MemUsageInBytes = 0;
 
 	/** The cooked memory usage. */
+	UE_DEPRECATED(5.4, "This member will be removed. You most likely want to store this value inside the CookedAssetSize member.")
 	uint64 CookedMemUsageInBytes = 0;
 
-	/** GPU memory usage. */
+	/** Estimated editor asset size. */
+	uint64 EditorAssetSizeInBytes = 0;
+
+	/** Estimated cooked asset size. */
+	uint64 CookedAssetSizeInBytes = 0;
+
+	/** Estimated GPU memory usage. */
 	uint64 GPUMemUsageInBytes = 0;
 #endif
 
diff --git a/Engine/Plugins/Animation/MLDeformer/MLDeformerFramework/Source/MLDeformerFrameworkEditor/Private/MLDeformerEditorToolkit.cpp b/Engine/Plugins/Animation/MLDeformer/MLDeformerFramework/Source/MLDeformerFrameworkEditor/Private/MLDeformerEditorToolkit.cpp
--- a/Engine/Plugins/Animation/MLDeformer/MLDeformerFramework/Source/MLDeformerFrameworkEditor/Private/MLDeformerEditorToolkit.cpp
+++ b/Engine/Plugins/Animation/MLDeformer/MLDeformerFramework/Source/MLDeformerFrameworkEditor/Private/MLDeformerEditorToolkit.cpp
@@ -773,7 +773,6 @@ namespace UE::MLDeformer
 					ActiveModel->ClampCurrentTrainingFrameIndex();
 					const int32 CurrentFrameNumber = ActiveModel->GetModel()->GetVizSettings()->GetTrainingFrameNumber();
 					ActiveModel->SetTrainingFrame(CurrentFrameNumber);
-					ActiveModel->GetSampler()->SetVertexDeltaSpace(EVertexDeltaSpace::PostSkinning);
 					ActiveModel->SampleDeltas();
 				}
 				else if (VizSettings->GetVisualizationMode() == EMLDeformerVizMode::TestData)
diff --git a/Engine/Plugins/Animation/MLDeformer/MLDeformerFramework/Source/MLDeformerFrameworkEditor/Private/MLDeformerGeomCacheEditorModel.cpp b/Engine/Plugins/Animation/MLDeformer/MLDeformerFramework/Source/MLDeformerFrameworkEditor/Private/MLDeformerGeomCacheEditorModel.cpp
--- a/Engine/Plugins/Animation/MLDeformer/MLDeformerFramework/Source/MLDeformerFrameworkEditor/Private/MLDeformerGeomCacheEditorModel.cpp
+++ b/Engine/Plugins/Animation/MLDeformer/MLDeformerFramework/Source/MLDeformerFrameworkEditor/Private/MLDeformerGeomCacheEditorModel.cpp
@@ -120,7 +120,11 @@ namespace UE::MLDeformer
 
 		// Check if we have any mappings at all.
 		FMLDeformerGeomCacheSampler* GeomCacheSampler = GetGeomCacheSampler();
-		GeomCacheSampler->Init(this);
+		if (!GeomCacheSampler->IsInitialized())
+		{
+			GeomCacheSampler->Init(this);
+		}
+
 		if (GeomCacheSampler->GetMeshMappings().IsEmpty())
 		{
 			return;
diff --git a/Engine/Plugins/Animation/MLDeformer/MLDeformerFramework/Source/MLDeformerFrameworkEditor/Private/MLDeformerMorphModelVizSettingsDetails.cpp b/Engine/Plugins/Animation/MLDeformer/MLDeformerFramework/Source/MLDeformerFrameworkEditor/Private/MLDeformerMorphModelVizSettingsDetails.cpp
--- a/Engine/Plugins/Animation/MLDeformer/MLDeformerFramework/Source/MLDeformerFrameworkEditor/Private/MLDeformerMorphModelVizSettingsDetails.cpp
+++ b/Engine/Plugins/Animation/MLDeformer/MLDeformerFramework/Source/MLDeformerFrameworkEditor/Private/MLDeformerMorphModelVizSettingsDetails.cpp
@@ -103,32 +103,6 @@ namespace UE::MLDeformer
 				)
 			];
 	}
-
-	void FMLDeformerMorphModelVizSettingsDetails::AddStatistics()
-	{
-		FMLDeformerGeomCacheVizSettingsDetails::AddStatistics();
-
-		StatsGPUMemUsageGroup->AddWidgetRow()
-			.NameContent()
-			[			
-				SNew(STextBlock)
-				.Text(LOCTEXT("MorphTargetsLabel", "Morph Targets"))
-				.Font(IDetailLayoutBuilder::GetDetailFont())
-			]
-			.ValueContent()
-			[
-				SNew(STextBlock)
-				.Text_Lambda
-				(
-					[this]()
-					{
-						const float CompressedMb = MorphModel->GetCompressedMorphDataSizeInBytes() / static_cast<float>(1024 * 1024);
-						return FText::Format(LOCTEXT("MorphTargetsValue", "{0} mb"), FText::AsNumber(CompressedMb, &MemUsageMetricFormat));
-					}
-				)
-				.Font(IDetailLayoutBuilder::GetDetailFont())
-			];
-	}
 }	//namespace UE::MLDeformer
 
 #undef LOCTEXT_NAMESPACE
diff --git a/Engine/Plugins/Animation/MLDeformer/MLDeformerFramework/Source/MLDeformerFrameworkEditor/Private/MLDeformerVizSettingsDetails.cpp b/Engine/Plugins/Animation/MLDeformer/MLDeformerFramework/Source/MLDeformerFrameworkEditor/Private/MLDeformerVizSettingsDetails.cpp
--- a/Engine/Plugins/Animation/MLDeformer/MLDeformerFramework/Source/MLDeformerFrameworkEditor/Private/MLDeformerVizSettingsDetails.cpp
+++ b/Engine/Plugins/Animation/MLDeformer/MLDeformerFramework/Source/MLDeformerFrameworkEditor/Private/MLDeformerVizSettingsDetails.cpp
@@ -259,105 +259,106 @@ namespace UE::MLDeformer
 
 	void FMLDeformerVizSettingsDetails::AddStatistics()
 	{
-		const FLinearColor MainColor = FMLDeformerEditorStyle::Get().GetColor("MLDeformer.Statistics.Performance");
+		const FLinearColor HighlightColor = FMLDeformerEditorStyle::Get().GetColor("MLDeformer.Statistics.Performance");
 
 		// Create the groups
 		StatsPerformanceGroup = &StatsCategoryBuilder->AddGroup("CPU Performance", LOCTEXT("CPUPerformanceLabel", "CPU Performance"), false, true);
-		StatsMemUsageGroup = &StatsCategoryBuilder->AddGroup("Approximated Memory Usage", LOCTEXT("MemoryUsageLabel", "Memory Usage (Approximated)"), false, true);
-		StatsMainMemUsageGroup = &StatsMemUsageGroup->AddGroup("Main Memory Usage", FText(), false);
-		StatsMainMemUsageGroup->HeaderRow()
+		StatsMemUsageGroup = &StatsCategoryBuilder->AddGroup("Estimated Memory Usage", LOCTEXT("MemoryUsageLabel", "Estimated Memory Usage"), false, true);
+		StatsAssetSizeGroup = &StatsCategoryBuilder->AddGroup("Estimated Asset Sizes", LOCTEXT("AssetSize", "Estimated Asset Sizes"), false, true);
+
+		AddStatsPerfRow(*StatsPerformanceGroup, LOCTEXT("AvgTickTimeLabel", "Avg Time"), EditorModel, PerformanceMetricFormat, true,
+			[](const FMLDeformerEditorModel* EditorModelPtr)
+			{
+				UMLDeformerComponent* MLDeformerComponent = EditorModelPtr->FindMLDeformerComponent(ActorID_Test_MLDeformed);
+				return MLDeformerComponent ? MLDeformerComponent->GetTickPerfCounter().GetCyclesAverage() : 0;
+			});
+
+		AddStatsPerfRow(*StatsPerformanceGroup, LOCTEXT("MinTickTimeLabel", "Min Time"), EditorModel, PerformanceMetricFormat, false,
+			[](const FMLDeformerEditorModel* EditorModelPtr)
+			{
+				UMLDeformerComponent* MLDeformerComponent = EditorModelPtr->FindMLDeformerComponent(ActorID_Test_MLDeformed);
+				return MLDeformerComponent ? MLDeformerComponent->GetTickPerfCounter().GetCyclesMin() : 0;
+			});
+
+		AddStatsPerfRow(*StatsPerformanceGroup, LOCTEXT("MaxTickTimeLabel", "Max Time"), EditorModel, PerformanceMetricFormat, false,
+			[](const FMLDeformerEditorModel* EditorModelPtr)
+			{
+				UMLDeformerComponent* MLDeformerComponent = EditorModelPtr->FindMLDeformerComponent(ActorID_Test_MLDeformed);
+				return MLDeformerComponent ? MLDeformerComponent->GetTickPerfCounter().GetCyclesMax() : 0;
+			});
+
+		StatsMemUsageGroup->AddWidgetRow()
 			.NameContent()
-			[
+			[			
 				SNew(STextBlock)
 				.Font(IDetailLayoutBuilder::GetDetailFont())
-				.Text(LOCTEXT("MainMemoryTotalMemUsageLabel", "Main Memory"))
-			]			
+				.Text(LOCTEXT("MainMemoryLabel", "Main Memory"))
+			]
 			.ValueContent()
 			[
 				SNew(STextBlock)
 				.Font(IDetailLayoutBuilder::GetDetailFont())
-				.ColorAndOpacity(MainColor)
+				.ColorAndOpacity(HighlightColor)
 				.Text_Lambda
 				(
 					[this]
 					{
-						const uint64 TotalBytes = Model->GetMemUsageInBytes(UE::MLDeformer::EMemUsageRequestFlags::Cooked);
-						return FText::Format(LOCTEXT("TotalMemUsageValue", "{0} mb"), FText::AsNumber(TotalBytes / static_cast<float>(1024*1024), &MemUsageMetricFormat));
+						const uint64 TotalBytes = Model->GetMainMemUsageInBytes();
+						return FText::Format(LOCTEXT("MainMemoryValue", "{0} mb"), FText::AsNumber(TotalBytes / static_cast<float>(1024*1024), &MemUsageMetricFormat));
 					}
 				)
 			];
 
-		StatsGPUMemUsageGroup = &StatsMemUsageGroup->AddGroup("GPU Memory Usage", FText(), false);
-		StatsGPUMemUsageGroup->HeaderRow()
+		StatsMemUsageGroup->AddWidgetRow()
 			.NameContent()
-			[
+			[			
 				SNew(STextBlock)
 				.Font(IDetailLayoutBuilder::GetDetailFont())
-				.Text(LOCTEXT("GPUMemoryUsageLabel", "GPU Memory"))
-			]			
+				.Text(LOCTEXT("GPUMemoryLabel", "GPU Memory"))
+			]
 			.ValueContent()
 			[
 				SNew(STextBlock)
 				.Font(IDetailLayoutBuilder::GetDetailFont())
-				.ColorAndOpacity(MainColor)
+				.ColorAndOpacity(HighlightColor)
 				.Text_Lambda
 				(
 					[this]
 					{
 						const uint64 TotalBytes = Model->GetGPUMemUsageInBytes();
-						return FText::Format(LOCTEXT("TotalGPUMemUsageValue", "{0} mb"), FText::AsNumber(TotalBytes / static_cast<float>(1024*1024), &MemUsageMetricFormat));
+						return FText::Format(LOCTEXT("GPUMemoryValue", "{0} mb"), FText::AsNumber(TotalBytes / static_cast<float>(1024*1024), &MemUsageMetricFormat));
 					}
 				)
 			];
 
-		AddStatsPerfRow(*StatsPerformanceGroup, LOCTEXT("AvgTickTimeLabel", "Avg Time"), EditorModel, PerformanceMetricFormat, true,
-			[](const FMLDeformerEditorModel* EditorModelPtr)
-			{
-				UMLDeformerComponent* MLDeformerComponent = EditorModelPtr->FindMLDeformerComponent(ActorID_Test_MLDeformed);
-				return MLDeformerComponent ? MLDeformerComponent->GetTickPerfCounter().GetCyclesAverage() : 0;
-			});
-
-		AddStatsPerfRow(*StatsPerformanceGroup, LOCTEXT("MinTickTimeLabel", "Min Time"), EditorModel, PerformanceMetricFormat, false,
-			[](const FMLDeformerEditorModel* EditorModelPtr)
-			{
-				UMLDeformerComponent* MLDeformerComponent = EditorModelPtr->FindMLDeformerComponent(ActorID_Test_MLDeformed);
-				return MLDeformerComponent ? MLDeformerComponent->GetTickPerfCounter().GetCyclesMin() : 0;
-			});
-
-		AddStatsPerfRow(*StatsPerformanceGroup, LOCTEXT("MaxTickTimeLabel", "Max Time"), EditorModel, PerformanceMetricFormat, false,
-			[](const FMLDeformerEditorModel* EditorModelPtr)
-			{
-				UMLDeformerComponent* MLDeformerComponent = EditorModelPtr->FindMLDeformerComponent(ActorID_Test_MLDeformed);
-				return MLDeformerComponent ? MLDeformerComponent->GetTickPerfCounter().GetCyclesMax() : 0;
-			});
-
-		StatsMainMemUsageGroup->AddWidgetRow()
+		StatsAssetSizeGroup->AddWidgetRow()
 			.NameContent()
 			[			
 				SNew(STextBlock)
 				.Font(IDetailLayoutBuilder::GetDetailFont())
-				.Text(LOCTEXT("ModelCookedMemUsageLabel", "Model - Cooked"))
+				.Text(LOCTEXT("CookedAssetSizeLabel", "Cooked"))
 			]
 			.ValueContent()
 			[
 				SNew(STextBlock)
 				.Font(IDetailLayoutBuilder::GetDetailFont())
+				.ColorAndOpacity(HighlightColor)
 				.Text_Lambda
 				(
 					[this]
 					{
-						const uint64 TotalBytes = Model->GetMemUsageInBytes(UE::MLDeformer::EMemUsageRequestFlags::Cooked);
-						return FText::Format(LOCTEXT("TotalCookedModelMemUsageValue", "{0} mb"), FText::AsNumber(TotalBytes / static_cast<float>(1024*1024), &MemUsageMetricFormat));
+						const uint64 TotalBytes = Model->GetCookedAssetSizeInBytes();
+						return FText::Format(LOCTEXT("CookedAssetSizeInBytesValue", "{0} mb"), FText::AsNumber(TotalBytes / static_cast<float>(1024*1024), &MemUsageMetricFormat));
 					}
 				)
 			];
 
-		StatsMainMemUsageGroup->AddWidgetRow()
+		StatsAssetSizeGroup->AddWidgetRow()
 			.NameContent()
 			[			
 				SNew(STextBlock)
 				.Font(IDetailLayoutBuilder::GetDetailFont())
-				.Text(LOCTEXT("ModelInEditorMemUsageLabel", "Model - In Editor"))
+				.Text(LOCTEXT("EditorAssetSizeLabel", "Uncooked"))
 			]
 			.ValueContent()
 			[
@@ -367,8 +368,8 @@ namespace UE::MLDeformer
 				(
 					[this]
 					{
-						const uint64 TotalBytes = Model->GetMemUsageInBytes(UE::MLDeformer::EMemUsageRequestFlags::Uncooked);
-						return FText::Format(LOCTEXT("TotalInEditorModelMemUsageValue", "{0} mb"), FText::AsNumber(TotalBytes / static_cast<float>(1024*1024), &MemUsageMetricFormat));
+						const uint64 TotalBytes = Model->GetEditorAssetSizeInBytes();
+						return FText::Format(LOCTEXT("EditorAssetSizeInBytesValue", "{0} mb"), FText::AsNumber(TotalBytes / static_cast<float>(1024*1024), &MemUsageMetricFormat));
 					}
 				)
 			];
diff --git a/Engine/Plugins/Animation/MLDeformer/MLDeformerFramework/Source/MLDeformerFrameworkEditor/Public/MLDeformerMorphModelVizSettingsDetails.h b/Engine/Plugins/Animation/MLDeformer/MLDeformerFramework/Source/MLDeformerFrameworkEditor/Public/MLDeformerMorphModelVizSettingsDetails.h
--- a/Engine/Plugins/Animation/MLDeformer/MLDeformerFramework/Source/MLDeformerFrameworkEditor/Public/MLDeformerMorphModelVizSettingsDetails.h
+++ b/Engine/Plugins/Animation/MLDeformer/MLDeformerFramework/Source/MLDeformerFrameworkEditor/Public/MLDeformerMorphModelVizSettingsDetails.h
@@ -20,7 +20,6 @@ namespace UE::MLDeformer
 		// FMLDeformerVizSettingsDetails overrides.
 		virtual bool UpdateMemberPointers(const TArray<TWeakObjectPtr<UObject>>& Objects) override;
 		virtual void AddAdditionalSettings() override;
-		virtual void AddStatistics() override;
 		// ~END FMLDeformerVizSettingsDetails overrides.
 
 	protected:
diff --git a/Engine/Plugins/Animation/MLDeformer/MLDeformerFramework/Source/MLDeformerFrameworkEditor/Public/MLDeformerVizSettingsDetails.h b/Engine/Plugins/Animation/MLDeformer/MLDeformerFramework/Source/MLDeformerFrameworkEditor/Public/MLDeformerVizSettingsDetails.h
--- a/Engine/Plugins/Animation/MLDeformer/MLDeformerFramework/Source/MLDeformerFrameworkEditor/Public/MLDeformerVizSettingsDetails.h
+++ b/Engine/Plugins/Animation/MLDeformer/MLDeformerFramework/Source/MLDeformerFrameworkEditor/Public/MLDeformerVizSettingsDetails.h
@@ -109,7 +109,11 @@ namespace UE::MLDeformer
 		/** The memory usage group inside the statistics category. */
 		IDetailGroup* StatsMemUsageGroup = nullptr;
 
+		/** The memory usage group inside the statistics category. */
+		IDetailGroup* StatsAssetSizeGroup = nullptr;
+
 		/** Main memory usage subgroup. */
+		UE_DEPRECATED(5.4, "This member will be removed.")
 		IDetailGroup* StatsMainMemUsageGroup = nullptr;
 
 		/** GPU memory usage subgroup. */
diff --git a/Engine/Plugins/Animation/MLDeformer/NearestNeighborModel/Source/NearestNeighborModel/Private/NearestNeighborModel.cpp b/Engine/Plugins/Animation/MLDeformer/NearestNeighborModel/Source/NearestNeighborModel/Private/NearestNeighborModel.cpp
--- a/Engine/Plugins/Animation/MLDeformer/NearestNeighborModel/Source/NearestNeighborModel/Private/NearestNeighborModel.cpp
+++ b/Engine/Plugins/Animation/MLDeformer/NearestNeighborModel/Source/NearestNeighborModel/Private/NearestNeighborModel.cpp
@@ -64,6 +64,26 @@ UMLDeformerInputInfo* UNearestNeighborModel::CreateInputInfo()
 	return NearestNeighborModelInputInfo;
 }
 
+#if WITH_EDITOR
+void UNearestNeighborModel::UpdateMemoryUsage()
+{
+	Super::UpdateMemoryUsage();
+
+	uint64 ClothDataSize = 0;
+	for (const FClothPartData& Part : ClothPartData)
+	{
+		ClothDataSize += Part.PCABasis.GetTypeSize() * Part.PCABasis.Num();
+		ClothDataSize += Part.NeighborOffsets.GetTypeSize() * Part.NeighborOffsets.Num();
+		ClothDataSize += Part.NeighborCoeffs.GetTypeSize() * Part.NeighborCoeffs.Num();
+		ClothDataSize += Part.VertexMean.GetTypeSize() * Part.VertexMean.Num();
+		ClothDataSize += Part.VertexMap.GetTypeSize() * Part.VertexMap.Num();
+	}
+
+	CookedAssetSizeInBytes -= ClothDataSize;
+	MemUsageInBytes -= ClothDataSize;
+}
+#endif
+
 UMLDeformerModelInstance* UNearestNeighborModel::CreateModelInstance(UMLDeformerComponent* Component)
 {
 	return NewObject<UNearestNeighborModelInstance>(Component);
@@ -520,4 +540,5 @@ int32 UNearestNeighborModel::GetOptimizedNetworkNumOutputs() const
 	return OptimizedNetwork ? OptimizedNetwork->GetNumOutputs() : 0;
 }
 
+
 #undef LOCTEXT_NAMESPACE
diff --git a/Engine/Plugins/Animation/MLDeformer/NearestNeighborModel/Source/NearestNeighborModel/Public/NearestNeighborModel.h b/Engine/Plugins/Animation/MLDeformer/NearestNeighborModel/Source/NearestNeighborModel/Public/NearestNeighborModel.h
--- a/Engine/Plugins/Animation/MLDeformer/NearestNeighborModel/Source/NearestNeighborModel/Public/NearestNeighborModel.h
+++ b/Engine/Plugins/Animation/MLDeformer/NearestNeighborModel/Source/NearestNeighborModel/Public/NearestNeighborModel.h
@@ -171,6 +171,9 @@ public:
 	virtual UMLDeformerModelInstance* CreateModelInstance(UMLDeformerComponent* Component) override;
 	virtual UMLDeformerInputInfo* CreateInputInfo() override;
 	virtual FString GetDisplayName() const override { return "Nearest Neighbor Model"; }
+#if WITH_EDITOR
+	virtual void UpdateMemoryUsage() override;
+#endif
 	// ~END UMLDeformerModel overrides.
 
 	friend class UNearestNeighborModelInstance;
diff --git a/Engine/Plugins/Animation/MLDeformer/VertexDeltaModel/Source/VertexDeltaModel/Private/VertexDeltaModel.cpp b/Engine/Plugins/Animation/MLDeformer/VertexDeltaModel/Source/VertexDeltaModel/Private/VertexDeltaModel.cpp
--- a/Engine/Plugins/Animation/MLDeformer/VertexDeltaModel/Source/VertexDeltaModel/Private/VertexDeltaModel.cpp
+++ b/Engine/Plugins/Animation/MLDeformer/VertexDeltaModel/Source/VertexDeltaModel/Private/VertexDeltaModel.cpp
@@ -62,19 +62,4 @@ void UVertexDeltaModel::PostLoad()
 	Super::PostLoad();
 }
 
-#if WITH_EDITOR
-	void UVertexDeltaModel::UpdateMemoryUsage()
-	{
-		Super::UpdateMemoryUsage();
-
-		if (NNEModel)
-		{
-				const uint64 NeuralNetSize = 0;
-				GPUMemUsageInBytes += NeuralNetSize;
-				MemUsageInBytes -= NeuralNetSize;
-				CookedMemUsageInBytes -= NeuralNetSize;
-		}
-	}
-#endif
-
 #undef LOCTEXT_NAMESPACE
diff --git a/Engine/Plugins/Animation/MLDeformer/VertexDeltaModel/Source/VertexDeltaModel/Public/VertexDeltaModel.h b/Engine/Plugins/Animation/MLDeformer/VertexDeltaModel/Source/VertexDeltaModel/Public/VertexDeltaModel.h
--- a/Engine/Plugins/Animation/MLDeformer/VertexDeltaModel/Source/VertexDeltaModel/Public/VertexDeltaModel.h
+++ b/Engine/Plugins/Animation/MLDeformer/VertexDeltaModel/Source/VertexDeltaModel/Public/VertexDeltaModel.h
@@ -37,9 +37,6 @@ public:
 	virtual UMLDeformerModelInstance* CreateModelInstance(UMLDeformerComponent* Component) override;
 	virtual bool IsNeuralNetworkOnGPU() const override		{ return true; }	// GPU neural network.
 	virtual FString GetDefaultDeformerGraphAssetPath() const override;
-#if WITH_EDITOR
-	virtual void UpdateMemoryUsage() override;
-#endif
 	// ~END UMLDeformerModel overrides.
 
 #if WITH_EDITORONLY_DATA
