[MLDeformer] Fix sampling inaccuracies, which could lead to large deltas at certain frames. And some small parallel optimizations to make samplers slightly faster.

diff --git a/Engine/Plugins/Animation/MLDeformer/MLDeformerFramework/Source/MLDeformerFramework/Private/MLDeformerGeomCacheHelpers.cpp b/Engine/Plugins/Animation/MLDeformer/MLDeformerFramework/Source/MLDeformerFramework/Private/MLDeformerGeomCacheHelpers.cpp
--- a/Engine/Plugins/Animation/MLDeformer/MLDeformerFramework/Source/MLDeformerFramework/Private/MLDeformerGeomCacheHelpers.cpp
+++ b/Engine/Plugins/Animation/MLDeformer/MLDeformerFramework/Source/MLDeformerFramework/Private/MLDeformerGeomCacheHelpers.cpp
@@ -10,6 +10,7 @@
 #include "Engine/SkeletalMesh.h"
 #include "Rendering/SkeletalMeshModel.h"
 #include "Rendering/SkeletalMeshLODModel.h"
+#include "Async/ParallelFor.h"
 
 #define LOCTEXT_NAMESPACE "MLDeformerGeomCacheHelpers"
 
@@ -212,18 +213,30 @@ namespace UE::MLDeformer
 					Mapping.ImportedVertexToRenderVertexMap.AddUninitialized(NumSkelMeshVerts);
 
 					// For all vertices (both skel mesh and geom cache mesh have the same number of verts here).
-					for (int32 VertexIndex = 0; VertexIndex < NumSkelMeshVerts; ++VertexIndex)
+					const int32 BatchSize = 100;
+					const int32 NumBatches = (MeshInfo.NumVertices / BatchSize) + 1;
+					ParallelFor(NumBatches, [&](int32 BatchIndex)
 					{
-						// Find the first vertex with the same dcc vertex in the geom cache mesh.
-						// When there are multiple vertices with the same vertex number here, they are duplicates with different normals or uvs etc.
-						// However they all share the same vertex position, so we can just find the first hit, as we only need the position later on.
-						const int32 GeomCacheVertexIndex = GeomCacheMeshData.ImportedVertexNumbers.Find(VertexIndex);
-						Mapping.SkelMeshToTrackVertexMap[VertexIndex] = GeomCacheVertexIndex;
-
-						// Map the source asset vertex number to a render vertex. This is the first duplicate of that vertex.
-						const int32 RenderVertexIndex = ImportedModel->LODModels[0].MeshToImportVertexMap.Find(MeshInfo.StartImportedVertex + VertexIndex);
-						Mapping.ImportedVertexToRenderVertexMap[VertexIndex] = RenderVertexIndex;
-					}
+						const int32 StartVertex = BatchIndex * BatchSize;
+						if (StartVertex >= NumSkelMeshVerts)
+						{
+							return;
+						}
+
+						const int32 NumVertsInBatch = (StartVertex + BatchSize) < MeshInfo.NumVertices ? BatchSize : FMath::Max(NumSkelMeshVerts - StartVertex, 0);
+						for (int32 VertexIndex = StartVertex; VertexIndex < StartVertex + NumVertsInBatch; ++VertexIndex)
+						{
+							// Find the first vertex with the same dcc vertex in the geom cache mesh.
+							// When there are multiple vertices with the same vertex number here, they are duplicates with different normals or uvs etc.
+							// However they all share the same vertex position, so we can just find the first hit, as we only need the position later on.
+							const int32 GeomCacheVertexIndex = GeomCacheMeshData.ImportedVertexNumbers.Find(VertexIndex);
+							Mapping.SkelMeshToTrackVertexMap[VertexIndex] = GeomCacheVertexIndex;
+
+							// Map the source asset vertex number to a render vertex. This is the first duplicate of that vertex.
+							const int32 RenderVertexIndex = ImportedModel->LODModels[0].MeshToImportVertexMap.Find(MeshInfo.StartImportedVertex + VertexIndex);
+							Mapping.ImportedVertexToRenderVertexMap[VertexIndex] = RenderVertexIndex;
+						}
+					});
 
 					// We found a match, no need to iterate over more Tracks.
 					bFoundMatch = true;
@@ -342,22 +355,31 @@ namespace UE::MLDeformer
 
 		if (InAnimSequence && InGeomCache)
 		{
+			const int32 NumGeomCacheFrames = InGeomCache->GetEndFrame() - InGeomCache->GetStartFrame() + 1;
+			const int32 NumAnimSeqFrames = InAnimSequence->GetNumberOfSampledKeys();
 			const float AnimSeqDuration = InAnimSequence->GetPlayLength();
 			const float GeomCacheDuration = InGeomCache->CalculateDuration();
-			if (FMath::Abs(AnimSeqDuration - GeomCacheDuration) > 0.001f)
+
+			if (NumGeomCacheFrames != NumAnimSeqFrames)
 			{
-				const int32 NumGeomCacheFrames = InGeomCache ? (InGeomCache->GetEndFrame() - InGeomCache->GetStartFrame()) + 1 : 0;
-				const int32 NumAnimSeqFrames = InAnimSequence ? InAnimSequence->GetNumberOfSampledKeys() : 0;
+				Result = FText::Format(
+					LOCTEXT("AnimSeqNumFramesMismatch", "Anim sequence has {0} frames, while the geometry cache has {1} frames."),
+					FText::AsNumber(NumAnimSeqFrames),
+					FText::AsNumber(NumGeomCacheFrames));
+			}
 
+			if (FMath::Abs(AnimSeqDuration - GeomCacheDuration) > 0.0001f)
+			{
 				FNumberFormattingOptions Options;
 				Options.SetUseGrouping(false);
-				Options.SetMaximumFractionalDigits(4);
+				Options.SetMaximumFractionalDigits(3);
+
 				Result = FText::Format(
-					LOCTEXT("AnimSeqNumFramesMismatch", "Anim sequence and Geometry Cache durations don't match!\n\nAnimSeq has a duration of {0} seconds ({1} frames}, while GeomCache has a duration of {2} seconds ({3} frames).\n\nThis can produce incorrect results."),
+					LOCTEXT("AnimSeqDurationMismatch", "{0}{1}Anim sequence duration ({2} secs) doesn't match the geometry cache's duration ({3} secs)."),
+					Result,
+					!Result.IsEmpty() ? FText::FromString("\n\n") : FText::GetEmpty(),
 					FText::AsNumber(AnimSeqDuration, &Options),
-					FText::AsNumber(NumAnimSeqFrames),
-					FText::AsNumber(GeomCacheDuration, &Options),
-					FText::AsNumber(NumGeomCacheFrames));
+					FText::AsNumber(GeomCacheDuration, &Options));
 			}
 		}
 
diff --git a/Engine/Plugins/Animation/MLDeformer/MLDeformerFramework/Source/MLDeformerFrameworkEditor/Private/MLDeformerGeomCacheSampler.cpp b/Engine/Plugins/Animation/MLDeformer/MLDeformerFramework/Source/MLDeformerFrameworkEditor/Private/MLDeformerGeomCacheSampler.cpp
--- a/Engine/Plugins/Animation/MLDeformer/MLDeformerFramework/Source/MLDeformerFrameworkEditor/Private/MLDeformerGeomCacheSampler.cpp
+++ b/Engine/Plugins/Animation/MLDeformer/MLDeformerFramework/Source/MLDeformerFrameworkEditor/Private/MLDeformerGeomCacheSampler.cpp
@@ -19,6 +19,7 @@
 #include "GeometryCache.h"
 #include "GeometryCacheMeshData.h"
 #include "GeometryCacheTrack.h"
+#include "Async/ParallelFor.h"
 
 namespace UE::MLDeformer
 {
@@ -77,7 +78,7 @@ namespace UE::MLDeformer
 
 				// Sample the mesh data of the geom cache.
 				FGeometryCacheMeshData& GeomCacheMeshData = GeomCacheMeshDatas[MeshMappingIndex];
-				if (!Track->GetMeshDataAtTime(SampleTime, GeomCacheMeshData))
+				if (!Track->GetMeshDataAtSampleIndex(InAnimFrameIndex, GeomCacheMeshData))
 				{
 					continue;
 				}
@@ -85,48 +86,61 @@ namespace UE::MLDeformer
 				// Calculate the vertex deltas.
 				const FSkeletalMeshLODRenderData& SkelMeshLODData = SkeletalMesh->GetResourceForRendering()->LODRenderData[LODIndex];
 				const FSkinWeightVertexBuffer& SkinWeightBuffer = *SkeletalMeshComponent->GetSkinWeightBuffer(LODIndex);
-				for (int32 VertexIndex = 0; VertexIndex < MeshInfo.NumVertices; ++VertexIndex)
+
+				const int32 BatchSize = 500;
+				const int32 NumBatches = (MeshInfo.NumVertices / BatchSize) + 1;
+				ParallelFor(NumBatches, [&](int32 BatchIndex)
 				{
-					const int32 SkinnedVertexIndex = MeshInfo.StartImportedVertex + VertexIndex;
-					const int32 GeomCacheVertexIndex = MeshMapping.SkelMeshToTrackVertexMap[VertexIndex];
-					if (GeomCacheVertexIndex != INDEX_NONE && GeomCacheMeshData.Positions.IsValidIndex(GeomCacheVertexIndex))
+					const int32 StartVertex = BatchIndex * BatchSize;
+					if (StartVertex >= MeshInfo.NumVertices)
 					{
-						FVector3f Delta = FVector3f::ZeroVector;
+						return;
+					}
 
-						const int32 ArrayIndex = 3 * SkinnedVertexIndex;
-						if (VertexDeltaSpace == EVertexDeltaSpace::PreSkinning)
+					const int32 NumVertsInBatch = (StartVertex + BatchSize) < MeshInfo.NumVertices ? BatchSize : FMath::Max(MeshInfo.NumVertices - StartVertex, 0);
+					for (int32 VertexIndex = StartVertex; VertexIndex < StartVertex + NumVertsInBatch; ++VertexIndex)
+					{
+						const int32 SkinnedVertexIndex = MeshInfo.StartImportedVertex + VertexIndex;
+						const int32 GeomCacheVertexIndex = MeshMapping.SkelMeshToTrackVertexMap[VertexIndex];
+						if (GeomCacheVertexIndex != INDEX_NONE && GeomCacheMeshData.Positions.IsValidIndex(GeomCacheVertexIndex))
 						{
-							// Calculate the inverse skinning transform for this vertex.
-							const int32 RenderVertexIndex = MeshMapping.ImportedVertexToRenderVertexMap[VertexIndex];
-							if (RenderVertexIndex != INDEX_NONE)
+							FVector3f Delta = FVector3f::ZeroVector;
+
+							const int32 ArrayIndex = 3 * SkinnedVertexIndex;
+							if (VertexDeltaSpace == EVertexDeltaSpace::PreSkinning)
 							{
-								const FMatrix44f InvSkinningTransform = CalcInverseSkinningTransform(RenderVertexIndex, SkelMeshLODData, SkinWeightBuffer);
+								// Calculate the inverse skinning transform for this vertex.
+								const int32 RenderVertexIndex = MeshMapping.ImportedVertexToRenderVertexMap[VertexIndex];
+								if (RenderVertexIndex != INDEX_NONE)
+								{
+									const FMatrix44f InvSkinningTransform = CalcInverseSkinningTransform(RenderVertexIndex, SkelMeshLODData, SkinWeightBuffer);
 
-								// Calculate the pre-skinning data.
-								const FSkeletalMeshLODRenderData& LODData = SkeletalMesh->GetResourceForRendering()->LODRenderData[0];
-								const FVector3f UnskinnedPosition = LODData.StaticVertexBuffers.PositionVertexBuffer.VertexPosition(RenderVertexIndex);
+									// Calculate the pre-skinning data.
+									const FSkeletalMeshLODRenderData& LODData = SkeletalMesh->GetResourceForRendering()->LODRenderData[0];
+									const FVector3f UnskinnedPosition = LODData.StaticVertexBuffers.PositionVertexBuffer.VertexPosition(RenderVertexIndex);
+									const FVector3f GeomCacheVertexPos = (FVector3f)AlignmentTransform.TransformPosition((FVector)GeomCacheMeshData.Positions[GeomCacheVertexIndex]);
+									const FVector3f PreSkinningTargetPos = InvSkinningTransform.TransformPosition(GeomCacheVertexPos);
+									Delta = PreSkinningTargetPos - UnskinnedPosition;
+								}
+							}
+							else // We're post skinning.
+							{
+								check(VertexDeltaSpace == EVertexDeltaSpace::PostSkinning);
+								const FVector3f SkinnedVertexPos = SkinnedVertexPositions[SkinnedVertexIndex];
 								const FVector3f GeomCacheVertexPos = (FVector3f)AlignmentTransform.TransformPosition((FVector)GeomCacheMeshData.Positions[GeomCacheVertexIndex]);
-								const FVector3f PreSkinningTargetPos = InvSkinningTransform.TransformPosition(GeomCacheVertexPos);
-								Delta = PreSkinningTargetPos - UnskinnedPosition;
+								Delta = GeomCacheVertexPos - SkinnedVertexPos;
 							}
-						}
-						else // We're post skinning.
-						{
-							check(VertexDeltaSpace == EVertexDeltaSpace::PostSkinning);
-							const FVector3f SkinnedVertexPos = SkinnedVertexPositions[SkinnedVertexIndex];
-							const FVector3f GeomCacheVertexPos = (FVector3f)AlignmentTransform.TransformPosition((FVector)GeomCacheMeshData.Positions[GeomCacheVertexIndex]);
-							Delta = GeomCacheVertexPos - SkinnedVertexPos;
-						}
 
-						// Set the delta.
-						if (Delta.Length() < DeltaCutoffLength)
-						{
-							VertexDeltas[ArrayIndex] = Delta.X;
-							VertexDeltas[ArrayIndex + 1] = Delta.Y;
-							VertexDeltas[ArrayIndex + 2] = Delta.Z;
+							// Set the delta.
+							if (Delta.Length() < DeltaCutoffLength)
+							{
+								VertexDeltas[ArrayIndex] = Delta.X;
+								VertexDeltas[ArrayIndex + 1] = Delta.Y;
+								VertexDeltas[ArrayIndex + 2] = Delta.Z;
+							}
 						}
 					}
-				}
+				});	// ParallelFor
 			}
 		}
 		else
@@ -156,7 +170,7 @@ namespace UE::MLDeformer
 
 		if (GeometryCacheComponent.Get())
 		{
-			return  GeometryCacheComponent->GetTimeAtFrame(InAnimFrameIndex);
+			return GeometryCacheComponent->GetTimeAtFrame(InAnimFrameIndex);
 		}
 		return 0.0f;
 	}
diff --git a/Engine/Plugins/Animation/MLDeformer/NearestNeighborModel/Source/NearestNeighborModelEditor/Private/NearestNeighborGeomCacheSampler.cpp b/Engine/Plugins/Animation/MLDeformer/NearestNeighborModel/Source/NearestNeighborModelEditor/Private/NearestNeighborGeomCacheSampler.cpp
--- a/Engine/Plugins/Animation/MLDeformer/NearestNeighborModel/Source/NearestNeighborModelEditor/Private/NearestNeighborGeomCacheSampler.cpp
+++ b/Engine/Plugins/Animation/MLDeformer/NearestNeighborModel/Source/NearestNeighborModelEditor/Private/NearestNeighborGeomCacheSampler.cpp
@@ -54,7 +54,7 @@ namespace UE::NearestNeighborModel
 
 				// Sample the mesh data of the geom cache.
 				FGeometryCacheMeshData& GeomCacheMeshData = GeomCacheMeshDatas[MeshMappingIndex];
-				if (!Track->GetMeshDataAtTime(SampleTime, GeomCacheMeshData))
+				if (!Track->GetMeshDataAtSampleIndex(InAnimFrameIndex, GeomCacheMeshData))
 				{
 					continue;
 				}
