Fix deprecation warnings
#rnx

diff --git a/Engine/Plugins/Animation/MLDeformer/MLDeformerFramework/Source/MLDeformerFrameworkEditor/Private/MLDeformerGeomCacheVizSettingsDetails.cpp b/Engine/Plugins/Animation/MLDeformer/MLDeformerFramework/Source/MLDeformerFrameworkEditor/Private/MLDeformerGeomCacheVizSettingsDetails.cpp
--- a/Engine/Plugins/Animation/MLDeformer/MLDeformerFramework/Source/MLDeformerFrameworkEditor/Private/MLDeformerGeomCacheVizSettingsDetails.cpp
+++ b/Engine/Plugins/Animation/MLDeformer/MLDeformerFramework/Source/MLDeformerFrameworkEditor/Private/MLDeformerGeomCacheVizSettingsDetails.cpp
@@ -19,6 +19,14 @@
 
 namespace UE::MLDeformer
 {
+	PRAGMA_DISABLE_DEPRECATION_WARNINGS
+	//UE_DEPRECATED(5.4)
+	FMLDeformerVizSettingsDetails::FMLDeformerVizSettingsDetails(const FMLDeformerVizSettingsDetails&) = default;
+	FMLDeformerVizSettingsDetails::FMLDeformerVizSettingsDetails(FMLDeformerVizSettingsDetails&&) = default;
+	FMLDeformerVizSettingsDetails& FMLDeformerVizSettingsDetails::operator=(const FMLDeformerVizSettingsDetails&) = default;
+	FMLDeformerVizSettingsDetails& FMLDeformerVizSettingsDetails::operator=(FMLDeformerVizSettingsDetails&&) = default;
+	PRAGMA_ENABLE_DEPRECATION_WARNINGS
+
 	bool FMLDeformerGeomCacheVizSettingsDetails::UpdateMemberPointers(const TArray<TWeakObjectPtr<UObject>>& Objects)
 	{
 		if (!FMLDeformerVizSettingsDetails::UpdateMemberPointers(Objects))
diff --git a/Engine/Plugins/Animation/MLDeformer/MLDeformerFramework/Source/MLDeformerFrameworkEditor/Public/MLDeformerVizSettingsDetails.h b/Engine/Plugins/Animation/MLDeformer/MLDeformerFramework/Source/MLDeformerFrameworkEditor/Public/MLDeformerVizSettingsDetails.h
--- a/Engine/Plugins/Animation/MLDeformer/MLDeformerFramework/Source/MLDeformerFrameworkEditor/Public/MLDeformerVizSettingsDetails.h
+++ b/Engine/Plugins/Animation/MLDeformer/MLDeformerFramework/Source/MLDeformerFrameworkEditor/Public/MLDeformerVizSettingsDetails.h
@@ -33,6 +33,12 @@ namespace UE::MLDeformer
 		/** Get a pointer to the detail layout builder. */
 		IDetailLayoutBuilder* GetDetailLayoutBuilder() const { return DetailLayoutBuilder; }
 
+		// UE_DEPRECATED(5.4)
+		FMLDeformerVizSettingsDetails(const FMLDeformerVizSettingsDetails&);
+		FMLDeformerVizSettingsDetails(FMLDeformerVizSettingsDetails&&);
+		FMLDeformerVizSettingsDetails& operator=(const FMLDeformerVizSettingsDetails&);
+		FMLDeformerVizSettingsDetails& operator=(FMLDeformerVizSettingsDetails&&);
+
 	protected:
 		/** The filter that only shows anim sequences that are compatible with the given skeleton. */
 		bool FilterAnimSequences(const FAssetData& AssetData, USkeleton* Skeleton);
