[MLDeformer]
* Fixed compile error

diff --git a/Engine/Plugins/Animation/MLDeformer/MLDeformerFramework/Source/MLDeformerFrameworkEditor/Public/MLDeformerVizSettingsDetails.h b/Engine/Plugins/Animation/MLDeformer/MLDeformerFramework/Source/MLDeformerFrameworkEditor/Public/MLDeformerVizSettingsDetails.h
--- a/Engine/Plugins/Animation/MLDeformer/MLDeformerFramework/Source/MLDeformerFrameworkEditor/Public/MLDeformerVizSettingsDetails.h
+++ b/Engine/Plugins/Animation/MLDeformer/MLDeformerFramework/Source/MLDeformerFrameworkEditor/Public/MLDeformerVizSettingsDetails.h
@@ -33,6 +33,8 @@ namespace UE::MLDeformer
 		/** Get a pointer to the detail layout builder. */
 		IDetailLayoutBuilder* GetDetailLayoutBuilder() const { return DetailLayoutBuilder; }
 
+		FMLDeformerVizSettingsDetails() = default;
+
 		// UE_DEPRECATED(5.4)
 		FMLDeformerVizSettingsDetails(const FMLDeformerVizSettingsDetails&);
 		FMLDeformerVizSettingsDetails(FMLDeformerVizSettingsDetails&&);
