Fix missing default constructor
Move implementations to correct cpp
Clarify reasoning
#rnx

diff --git a/Engine/Plugins/Animation/MLDeformer/MLDeformerFramework/Source/MLDeformerFrameworkEditor/Private/MLDeformerGeomCacheVizSettingsDetails.cpp b/Engine/Plugins/Animation/MLDeformer/MLDeformerFramework/Source/MLDeformerFrameworkEditor/Private/MLDeformerGeomCacheVizSettingsDetails.cpp
--- a/Engine/Plugins/Animation/MLDeformer/MLDeformerFramework/Source/MLDeformerFrameworkEditor/Private/MLDeformerGeomCacheVizSettingsDetails.cpp
+++ b/Engine/Plugins/Animation/MLDeformer/MLDeformerFramework/Source/MLDeformerFrameworkEditor/Private/MLDeformerGeomCacheVizSettingsDetails.cpp
@@ -19,14 +19,6 @@
 
 namespace UE::MLDeformer
 {
-	PRAGMA_DISABLE_DEPRECATION_WARNINGS
-	//UE_DEPRECATED(5.4)
-	FMLDeformerVizSettingsDetails::FMLDeformerVizSettingsDetails(const FMLDeformerVizSettingsDetails&) = default;
-	FMLDeformerVizSettingsDetails::FMLDeformerVizSettingsDetails(FMLDeformerVizSettingsDetails&&) = default;
-	FMLDeformerVizSettingsDetails& FMLDeformerVizSettingsDetails::operator=(const FMLDeformerVizSettingsDetails&) = default;
-	FMLDeformerVizSettingsDetails& FMLDeformerVizSettingsDetails::operator=(FMLDeformerVizSettingsDetails&&) = default;
-	PRAGMA_ENABLE_DEPRECATION_WARNINGS
-
 	bool FMLDeformerGeomCacheVizSettingsDetails::UpdateMemberPointers(const TArray<TWeakObjectPtr<UObject>>& Objects)
 	{
 		if (!FMLDeformerVizSettingsDetails::UpdateMemberPointers(Objects))
diff --git a/Engine/Plugins/Animation/MLDeformer/MLDeformerFramework/Source/MLDeformerFrameworkEditor/Private/MLDeformerVizSettingsDetails.cpp b/Engine/Plugins/Animation/MLDeformer/MLDeformerFramework/Source/MLDeformerFrameworkEditor/Private/MLDeformerVizSettingsDetails.cpp
--- a/Engine/Plugins/Animation/MLDeformer/MLDeformerFramework/Source/MLDeformerFrameworkEditor/Private/MLDeformerVizSettingsDetails.cpp
+++ b/Engine/Plugins/Animation/MLDeformer/MLDeformerFramework/Source/MLDeformerFrameworkEditor/Private/MLDeformerVizSettingsDetails.cpp
@@ -37,6 +37,16 @@
 
 namespace UE::MLDeformer
 {
+	// UE_DEPRECATED(5.4, Remove these constructors when StatsMainMemUsageGroup is removed and there is no longer warnings in the implicit functions)
+	PRAGMA_DISABLE_DEPRECATION_WARNINGS
+	FMLDeformerVizSettingsDetails::FMLDeformerVizSettingsDetails(const FMLDeformerVizSettingsDetails&) = default;
+	FMLDeformerVizSettingsDetails::FMLDeformerVizSettingsDetails(FMLDeformerVizSettingsDetails&&) = default;
+	FMLDeformerVizSettingsDetails& FMLDeformerVizSettingsDetails::operator=(const FMLDeformerVizSettingsDetails&) = default;
+	FMLDeformerVizSettingsDetails& FMLDeformerVizSettingsDetails::operator=(FMLDeformerVizSettingsDetails&&) = default;
+	PRAGMA_ENABLE_DEPRECATION_WARNINGS
+	// END UE_DEPRECATED(5.4)
+
+
 	bool FMLDeformerVizSettingsDetails::UpdateMemberPointers(const TArray<TWeakObjectPtr<UObject>>& Objects)
 	{
 		Model = nullptr;
diff --git a/Engine/Plugins/Animation/MLDeformer/MLDeformerFramework/Source/MLDeformerFrameworkEditor/Public/MLDeformerVizSettingsDetails.h b/Engine/Plugins/Animation/MLDeformer/MLDeformerFramework/Source/MLDeformerFrameworkEditor/Public/MLDeformerVizSettingsDetails.h
--- a/Engine/Plugins/Animation/MLDeformer/MLDeformerFramework/Source/MLDeformerFrameworkEditor/Public/MLDeformerVizSettingsDetails.h
+++ b/Engine/Plugins/Animation/MLDeformer/MLDeformerFramework/Source/MLDeformerFrameworkEditor/Public/MLDeformerVizSettingsDetails.h
@@ -33,13 +33,13 @@ namespace UE::MLDeformer
 		/** Get a pointer to the detail layout builder. */
 		IDetailLayoutBuilder* GetDetailLayoutBuilder() const { return DetailLayoutBuilder; }
 
+		// UE_DEPRECATED(5.4, Remove these constructors when StatsMainMemUsageGroup is removed and there is no longer warnings in the implicit functions)
 		FMLDeformerVizSettingsDetails() = default;
-
-		// UE_DEPRECATED(5.4)
 		FMLDeformerVizSettingsDetails(const FMLDeformerVizSettingsDetails&);
 		FMLDeformerVizSettingsDetails(FMLDeformerVizSettingsDetails&&);
 		FMLDeformerVizSettingsDetails& operator=(const FMLDeformerVizSettingsDetails&);
 		FMLDeformerVizSettingsDetails& operator=(FMLDeformerVizSettingsDetails&&);
+		// END UE_DEPRECATED(5.4)
 
 	protected:
 		/** The filter that only shows anim sequences that are compatible with the given skeleton. */
