[MLDeformer] Fixed some issue where just saved newly created assets would still be marked dirty after opening them again.

diff --git a/Engine/Plugins/Animation/MLDeformer/MLDeformerFramework/Source/MLDeformerFrameworkEditor/Private/MLDeformerAssetFactory.cpp b/Engine/Plugins/Animation/MLDeformer/MLDeformerFramework/Source/MLDeformerFrameworkEditor/Private/MLDeformerAssetFactory.cpp
--- a/Engine/Plugins/Animation/MLDeformer/MLDeformerFramework/Source/MLDeformerFrameworkEditor/Private/MLDeformerAssetFactory.cpp
+++ b/Engine/Plugins/Animation/MLDeformer/MLDeformerFramework/Source/MLDeformerFrameworkEditor/Private/MLDeformerAssetFactory.cpp
@@ -2,7 +2,11 @@
 
 #include "MLDeformerAssetFactory.h"
 #include "MLDeformerAsset.h"
+#include "MLDeformerEditorModule.h"
+#include "MLDeformerModel.h"
 #include "AssetTypeCategories.h"
+#include "Modules/ModuleManager.h"
+#include "MLDeformerModelRegistry.h"
 
 #define LOCTEXT_NAMESPACE "MLDeformerAssetFactory"
 
@@ -15,7 +19,25 @@ UMLDeformerFactory::UMLDeformerFactory()
 
 UObject* UMLDeformerFactory::FactoryCreateNew(UClass* Class, UObject* InParent, FName Name, EObjectFlags Flags, UObject* Context, FFeedbackContext* Warn)
 {
-	return NewObject<UMLDeformerAsset>(InParent, Name, Flags | RF_Transactional);;
+	using namespace UE::MLDeformer;
+
+	UMLDeformerAsset* DeformerAsset = NewObject<UMLDeformerAsset>(InParent, Name, Flags | RF_Transactional);
+	if (DeformerAsset)
+	{
+		FMLDeformerEditorModule& EditorModule = FModuleManager::GetModuleChecked<FMLDeformerEditorModule>("MLDeformerFrameworkEditor");
+		FMLDeformerEditorModelRegistry& ModelRegistry = EditorModule.GetModelRegistry();
+
+		// Create the highest priority ML model on default and use that in our asset.
+		const int32 HighestPriorityIndex = ModelRegistry.GetHighestPriorityModelIndex();
+		TArray<UClass*> ModelTypes;
+		ModelRegistry.GetRegisteredModels().GenerateKeyArray(ModelTypes);
+		TObjectPtr<UMLDeformerModel> Model = NewObject<UMLDeformerModel>(DeformerAsset, ModelTypes[HighestPriorityIndex]);
+		check(Model);
+		Model->Init(DeformerAsset);
+		DeformerAsset->SetModel(Model);
+	}
+
+	return DeformerAsset;
 }
 
 bool UMLDeformerFactory::ShouldShowInNewMenu() const
diff --git a/Engine/Plugins/Animation/MLDeformer/MLDeformerFramework/Source/MLDeformerFrameworkEditor/Private/MLDeformerEditorToolkit.cpp b/Engine/Plugins/Animation/MLDeformer/MLDeformerFramework/Source/MLDeformerFrameworkEditor/Private/MLDeformerEditorToolkit.cpp
--- a/Engine/Plugins/Animation/MLDeformer/MLDeformerFramework/Source/MLDeformerFrameworkEditor/Private/MLDeformerEditorToolkit.cpp
+++ b/Engine/Plugins/Animation/MLDeformer/MLDeformerFramework/Source/MLDeformerFrameworkEditor/Private/MLDeformerEditorToolkit.cpp
@@ -60,16 +60,7 @@ namespace UE::MLDeformer
 
 		// Create a new model if none was selected yet.
 		UMLDeformerModel* DeformerModel = DeformerAsset->GetModel();
-		if (DeformerModel == nullptr)
-		{
-			if (ModelRegistry.GetNumRegisteredModels() > 0)
-			{
-				const int32 HighestPriorityIndex = ModelRegistry.GetHighestPriorityModelIndex();
-				OnModelChanged(HighestPriorityIndex, true);
-				DeformerModel = DeformerAsset->GetModel();
-			}
-		}
-		else
+		if (DeformerModel)
 		{
 			DeformerModel->Init(DeformerAsset);
 			FMLDeformerEditorModel* EditorModel = ModelRegistry.CreateEditorModel(DeformerModel);
