[MLDeformer] Add support for multiple animations as training inputs.

diff --git a/Engine/Plugins/Animation/MLDeformer/ChaosClothGenerator/Source/ChaosClothGenerator/Private/ChaosClothGeneratorToolsMenuExtender.cpp b/Engine/Plugins/Animation/MLDeformer/ChaosClothGenerator/Source/ChaosClothGenerator/Private/ChaosClothGeneratorToolsMenuExtender.cpp
--- a/Engine/Plugins/Animation/MLDeformer/ChaosClothGenerator/Source/ChaosClothGenerator/Private/ChaosClothGeneratorToolsMenuExtender.cpp
+++ b/Engine/Plugins/Animation/MLDeformer/ChaosClothGenerator/Source/ChaosClothGenerator/Private/ChaosClothGeneratorToolsMenuExtender.cpp
@@ -49,10 +49,11 @@ namespace UE::Chaos::ClothGenerator
 				if (UMLDeformerModel* Model = EditorModel->GetModel())
 				{
 					Properties->SkeletalMeshAsset = Model->GetSkeletalMesh();
-					Properties->AnimationSequence = Model->GetAnimSequence();
+					Properties->AnimationSequence = EditorModel->GetActiveTrainingInputAnimSequence();
 					if (UMLDeformerGeomCacheModel* GeomCacheModel = Cast<UMLDeformerGeomCacheModel>(Model))
 					{
-						if (UGeometryCache* GeomCache = GeomCacheModel->GetGeometryCache())
+						const UE::MLDeformer::FMLDeformerGeomCacheEditorModel* GeomCacheEditorModel = static_cast<const UE::MLDeformer::FMLDeformerGeomCacheEditorModel*>(EditorModel);
+						if (UGeometryCache* GeomCache = GeomCacheEditorModel->GetActiveGeometryCache())
 						{
 							Properties->SimulatedCache = GeomCache;
 						}
diff --git a/Engine/Plugins/Animation/MLDeformer/ChaosClothGenerator/Source/ChaosClothGenerator/Private/SClothGeneratorWidget.h b/Engine/Plugins/Animation/MLDeformer/ChaosClothGenerator/Source/ChaosClothGenerator/Private/SClothGeneratorWidget.h
--- a/Engine/Plugins/Animation/MLDeformer/ChaosClothGenerator/Source/ChaosClothGenerator/Private/SClothGeneratorWidget.h
+++ b/Engine/Plugins/Animation/MLDeformer/ChaosClothGenerator/Source/ChaosClothGenerator/Private/SClothGeneratorWidget.h
@@ -31,7 +31,5 @@ namespace UE::Chaos::ClothGenerator
 	public:
 		static TSharedRef<IDetailCustomization> MakeInstance();
 		virtual void CustomizeDetails(IDetailLayoutBuilder& DetailBuilder) override;
-	private:
-		TWeakObjectPtr<UClothGeneratorProperties> GetProperties(IDetailLayoutBuilder& DetailBuilder) const;
 	};
 };
\ No newline at end of file
diff --git a/Engine/Plugins/Animation/MLDeformer/MLDeformerFramework/Config/BaseMLDeformerFramework.ini b/Engine/Plugins/Animation/MLDeformer/MLDeformerFramework/Config/BaseMLDeformerFramework.ini
--- a/Engine/Plugins/Animation/MLDeformer/MLDeformerFramework/Config/BaseMLDeformerFramework.ini
+++ b/Engine/Plugins/Animation/MLDeformer/MLDeformerFramework/Config/BaseMLDeformerFramework.ini
@@ -7,8 +7,4 @@
 +PropertyRedirects=(OldName="/Script/MLDeformerFramework.MLDeformerInputInfo.CurveNameStrings",NewName="CurveNames")
 +PropertyRedirects=(OldName="/Script/MLDeformerFramework.MLDeformerMorphModel.MorphTargetErrorTolerance",NewName="MorphCompressionLevel")
 +PropertyRedirects=(OldName="/Script/MLDeformerFramework.MLDeformerMorphModel.MorphTargetDeltaThreshold",NewName="MorphDeltaZeroThreshold")
-+PropertyRedirects=(OldName="/Script/MLDeformerFramework.MLDeformerMorphModelVizSettings.MorphTargetDeltaThreshold",NewName="MorphTargetDeltaThreshold_DEPRECATED")
 
-; Deprecated model properties.
-+PropertyRedirects=(OldName="/Script/MLDeformerFramework.MLDeformerModel.bIncludeBones",NewName="bIncludeBones_DEPRECATED")
-+PropertyRedirects=(OldName="/Script/MLDeformerFramework.MLDeformerModel.bIncludeCurves",NewName="bIncludeCurves_DEPRECATED")
diff --git a/Engine/Plugins/Animation/MLDeformer/MLDeformerFramework/Source/MLDeformerFramework/Private/MLDeformerGeomCacheHelpers.cpp b/Engine/Plugins/Animation/MLDeformer/MLDeformerFramework/Source/MLDeformerFramework/Private/MLDeformerGeomCacheHelpers.cpp
--- a/Engine/Plugins/Animation/MLDeformer/MLDeformerFramework/Source/MLDeformerFramework/Private/MLDeformerGeomCacheHelpers.cpp
+++ b/Engine/Plugins/Animation/MLDeformer/MLDeformerFramework/Source/MLDeformerFramework/Private/MLDeformerGeomCacheHelpers.cpp
@@ -349,7 +349,7 @@ namespace UE::MLDeformer
 		return NumGeomCacheImportedVerts;
 	}
 
-	FText GetGeomCacheAnimSequenceErrorText(UGeometryCache* InGeomCache, UAnimSequence* InAnimSequence)
+	FText GetGeomCacheAnimSequenceErrorText(const UGeometryCache* InGeomCache, const UAnimSequence* InAnimSequence)
 	{
 		FText Result;
 
diff --git a/Engine/Plugins/Animation/MLDeformer/MLDeformerFramework/Source/MLDeformerFramework/Private/MLDeformerGeomCacheModel.cpp b/Engine/Plugins/Animation/MLDeformer/MLDeformerFramework/Source/MLDeformerFramework/Private/MLDeformerGeomCacheModel.cpp
--- a/Engine/Plugins/Animation/MLDeformer/MLDeformerFramework/Source/MLDeformerFramework/Private/MLDeformerGeomCacheModel.cpp
+++ b/Engine/Plugins/Animation/MLDeformer/MLDeformerFramework/Source/MLDeformerFramework/Private/MLDeformerGeomCacheModel.cpp
@@ -17,21 +17,56 @@ void UMLDeformerGeomCacheModel::Serialize(FArchive& Archive)
 	#if WITH_EDITOR
 		if (Archive.IsSaving() && Archive.IsCooking())
 		{
-			GeometryCache = nullptr;
+			GeometryCache_DEPRECATED = nullptr;
+			for (FMLDeformerGeomCacheTrainingInputAnim& Anim : TrainingInputAnims)
+			{
+				Anim.SetGeometryCache(nullptr);
+				Anim.SetAnimSequence(nullptr);
+			}
 		}
 	#endif
 
 	Super::Serialize(Archive);
 }
 
+void UMLDeformerGeomCacheModel::PostLoad()
+{
+	Super::PostLoad();
+
+	#if WITH_EDITORONLY_DATA
+		// Handle backward compatibility by converting the asset that just has a single asset, into an entry
+		// inside the multiple training input animation array.
+		if (!AnimSequence_DEPRECATED.IsNull() || !GeometryCache_DEPRECATED.IsNull())
+		{
+			check(TrainingInputAnims.IsEmpty());	// This should only happen when we don't use the training input anims array yet.
+
+			// Add an item to the training input anims list.
+			TrainingInputAnims.AddDefaulted();
+			FMLDeformerGeomCacheTrainingInputAnim& AnimEntry = TrainingInputAnims.Last();
+			AnimEntry.SetGeometryCache(GeometryCache_DEPRECATED);
+			AnimEntry.SetAnimSequence(AnimSequence_DEPRECATED);
+
+			// Set the old unused properties to nullptr, as we don't want to use those anymore.
+			GeometryCache_DEPRECATED = nullptr;
+			AnimSequence_DEPRECATED = nullptr;
+		}
+	#endif
+}
+
 #if WITH_EDITOR
 void UMLDeformerGeomCacheModel::UpdateNumTargetMeshVertices()
 {
-	UGeometryCache* GeomCache = GetGeometryCache();
-	if (GeomCache)
+	for (FMLDeformerGeomCacheTrainingInputAnim& Anim : TrainingInputAnims)
 	{
-		SetNumTargetMeshVerts(UE::MLDeformer::ExtractNumImportedGeomCacheVertices(GeomCache));
+		UGeometryCache* GeomCache = Anim.GetGeometryCache();
+		if (GeomCache && Anim.IsEnabled())
+		{
+			SetNumTargetMeshVerts(UE::MLDeformer::ExtractNumImportedGeomCacheVertices(GeomCache));
+			return;
+		}
 	}
+
+	SetNumTargetMeshVerts(0);
 }
 
 UMLDeformerGeomCacheVizSettings* UMLDeformerGeomCacheModel::GetGeomCacheVizSettings() const
@@ -41,6 +76,18 @@ UMLDeformerGeomCacheVizSettings* UMLDeformerGeomCacheModel::GetGeomCacheVizSetti
 #endif // WITH_EDITOR
 
 #if WITH_EDITORONLY_DATA
+bool UMLDeformerGeomCacheModel::HasTrainingGroundTruth() const
+{
+	for (const FMLDeformerGeomCacheTrainingInputAnim& Anim : TrainingInputAnims)
+	{
+		if (Anim.GetGeometryCache())
+		{
+			return true;
+		}
+	}
+	return false;
+}
+
 void UMLDeformerGeomCacheModel::SampleGroundTruthPositions(float SampleTime, TArray<FVector3f>& OutPositions)
 {
 	const UMLDeformerGeomCacheVizSettings* GeomCacheVizSettings = GetGeomCacheVizSettings();
diff --git a/Engine/Plugins/Animation/MLDeformer/MLDeformerFramework/Source/MLDeformerFramework/Private/MLDeformerGeomCacheTrainingInputAnim.cpp b/Engine/Plugins/Animation/MLDeformer/MLDeformerFramework/Source/MLDeformerFramework/Private/MLDeformerGeomCacheTrainingInputAnim.cpp
new file mode 100644
--- /dev/null
+++ b/Engine/Plugins/Animation/MLDeformer/MLDeformerFramework/Source/MLDeformerFramework/Private/MLDeformerGeomCacheTrainingInputAnim.cpp
@@ -0,0 +1,61 @@
+#include "MLDeformerGeomCacheTrainingInputAnim.h"
+
+#if WITH_EDITORONLY_DATA
+#include "GeometryCache.h"
+#include "Animation/AnimSequence.h"
+
+int32 FMLDeformerGeomCacheTrainingInputAnim::GetNumFramesToSample() const
+{
+	const UGeometryCache* GeomCache = GetGeometryCache();
+	const UAnimSequence* AnimSeq = GetAnimSequence();
+	if (!IsEnabled() || !AnimSeq || !GeomCache)
+	{
+		return 0;
+	}
+	
+	const int32 NumAnimFrames = GeomCache->GetEndFrame() - GeomCache->GetStartFrame() + 1;
+	int32 NumFrames = 0;
+	if (!GetUseCustomRange())
+	{
+		NumFrames = NumAnimFrames;
+	}
+	else
+	{
+		int32 SafeStartFrame = FMath::Min<int32>(GetStartFrame(), GetEndFrame());
+		int32 SafeEndFrame = FMath::Max<int32>(GetStartFrame(), GetEndFrame());
+
+		if (SafeStartFrame >= GeomCache->GetEndFrame() - GeomCache->GetStartFrame())
+		{
+			return 0;
+		}
+
+		if (SafeEndFrame >= GeomCache->GetEndFrame() - GeomCache->GetStartFrame())
+		{
+			SafeEndFrame = GeomCache->GetEndFrame() - GeomCache->GetStartFrame();
+		}
+
+		const int32 NumManualFrames = SafeEndFrame - SafeStartFrame + 1;
+		NumFrames = FMath::Min<int32>(NumAnimFrames, NumManualFrames);
+	}
+
+	return NumFrames;
+}
+
+bool FMLDeformerGeomCacheTrainingInputAnim::IsValid() const
+{
+	return (FMLDeformerTrainingInputAnim::IsValid() && GetGeometryCache() != nullptr);
+}
+
+int32 FMLDeformerGeomCacheTrainingInputAnim::ExtractNumAnimFrames() const
+{ 
+	const UAnimSequence* Anim = GetAnimSequence();
+	const UGeometryCache* GeomCache = GetGeometryCache();
+	if (!Anim || !GeomCache)
+	{
+		return 0;
+	}
+
+	return GeomCache->GetEndFrame() - GeomCache->GetStartFrame() + 1;
+}
+
+#endif // #if WITH_EDITORONLY_DATA
diff --git a/Engine/Plugins/Animation/MLDeformer/MLDeformerFramework/Source/MLDeformerFramework/Private/MLDeformerModel.cpp b/Engine/Plugins/Animation/MLDeformer/MLDeformerFramework/Source/MLDeformerFramework/Private/MLDeformerModel.cpp
--- a/Engine/Plugins/Animation/MLDeformer/MLDeformerFramework/Source/MLDeformerFramework/Private/MLDeformerModel.cpp
+++ b/Engine/Plugins/Animation/MLDeformer/MLDeformerFramework/Source/MLDeformerFramework/Private/MLDeformerModel.cpp
@@ -76,7 +76,7 @@ void UMLDeformerModel::Serialize(FArchive& Archive)
 		{
 			if (Archive.IsCooking())
 			{
-				AnimSequence = nullptr;
+				AnimSequence_DEPRECATED = nullptr;
 				VizSettings = nullptr;
 			}
 
diff --git a/Engine/Plugins/Animation/MLDeformer/MLDeformerFramework/Source/MLDeformerFramework/Public/MLDeformerGeomCacheHelpers.h b/Engine/Plugins/Animation/MLDeformer/MLDeformerFramework/Source/MLDeformerFramework/Public/MLDeformerGeomCacheHelpers.h
--- a/Engine/Plugins/Animation/MLDeformer/MLDeformerFramework/Source/MLDeformerFramework/Public/MLDeformerGeomCacheHelpers.h
+++ b/Engine/Plugins/Animation/MLDeformer/MLDeformerFramework/Source/MLDeformerFramework/Public/MLDeformerGeomCacheHelpers.h
@@ -42,7 +42,7 @@ namespace UE::MLDeformer
 	 * @param InAnimSequence The animation sequence object.
 	 * @return The error text in case there are errors, or an empty text object in case there were no errors.
 	 */
-	MLDEFORMERFRAMEWORK_API FText GetGeomCacheAnimSequenceErrorText(UGeometryCache* InGeomCache, UAnimSequence* InAnimSequence);
+	MLDEFORMERFRAMEWORK_API FText GetGeomCacheAnimSequenceErrorText(const UGeometryCache* InGeomCache, const UAnimSequence* InAnimSequence);
 
 	/**
 	 * Check for any mesh related issues between a skeletal mesh and geometry cache.
diff --git a/Engine/Plugins/Animation/MLDeformer/MLDeformerFramework/Source/MLDeformerFramework/Public/MLDeformerGeomCacheModel.h b/Engine/Plugins/Animation/MLDeformer/MLDeformerFramework/Source/MLDeformerFramework/Public/MLDeformerGeomCacheModel.h
--- a/Engine/Plugins/Animation/MLDeformer/MLDeformerFramework/Source/MLDeformerFramework/Public/MLDeformerGeomCacheModel.h
+++ b/Engine/Plugins/Animation/MLDeformer/MLDeformerFramework/Source/MLDeformerFramework/Public/MLDeformerGeomCacheModel.h
@@ -7,6 +7,7 @@
 #include "MLDeformerVizSettings.h"
 #include "MLDeformerInputInfo.h"
 #include "MLDeformerGeomCacheHelpers.h"
+#include "MLDeformerGeomCacheTrainingInputAnim.h"
 #include "GeometryCache.h"
 #include "UObject/ObjectPtr.h"
 #include "UObject/SoftObjectPtr.h"
@@ -28,11 +29,12 @@ class MLDEFORMERFRAMEWORK_API UMLDeformerGeomCacheModel
 public:
 	// UObject overrides.
 	virtual void Serialize(FArchive& Archive) override;
+	virtual void PostLoad() override;
 	// ~END UObject overrides.
 
 	// UMLDeformerModel overrides.
 #if WITH_EDITORONLY_DATA
-	virtual bool HasTrainingGroundTruth() const override	{ return (GetGeometryCache() != nullptr); }
+	virtual bool HasTrainingGroundTruth() const override;
 	virtual void SampleGroundTruthPositions(float SampleTime, TArray<FVector3f>& OutPositions) override;
 #endif
 #if WITH_EDITOR
@@ -52,21 +54,24 @@ public:
 	 * This is our training target.
 	 * @return A pointer to the geometry cache, in read-only mode.
 	 */
-	const UGeometryCache* GetGeometryCache() const			{ return GeometryCache.LoadSynchronous(); }
+	UE_DEPRECATED(5.4, "This method will be removed. Please use GetTrainingInputAnims() instead.")
+	const UGeometryCache* GetGeometryCache() const			{ return nullptr; }
 
 	/** 
 	 * Get the geometry cache that represents the target deformation.
 	 * This is our training target.
 	 * @return A pointer to the geometry cache.
 	 */
-	UGeometryCache* GetGeometryCache()						{ return GeometryCache.LoadSynchronous(); }
+	UE_DEPRECATED(5.4, "This method will be removed. Please use GetTrainingInputAnims() instead.")
+	UGeometryCache* GetGeometryCache()						{ return nullptr; }
 
 	/**
 	 * Set the geometry cache object to use for training.
 	 * Keep in mind that the editor still needs to handle a change of this property for things to be initialized correctly.
 	 * @param GeomCache The geometry cache to use for training.
 	 */
-	void SetGeometryCache(UGeometryCache* GeomCache)		{ GeometryCache = GeomCache; }
+	UE_DEPRECATED(5.4, "This method will be removed. Please use GetTrainingInputAnims() instead.")
+	void SetGeometryCache(UGeometryCache* GeomCache)		{}
 
 	/**
 	 * Get the mapping between geometry cache tracks and meshes inside the skeletal mesh.
@@ -77,15 +82,24 @@ public:
 	TArray<UE::MLDeformer::FMLDeformerGeomCacheMeshMapping>& GetGeomCacheMeshMappings()				{ return MeshMappings; }
 	const TArray<UE::MLDeformer::FMLDeformerGeomCacheMeshMapping>& GetGeomCacheMeshMappings() const { return MeshMappings; }
 
+	TArray<FMLDeformerGeomCacheTrainingInputAnim>& GetTrainingInputAnims()							{ return TrainingInputAnims; }
+	const TArray<FMLDeformerGeomCacheTrainingInputAnim>& GetTrainingInputAnims() const				{ return TrainingInputAnims; }
+
 	// Get property names.
-	static FName GetGeometryCachePropertyName()				{ return GET_MEMBER_NAME_CHECKED(UMLDeformerGeomCacheModel, GeometryCache); }
+	UE_DEPRECATED(5.4, "This method will be removed.")
+	static FName GetGeometryCachePropertyName()				{ return GET_MEMBER_NAME_CHECKED(UMLDeformerGeomCacheModel, GeometryCache_DEPRECATED); }
+
+	static FName GetTrainingInputAnimsPropertyName()		{ return GET_MEMBER_NAME_CHECKED(UMLDeformerGeomCacheModel, TrainingInputAnims); }
 
 private:
 	/** The mappings between the geometry cache tracks and skeletal mesh imported meshes. */
 	TArray<UE::MLDeformer::FMLDeformerGeomCacheMeshMapping> MeshMappings;
 
 	/** The geometry cache that represents the target deformations. */
+	UPROPERTY(meta=(DeprecatedProperty, DeprecationMessage="Use the training input anims instead."))
+	TSoftObjectPtr<UGeometryCache> GeometryCache_DEPRECATED;
+
 	UPROPERTY(EditAnywhere, Category = "Target Mesh")
-	TSoftObjectPtr<UGeometryCache> GeometryCache = nullptr;
+	TArray<FMLDeformerGeomCacheTrainingInputAnim> TrainingInputAnims;
 #endif // WITH_EDITORONLY_DATA
 };
diff --git a/Engine/Plugins/Animation/MLDeformer/MLDeformerFramework/Source/MLDeformerFramework/Public/MLDeformerGeomCacheTrainingInputAnim.h b/Engine/Plugins/Animation/MLDeformer/MLDeformerFramework/Source/MLDeformerFramework/Public/MLDeformerGeomCacheTrainingInputAnim.h
new file mode 100644
--- /dev/null
+++ b/Engine/Plugins/Animation/MLDeformer/MLDeformerFramework/Source/MLDeformerFramework/Public/MLDeformerGeomCacheTrainingInputAnim.h
@@ -0,0 +1,40 @@
+// Copyright Epic Games, Inc. All Rights Reserved.
+
+#pragma once
+
+#include "CoreTypes.h"
+#include "UObject/SoftObjectPtr.h"
+#include "MLDeformerTrainingInputAnim.h"
+#include "GeometryCache.h"
+#include "MLDeformerGeomCacheTrainingInputAnim.generated.h"
+
+/**
+ * An animation that is input to the training process, which has a geometry cache as target.
+ */
+USTRUCT()
+struct MLDEFORMERFRAMEWORK_API FMLDeformerGeomCacheTrainingInputAnim
+	: public FMLDeformerTrainingInputAnim
+{
+	GENERATED_BODY()
+
+public:
+#if WITH_EDITORONLY_DATA
+	// FMLDeformerTrainingInputAnim overrides
+	virtual int32 GetNumFramesToSample() const override;
+	virtual bool IsValid() const override;
+	virtual int32 ExtractNumAnimFrames() const override;
+	// ~END FMLDeformerTrainingInputAnim overrides
+
+	void SetGeometryCache(TSoftObjectPtr<UGeometryCache> GeomCache)	{ GeometryCache = GeomCache; }
+
+	const UGeometryCache* GetGeometryCache() const					{ return GeometryCache.LoadSynchronous(); }
+	UGeometryCache* GetGeometryCache()								{ return GeometryCache.LoadSynchronous(); }
+
+	static FName GetGeomCachePropertyName()							{ return GET_MEMBER_NAME_CHECKED(FMLDeformerGeomCacheTrainingInputAnim, GeometryCache); }
+
+private:
+	/** The geometry cache which contains the target deformations. */
+	UPROPERTY(EditAnywhere, Category = "Settings", meta = (DisplayPriority = 10))
+	TSoftObjectPtr<UGeometryCache> GeometryCache = nullptr;
+#endif	// #if WITH_EDITORONLY_DATA
+};
diff --git a/Engine/Plugins/Animation/MLDeformer/MLDeformerFramework/Source/MLDeformerFramework/Public/MLDeformerModel.h b/Engine/Plugins/Animation/MLDeformer/MLDeformerFramework/Source/MLDeformerFramework/Public/MLDeformerModel.h
--- a/Engine/Plugins/Animation/MLDeformer/MLDeformerFramework/Source/MLDeformerFramework/Public/MLDeformerModel.h
+++ b/Engine/Plugins/Animation/MLDeformer/MLDeformerFramework/Source/MLDeformerFramework/Public/MLDeformerModel.h
@@ -20,6 +20,7 @@ class UMLDeformerModelInstance;
 class UMLDeformerComponent;
 class UMLDeformerInputInfo;
 
+
 /** The channel to get the mask data from. */
 UENUM()
 enum class EMLDeformerMaskChannel : uint8
@@ -431,15 +432,19 @@ public:
 	 * Each frame of this anim sequence will contain a training pose.
 	 * @return A pointer to the animation sequence used for training.
 	 */
-	const UAnimSequence* GetAnimSequence() const				{ return AnimSequence.LoadSynchronous();  }
-	UAnimSequence* GetAnimSequence()							{ return AnimSequence.LoadSynchronous(); }
+	UE_DEPRECATED(5.4, "This method will be removed. Please look at FMLDeformerEditorModel::GetActiveTrainingInputAnimIndex().")
+	const UAnimSequence* GetAnimSequence() const				{ return nullptr; }
+
+	UE_DEPRECATED(5.4, "This method will be removed. Please look at FMLDeformerEditorModel::GetActiveTrainingInputAnimIndex().")
+	UAnimSequence* GetAnimSequence()							{ return nullptr; }
 
 	/**
 	 * Set the animation sequence object to use for training.
 	 * Keep in mind that the editor still needs to handle a change of this property for things to be initialized correctly.
 	 * @param AnimSeq The animation sequence to use for training.
 	 */
-	void SetAnimSequence(UAnimSequence* AnimSeq)				{ AnimSequence = AnimSeq; }
+	UE_DEPRECATED(5.4, "This method will be removed.")
+	void SetAnimSequence(UAnimSequence* AnimSeq)				{}
 
 	/**
 	 * Get the maximum number of training frames to use during training.
@@ -539,8 +544,10 @@ public:
 	void SetVizSettings(UMLDeformerVizSettings* VizSettingsObject)			{ VizSettings = VizSettingsObject; }
 
 	// Get property names.
+	UE_DEPRECATED(5.4, "This method will be removed.")
+	static FName GetAnimSequencePropertyName()			{ return GET_MEMBER_NAME_CHECKED(UMLDeformerModel, AnimSequence_DEPRECATED); }
+
 	static FName GetSkeletalMeshPropertyName()			{ return GET_MEMBER_NAME_CHECKED(UMLDeformerModel, SkeletalMesh); }
-	static FName GetAnimSequencePropertyName()			{ return GET_MEMBER_NAME_CHECKED(UMLDeformerModel, AnimSequence); }
 	static FName GetAlignmentTransformPropertyName()	{ return GET_MEMBER_NAME_CHECKED(UMLDeformerModel, AlignmentTransform); }
 	static FName GetBoneIncludeListPropertyName()		{ return GET_MEMBER_NAME_CHECKED(UMLDeformerModel, BoneIncludeList); }
 	static FName GetCurveIncludeListPropertyName()		{ return GET_MEMBER_NAME_CHECKED(UMLDeformerModel, CurveIncludeList); }
@@ -552,7 +559,6 @@ public:
 
 	UE_DEPRECATED(5.3, "This property has been removed and shouldn't be used anymore.")
 	static FName GetShouldIncludeCurvesPropertyName()	{ return GET_MEMBER_NAME_CHECKED(UMLDeformerModel, bIncludeCurves_DEPRECATED); }
-
 #endif	// #if WITH_EDITORONLY_DATA
 
 protected:
@@ -609,6 +615,15 @@ protected:
 	uint64 GPUMemUsageInBytes = 0;
 #endif
 
+#if WITH_EDITORONLY_DATA
+	/**
+	 * The animation sequence to apply to the base mesh. This has to match the animation of the target mesh's geometry cache. 
+	 * Internally we force the Interpolation property for this motion to be "Step".
+	 */
+	UPROPERTY(meta=(DeprecatedProperty, DeprecationMessage="Use the training input anims instead."))
+	TSoftObjectPtr<UAnimSequence> AnimSequence_DEPRECATED;
+#endif
+
 private:
 	/** The deformer asset that this model is part of. */
 	TObjectPtr<UMLDeformerAsset> DeformerAsset = nullptr;
@@ -670,15 +685,8 @@ private:
 	UPROPERTY()
 	bool bIncludeCurves_DEPRECATED = false;
 
-	/**
-	 * The animation sequence to apply to the base mesh. This has to match the animation of the target mesh's geometry cache. 
-	 * Internally we force the Interpolation property for this motion to be "Step".
-	 */
-	UPROPERTY(EditAnywhere, Category = "Base Mesh")
-	TSoftObjectPtr<UAnimSequence> AnimSequence = nullptr;
-
 	/** The transform that aligns the Geometry Cache to the SkeletalMesh. This will mostly apply some scale and a rotation, but no translation. */
-	UPROPERTY(EditAnywhere, Category = "Target Mesh")
+	UPROPERTY(EditAnywhere, Category = "Inputs", DisplayName = "Target Alignment Transform")
 	FTransform AlignmentTransform = FTransform::Identity;
 
 	/** The bones to include during training. When none are provided, all bones of the Skeleton will be included. */
@@ -690,14 +698,14 @@ private:
 	TArray<FMLDeformerCurveReference> CurveIncludeList;
 
 	/** The maximum numer of training frames (samples) to train on. Use this to train on a sub-section of your full training data. */
-	UPROPERTY(EditAnywhere, Category = "Inputs", meta = (ClampMin = "1"))
+	UPROPERTY(EditAnywhere, Category = "Training Settings", meta = (ClampMin = "1"))
 	int32 MaxTrainingFrames = 1000000;
 
 	/**
 	 * Sometimes there can be some vertices that cause some issues that cause deltas to be very long. We can ignore these deltas by setting a cutoff value. 
 	 * Deltas that are longer than the cutoff value (in units), will be ignored and set to zero length. 
 	 */
-	UPROPERTY(EditAnywhere, Category = "Inputs", meta = (ClampMin = "0.01", ForceUnits="cm"))
+	UPROPERTY(EditAnywhere, Category = "Training Settings", meta = (ClampMin = "0.01", ForceUnits="cm"))
 	float DeltaCutoffLength = 30.0f;
 #endif
 };
diff --git a/Engine/Plugins/Animation/MLDeformer/MLDeformerFramework/Source/MLDeformerFramework/Public/MLDeformerTrainingInputAnim.h b/Engine/Plugins/Animation/MLDeformer/MLDeformerFramework/Source/MLDeformerFramework/Public/MLDeformerTrainingInputAnim.h
new file mode 100644
--- /dev/null
+++ b/Engine/Plugins/Animation/MLDeformer/MLDeformerFramework/Source/MLDeformerFramework/Public/MLDeformerTrainingInputAnim.h
@@ -0,0 +1,86 @@
+// Copyright Epic Games, Inc. All Rights Reserved.
+
+#pragma once
+
+#include "CoreTypes.h"
+#include "UObject/Object.h"
+#include "UObject/SoftObjectPtr.h"
+#include "Animation/AnimSequence.h"
+#include "MLDeformerTrainingInputAnim.generated.h"
+
+/**
+ * An animation that is input to the training process.
+ * This is a base struct that only contains the animation sequence.
+ * For example a geometry cache based model should inherit from this struct and add a geom cache member property.
+ */
+USTRUCT()
+struct MLDEFORMERFRAMEWORK_API FMLDeformerTrainingInputAnim
+{
+	GENERATED_BODY()
+
+public:
+	virtual ~FMLDeformerTrainingInputAnim() = default;
+
+#if WITH_EDITORONLY_DATA
+	/** 
+	 * Get the number of frames that we are sampling inside this animation. This can be less than the value returned by ExtractNumAnimFrames().
+	 * This might be the custom frame range, so a subsection of the animation frames.
+	 */
+	virtual int32 GetNumFramesToSample() const					{ ensureMsgf(false, TEXT("Please overload the FMLDeformerTrainingInputAnim::GetNumFramesToSample() method.")); return 0; }
+
+	/** Is this input animation valid? Returns true if we can sample from it. This typically checks if the pointers to the source and target animations are valid. */
+	virtual bool IsValid() const								{ return GetAnimSequence() != nullptr; }
+
+	/** How many frames does this input animation have? This is the maximum number of frames we can sample. */
+	virtual int32 ExtractNumAnimFrames() const					{ const UAnimSequence* Anim = GetAnimSequence(); return Anim ? Anim->GetNumberOfSampledKeys() : 0; }
+
+	void SetEnabled(bool bEnable)								{ bEnabled = bEnable; }
+	void SetAnimSequence(TSoftObjectPtr<UAnimSequence> Seq)		{ AnimSequence = Seq; }
+	void SetUseCustomRange(bool bCustomRange)					{ bUseCustomRange = bCustomRange; }
+	void SetStartFrame(int32 FrameIndex)						{ StartFrame = FrameIndex; }
+	void SetEndFrame(int32 FrameIndex)							{ EndFrame = FrameIndex; }
+
+	bool IsEnabled() const										{ return bEnabled; }
+	bool GetUseCustomRange() const								{ return bUseCustomRange; }
+	int32 GetStartFrame() const									{ return StartFrame; }
+	int32 GetEndFrame() const									{ return EndFrame; }
+	const UAnimSequence* GetAnimSequence() const				{ return AnimSequence.LoadSynchronous(); }
+	UAnimSequence* GetAnimSequence()							{ return AnimSequence.LoadSynchronous(); }
+
+	static FName GetAnimSequencePropertyName()					{ return GET_MEMBER_NAME_CHECKED(FMLDeformerTrainingInputAnim, AnimSequence); }
+	static FName GetEnabledPropertyName()						{ return GET_MEMBER_NAME_CHECKED(FMLDeformerTrainingInputAnim, bEnabled); }
+	static FName GetUseCustomRangePropertyName()				{ return GET_MEMBER_NAME_CHECKED(FMLDeformerTrainingInputAnim, bUseCustomRange); }
+	static FName GetStartFramePropertyName()					{ return GET_MEMBER_NAME_CHECKED(FMLDeformerTrainingInputAnim, StartFrame); }
+	static FName GetEndFramePropertyName()						{ return GET_MEMBER_NAME_CHECKED(FMLDeformerTrainingInputAnim, EndFrame); }
+
+private:
+	/** The animation sequence. */
+	UPROPERTY(EditAnywhere, Category = "Settings", meta = (DisplayPriority = 0))	// Show as first property.
+	TSoftObjectPtr<UAnimSequence> AnimSequence = nullptr;
+
+	/** On default all frames are included, unless we specify a custom frame include range. */
+	UPROPERTY(EditAnywhere, Category = "Frame Settings")
+	bool bUseCustomRange = false;
+
+	/** The start frame number of the range of included frames. Only used when not all frames are included. */
+	UPROPERTY(EditAnywhere, Category = "Frame Settings", meta = (ClampMin="0", EditCondition="bUseCustomRange"))
+	int32 StartFrame = 0;
+
+	/** The end frame number of the range of included frames. Only used when not all frames are included. */
+	UPROPERTY(EditAnywhere, Category = "Frame Settings", meta = (ClampMin="0", EditCondition="bUseCustomRange"))
+	int32 EndFrame = 0;
+
+	/** Is this animation enabled? If not, it is excluded from the training process. */
+	UPROPERTY(EditAnywhere, Category = "Settings", meta = (DisplayPriority = MAX_int32))	// Show as last property.
+	bool bEnabled = true;
+#endif	// #if WITH_EDITORONLY_DATA
+};
+
+struct MLDEFORMERFRAMEWORK_API FMLDeformerTrainingInputAnimName
+{
+	/** The name to display in the timeline combobox. */
+	FString Name;
+
+	/** The training input anim index. This is the index inside the FMLDeformerEditorModel::GetTrainingInputAnim(Index) method. */
+	int32 TrainingInputAnimIndex = INDEX_NONE;
+};
diff --git a/Engine/Plugins/Animation/MLDeformer/MLDeformerFramework/Source/MLDeformerFrameworkEditor/Private/MLDeformerEditorActor.cpp b/Engine/Plugins/Animation/MLDeformer/MLDeformerFramework/Source/MLDeformerFrameworkEditor/Private/MLDeformerEditorActor.cpp
--- a/Engine/Plugins/Animation/MLDeformer/MLDeformerFramework/Source/MLDeformerFrameworkEditor/Private/MLDeformerEditorActor.cpp
+++ b/Engine/Plugins/Animation/MLDeformer/MLDeformerFramework/Source/MLDeformerFrameworkEditor/Private/MLDeformerEditorActor.cpp
@@ -23,12 +23,7 @@ namespace UE::MLDeformer
 	{
 		if (LabelComponent)
 		{
-			Actor->RemoveOwnedComponent(LabelComponent);
-		}
-
-		if (SkeletalMeshComponent)
-		{
-			Actor->RemoveOwnedComponent(SkeletalMeshComponent);
+			LabelComponent->UnregisterComponent();
 		}
 	}
 
diff --git a/Engine/Plugins/Animation/MLDeformer/MLDeformerFramework/Source/MLDeformerFrameworkEditor/Private/MLDeformerEditorModel.cpp b/Engine/Plugins/Animation/MLDeformer/MLDeformerFramework/Source/MLDeformerFrameworkEditor/Private/MLDeformerEditorModel.cpp
--- a/Engine/Plugins/Animation/MLDeformer/MLDeformerFramework/Source/MLDeformerFrameworkEditor/Private/MLDeformerEditorModel.cpp
+++ b/Engine/Plugins/Animation/MLDeformer/MLDeformerFramework/Source/MLDeformerFrameworkEditor/Private/MLDeformerEditorModel.cpp
@@ -14,8 +14,11 @@
 #include "MLDeformerInputInfo.h"
 #include "MLDeformerSampler.h"
 #include "MLDeformerModelInstance.h"
+#include "MLDeformerTrainingInputAnim.h"
 #include "AnimationEditorPreviewActor.h"
 #include "AnimationEditorViewportClient.h"
+#include "EditorModeManager.h"
+#include "Selection.h"
 #include "AnimPreviewInstance.h"
 #include "Animation/MeshDeformer.h"
 #include "Animation/DebugSkelMeshComponent.h"
@@ -39,6 +42,7 @@
 #include "Animation/AnimData/IAnimationDataModel.h"
 #include "Engine/SkinnedAssetCommon.h"
 #include "SMLDeformerInputWidget.h"
+#include "SMLDeformerTimeline.h"
 
 #define LOCTEXT_NAMESPACE "MLDeformerEditorModel"
 
@@ -63,6 +67,15 @@ namespace UE::MLDeformer
 
 		ClearWorld();
 
+		// Delete the samplers.
+		for (int32 Index = 0; Index < Samplers.Num(); ++Index)
+		{
+			Samplers[Index].Reset();
+		}
+		Samplers.Empty();
+
+		Editor->GetPersonaToolkit()->GetPreviewScene()->DeselectAll();
+
 		FMLDeformerEditorModule& EditorModule = FModuleManager::LoadModuleChecked<FMLDeformerEditorModule>("MLDeformerFrameworkEditor");
 		EditorModule.GetModelRegistry().RemoveEditorModelInstance(this);
 	}
@@ -78,9 +91,7 @@ namespace UE::MLDeformer
 		EditorInputInfo = Model->CreateInputInfo();
 		check(EditorInputInfo);
 
-		Sampler = CreateSampler();
-		check(Sampler);
-		Sampler->Init(this);
+		CreateSamplers();
 
 		ViewRange = TRange<double>(0.0, 100.0);
 		WorkingRange = TRange<double>(0.0, 100.0);
@@ -105,17 +116,31 @@ namespace UE::MLDeformer
 		);
 	}
 
-	void FMLDeformerEditorModel::OnObjectModified(UObject* Object)
+	void FMLDeformerEditorModel::CreateSamplers()
 	{
-		if (Model->GetSkeletalMesh() == Object || Model->GetInputInfo()->GetSkeletalMesh() == Object)
+		// First get rid of the existing samplers.
+		for (int32 Index = 0; Index < Samplers.Num(); ++Index)
 		{
-			UE_LOG(LogMLDeformer, Display, TEXT("Detected a modification in skeletal mesh %s, reinitializing inputs."), *Object->GetName());
-			bNeedsAssetReinit = true;
+			Samplers[Index].Reset();
+		}
+		Samplers.Reset();
+
+		// Create a sampler for each training input anim entry.
+		const int32 NumAnims = GetNumTrainingInputAnims();
+		for (int32 AnimIndex = 0; AnimIndex < NumAnims; ++AnimIndex)
+		{
+			TSharedPtr<FMLDeformerSampler> Sampler = CreateSamplerObject();
+			check(Sampler.IsValid());
+			Sampler->Init(this, AnimIndex);
+			Samplers.Add(Sampler);
 		}
+	}
 
-		if (Model->GetAnimSequence() == Object)
+	void FMLDeformerEditorModel::OnObjectModified(UObject* Object)
+	{
+		if (Model->GetSkeletalMesh() == Object || Model->GetInputInfo()->GetSkeletalMesh() == Object)
 		{
-			UE_LOG(LogMLDeformer, Display, TEXT("Detected a modification in training anim sequence %s, reinitializing inputs."), *Object->GetName());
+			UE_LOG(LogMLDeformer, Display, TEXT("Detected a modification in skeletal mesh %s, reinitializing inputs."), *Object->GetName());
 			bNeedsAssetReinit = true;
 		}
 
@@ -142,11 +167,6 @@ namespace UE::MLDeformer
 		return Editor->GetPersonaToolkit()->GetPreviewScene()->GetWorld();
 	}
 
-	FMLDeformerSampler* FMLDeformerEditorModel::CreateSampler() const
-	{
-		return new FMLDeformerSampler();
-	}
-
 	void FMLDeformerEditorModel::CreateTrainingLinearSkinnedActor(const TSharedRef<IPersonaPreviewScene>& InPersonaPreviewScene)
 	{
 		UWorld* World = InPersonaPreviewScene->GetWorld();
@@ -168,11 +188,11 @@ namespace UE::MLDeformer
 		// NOTE: This must be done AFTER setting the AnimInstance so that the correct root anim node is loaded.
 		USkeletalMesh* Mesh = Model->GetSkeletalMesh();
 		SkelMeshComponent->SetSkeletalMesh(Mesh);
+		SkelMeshComponent->RegisterComponent();
 
 		// Update the persona scene.
 		InPersonaPreviewScene->SetActor(Actor);
 		InPersonaPreviewScene->SetPreviewMeshComponent(SkelMeshComponent);
-		InPersonaPreviewScene->AddComponent(SkelMeshComponent, FTransform::Identity);
 		InPersonaPreviewScene->SetAdditionalMeshesSelectable(false);
 		InPersonaPreviewScene->SetPreviewMesh(Mesh);
 		if (SkelMeshComponent->GetAnimInstance())
@@ -205,10 +225,10 @@ namespace UE::MLDeformer
 		UDebugSkelMeshComponent* SkelMeshComponent = NewObject<UDebugSkelMeshComponent>(Actor);
 		SkelMeshComponent->SetWireframeMeshOverlayColor(BaseWireColor);
 		SkelMeshComponent->SetSkeletalMesh(Model->GetSkeletalMesh());
-		Actor->SetRootComponent(SkelMeshComponent);
 		SkelMeshComponent->RegisterComponent();
 		SkelMeshComponent->SetVisibility(false);
 		SkelMeshComponent->MarkRenderStateDirty();
+		Actor->SetRootComponent(SkelMeshComponent);
 
 		// Register the editor actor.
 		const FLinearColor LabelColor = FMLDeformerEditorStyle::Get().GetColor("MLDeformer.BaseMesh.LabelColor");
@@ -301,7 +321,6 @@ namespace UE::MLDeformer
 		{
 			if (EditorActor)
 			{				
-				World->RemoveActor(EditorActor->GetActor(), true);
 				ActorsToDestroy.Add(EditorActor->GetActor());
 			}
 		}
@@ -311,6 +330,7 @@ namespace UE::MLDeformer
 
 		for (AActor* Actor : ActorsToDestroy)
 		{
+			World->EditorDestroyActor(Actor, true);
 			Actor->Destroy();
 		}
 	}
@@ -321,15 +341,31 @@ namespace UE::MLDeformer
 		UWorld* World = PreviewScene->GetWorld();
 
 		PreviewScene->DeselectAll();
+
+		if (GetEditor()->GetEditorModeManager().GetSelectedComponents())
+		{
+			GetEditor()->GetEditorModeManager().GetSelectedComponents()->DeselectAll();
+		}
+
+		for (FMLDeformerEditorActor* EditorActor : EditorActors)
+		{
+			USkeletalMeshComponent* SkelMeshComponent = EditorActor->GetSkeletalMeshComponent();
+			if (SkelMeshComponent)
+			{
+				PreviewScene->RemoveComponent(SkelMeshComponent);
+			}
+		}
+
 		PreviewScene->SetPreviewAnimationAsset(nullptr);
+		PreviewScene->SetPreviewMesh(nullptr);
 		PreviewScene->SetPreviewMeshComponent(nullptr);
 		PreviewScene->SetActor(nullptr);
 	}
 
 	void FMLDeformerEditorModel::ClearWorldAndPersonaPreviewScene()
 	{
-		ClearWorld();
 		ClearPersonaPreviewScene();
+		ClearWorld();
 	}
 
 	FMLDeformerEditorActor* FMLDeformerEditorModel::CreateEditorActor(const FMLDeformerEditorActor::FConstructSettings& Settings) const
@@ -369,15 +405,7 @@ namespace UE::MLDeformer
 			bNeedsAssetReinit = false;
 		}
 
-		// Force the training sequence to use Step interpolation.
-		// We do this in the Tick, as this is reset to false in the engine sometimes.
-		UAnimSequence* TrainingAnimSequence = Model->GetAnimSequence();
-		if (TrainingAnimSequence)
-		{
-			TrainingAnimSequence->Interpolation = EAnimInterpolationType::Step;
-		}
-
-		// Use raw data for everything.
+		// Make sure we use raw data for everything.
 		for (FMLDeformerEditorActor* Actor : EditorActors)
 		{
 			if (Actor && Actor->GetSkeletalMeshComponent())
@@ -390,24 +418,6 @@ namespace UE::MLDeformer
 			}
 		}
 
-		// Do the same for the test anim sequence.
-		UMLDeformerVizSettings* VizSettings = Model->GetVizSettings();
-		{
-			UAnimSequence* TestAnimSequence = VizSettings->GetTestAnimSequence();
-			if (TestAnimSequence)
-			{			
-				// Enable step interpolation when doing showing a heatmap vs ground truth.
-				if (VizSettings->HasTestGroundTruth() && VizSettings->GetShowHeatMap() && VizSettings->GetHeatMapMode() == EMLDeformerHeatMapMode::GroundTruth)
-				{
-					TestAnimSequence->Interpolation = EAnimInterpolationType::Step;
-				}
-				else
-				{
-					TestAnimSequence->Interpolation = EAnimInterpolationType::Linear;
-				}
-			}
-		}
-
 		UpdateActorTransforms();
 		UpdateLabels();
 		CheckTrainingDataFrameChanged();
@@ -419,6 +429,7 @@ namespace UE::MLDeformer
 			UMLDeformerComponent* DeformerComponent = EditorActor->GetMLDeformerComponent();
 			if (DeformerComponent)
 			{		
+				const UMLDeformerVizSettings* VizSettings = Model->GetVizSettings();
 				DeformerComponent->SetWeight(VizSettings->GetWeight());
 				DeformerComponent->SetQualityLevel(VizSettings->GetQualityLevel());
 			}
@@ -536,13 +547,6 @@ namespace UE::MLDeformer
 
 	void FMLDeformerEditorModel::OnInputAssetsChanged()
 	{
-		// Force the training sequence to use Step interpolation and sample raw animation data.
-		UAnimSequence* TrainingAnimSequence = Model->GetAnimSequence();
-		if (TrainingAnimSequence)
-		{
-			TrainingAnimSequence->Interpolation = EAnimInterpolationType::Step;
-		}
-
 		// Update the training base actor.
 		UMLDeformerVizSettings* VizSettings = Model->GetVizSettings();
 		UAnimSequence* TestAnimSequence = VizSettings->GetTestAnimSequence();
@@ -557,6 +561,7 @@ namespace UE::MLDeformer
 			{
 				GetEditor()->GetPersonaToolkit()->GetPreviewScene()->SetPreviewMesh(Model->GetSkeletalMesh());
 			}
+			UAnimSequence* TrainingAnimSequence = GetActiveTrainingInputAnimSequence();
 			SkeletalMeshComponent->SetAnimationMode(EAnimationMode::AnimationSingleNode);
 			const float CurrentPlayTime = SkeletalMeshComponent->GetPosition();
 			SkeletalMeshComponent->SetAnimation(TrainingAnimSequence);
@@ -601,10 +606,13 @@ namespace UE::MLDeformer
 	void FMLDeformerEditorModel::OnPostInputAssetChanged()
 	{
 		CurrentTrainingFrame = -1;
+		UpdateTimelineTrainingAnimList();
 		Editor->UpdateTimeSliderRange();
 		Model->UpdateCachedNumVertices();
+		UpdateNumTrainingFrames();
 		UpdateDeformerGraph();
 		RefreshMLDeformerComponents();
+		CreateSamplers();
 		UpdateIsReadyForTrainingState();
 
 		const int32 TrainingFrame = Model->GetVizSettings()->GetTrainingFrameNumber();
@@ -670,7 +678,7 @@ namespace UE::MLDeformer
 
 	double FMLDeformerEditorModel::GetTrainingTimeAtFrame(int32 FrameNumber) const
 	{
-		UAnimSequence* AnimSequence = Model->GetAnimSequence();
+		UAnimSequence* AnimSequence = GetActiveTrainingInputAnimSequence();
 		if (AnimSequence)
 		{
 			const FFrameRate FrameRate = AnimSequence->GetSamplingFrameRate();
@@ -683,7 +691,8 @@ namespace UE::MLDeformer
 
 	int32 FMLDeformerEditorModel::GetTrainingFrameAtTime(double TimeInSeconds) const
 	{
-		return Model->GetAnimSequence() ? Model->GetAnimSequence()->GetFrameAtTime(TimeInSeconds) : 0;
+		const UAnimSequence* Anim = GetActiveTrainingInputAnimSequence();
+		return Anim ? Anim->GetFrameAtTime(TimeInSeconds) : 0;
 	}
 
 	double FMLDeformerEditorModel::GetTestTimeAtFrame(int32 FrameNumber) const
@@ -734,16 +743,15 @@ namespace UE::MLDeformer
 			Model->InitVertexMap();
 			Model->InitGPUData();
 			UpdateDeformerGraph();
+			if (GetEditor())
+			{
+				GetEditor()->ZoomOnActors();
+			}
 		}
 		else if (Property->GetFName() == UMLDeformerVizSettings::GetTestAnimSequencePropertyName())
 		{
 			TriggerInputAssetChanged(true);
 		}
-		else if (Property->GetFName() == UMLDeformerModel::GetAnimSequencePropertyName())
-		{
-			TriggerInputAssetChanged();
-			SetResamplingInputOutputsNeeded(true);
-		}
 		if (Property->GetFName() == UMLDeformerModel::GetAlignmentTransformPropertyName() ||
 		    Property->GetFName() == UMLDeformerModel::GetDeltaCutoffLengthPropertyName())
 		{
@@ -802,7 +810,7 @@ namespace UE::MLDeformer
 		{
 			UpdateDeformerGraph();
 			GetEditor()->GetVizSettingsDetailsView()->ForceRefresh();
-		};
+		}
 	}
 
 	FMLDeformerEditorActor* FMLDeformerEditorModel::GetVisualizationModeBaseActor() const
@@ -918,9 +926,12 @@ namespace UE::MLDeformer
 	void FMLDeformerEditorModel::ClampCurrentTrainingFrameIndex()
 	{
 		UMLDeformerVizSettings* VizSettings = Model->GetVizSettings();
-		if (GetNumTrainingFrames() > 0)
+		const int32 ActiveAnimIndex = GetActiveTrainingInputAnimIndex();
+		if (ActiveAnimIndex != INDEX_NONE)
 		{
-			const int32 ClampedNumber = FMath::Min(VizSettings->GetTrainingFrameNumber(), GetNumTrainingFrames() - 1);
+			const int32 AnimNumFrames = GetTrainingInputAnim(ActiveAnimIndex)->ExtractNumAnimFrames();
+			const int32 MaxFrameIndex = (AnimNumFrames > 0) ? AnimNumFrames - 1 : 0;
+			const int32 ClampedNumber = FMath::Min(VizSettings->GetTrainingFrameNumber(), MaxFrameIndex);
 			VizSettings->SetTrainingFrameNumber(ClampedNumber);
 		}
 		else
@@ -954,7 +965,7 @@ namespace UE::MLDeformer
 	}
 
 	int32 FMLDeformerEditorModel::GetNumFramesForTraining() const 
-	{ 
+	{
 		return FMath::Min(GetNumTrainingFrames(), Model->GetTrainingFrameLimit());
 	}
 
@@ -1040,7 +1051,7 @@ namespace UE::MLDeformer
 		return FText();
 	}
 
-	FText FMLDeformerEditorModel::GetIncompatibleSkeletonErrorText(USkeletalMesh* InSkelMesh, UAnimSequence* InAnimSeq) const
+	FText FMLDeformerEditorModel::GetIncompatibleSkeletonErrorText(const USkeletalMesh* InSkelMesh, const UAnimSequence* InAnimSeq) const
 	{
 		FText Result;
 		if (InSkelMesh && InAnimSeq)
@@ -1073,37 +1084,24 @@ namespace UE::MLDeformer
 		return Result;
 	}
 
-	void FMLDeformerEditorModel::AddAnimatedBonesToBonesIncludeList()
+	void FMLDeformerEditorModel::AddAnimatedBonesToBonesIncludeList(const UAnimSequence* AnimSequence)
 	{
-		if (!Model->GetAnimSequence())
+		if (AnimSequence == nullptr)
 		{
 			UE_LOG(LogMLDeformer, Warning, TEXT("Cannot initialize bone list as no Anim Sequence has been picked."));
-			UpdateEditorInputInfo();
 			return;
 		}
 
-		const IAnimationDataModel* DataModel = Model->GetAnimSequence()->GetDataModel();
+		const IAnimationDataModel* DataModel = AnimSequence->GetDataModel();
 		if (!DataModel)
 		{
 			UE_LOG(LogMLDeformer, Warning, TEXT("Anim sequence has no data model."));
-			UpdateEditorInputInfo();
-			return;
-		}
-
-		if (!Model->GetSkeletalMesh())
-		{
-			UE_LOG(LogMLDeformer, Warning, TEXT("Skeletal Mesh has not been set."));
-			UpdateEditorInputInfo();
 			return;
 		}
 
+		check(Model->GetSkeletalMesh());
 		USkeleton* Skeleton = Model->GetSkeletalMesh()->GetSkeleton();
-		if (!Skeleton)
-		{
-			UE_LOG(LogMLDeformer, Warning, TEXT("Skeletal Mesh has no skeleton."));
-			UpdateEditorInputInfo();
-			return;
-		}
+		check(Skeleton);
 
 		// Iterate over all bones that are both in the skeleton and the animation.
 		TArray<FName> AnimatedBoneList;
@@ -1155,42 +1153,58 @@ namespace UE::MLDeformer
 		{
 			BoneList.AddUnique(FBoneReference(BoneName));
 		}
-
-		UpdateEditorInputInfo();
 	}
 
-	void FMLDeformerEditorModel::AddAnimatedCurvesToCurvesIncludeList()
+	void FMLDeformerEditorModel::AddAnimatedBonesToBonesIncludeList()
 	{
-		if (!Model->GetAnimSequence())
+		if (!Model->GetSkeletalMesh())
 		{
-			UE_LOG(LogMLDeformer, Warning, TEXT("Cannot initialize curve list as no Anim Sequence has been picked."));
+			UE_LOG(LogMLDeformer, Warning, TEXT("Skeletal Mesh has not been set."));
 			UpdateEditorInputInfo();
 			return;
 		}
 
-		const IAnimationDataModel* DataModel = Model->GetAnimSequence()->GetDataModel();
-		if (!DataModel)
+		USkeleton* Skeleton = Model->GetSkeletalMesh()->GetSkeleton();
+		if (!Skeleton)
 		{
-			UE_LOG(LogMLDeformer, Warning, TEXT("Anim sequence has no data model."));
+			UE_LOG(LogMLDeformer, Warning, TEXT("Skeletal Mesh has no skeleton."));
 			UpdateEditorInputInfo();
 			return;
 		}
 
-		if (!Model->GetSkeletalMesh())
+		// Add animated bones from all animations.
+		const int32 NumAnims = GetNumTrainingInputAnims();
+		for (int32 Index = 0; Index < NumAnims; ++Index)
 		{
-			UE_LOG(LogMLDeformer, Warning, TEXT("Skeletal Mesh has not been set."));
-			UpdateEditorInputInfo();
+			FMLDeformerTrainingInputAnim* Anim = GetTrainingInputAnim(Index);
+			if (Anim && Anim->IsEnabled())
+			{
+				AddAnimatedBonesToBonesIncludeList(Anim->GetAnimSequence());
+			}
+		}
+
+		UpdateEditorInputInfo();
+	}
+
+	void FMLDeformerEditorModel::AddAnimatedCurvesToCurvesIncludeList(const UAnimSequence* AnimSequence)
+	{
+		if (AnimSequence == nullptr)
+		{
+			UE_LOG(LogMLDeformer, Warning, TEXT("Cannot initialize curve list as no Anim Sequence has been picked."));
 			return;
 		}
 
-		USkeleton* Skeleton = Model->GetSkeletalMesh()->GetSkeleton();
-		if (!Skeleton)
+		const IAnimationDataModel* DataModel = AnimSequence->GetDataModel();
+		if (!DataModel)
 		{
-			UE_LOG(LogMLDeformer, Warning, TEXT("Skeletal Mesh has no skeleton."));
-			UpdateEditorInputInfo();
+			UE_LOG(LogMLDeformer, Warning, TEXT("Anim sequence has no data model."));
 			return;
 		}
 
+		check(Model->GetSkeletalMesh());
+		USkeleton* Skeleton = Model->GetSkeletalMesh()->GetSkeleton();
+		check(Skeleton);
+
 		TArray<FName> AnimatedCurveList;
 		TArray<FName> SkeletonCurveNames;
 		Skeleton->GetCurveMetaDataNames(SkeletonCurveNames);
@@ -1226,6 +1240,36 @@ namespace UE::MLDeformer
 		{
 			CurveList.AddUnique(FMLDeformerCurveReference(CurveName));
 		}
+	}
+
+	void FMLDeformerEditorModel::AddAnimatedCurvesToCurvesIncludeList()
+	{
+		if (!Model->GetSkeletalMesh())
+		{
+			UE_LOG(LogMLDeformer, Warning, TEXT("Skeletal Mesh has not been set."));
+			UpdateEditorInputInfo();
+			return;
+		}
+
+		USkeleton* Skeleton = Model->GetSkeletalMesh()->GetSkeleton();
+		if (!Skeleton)
+		{
+			UE_LOG(LogMLDeformer, Warning, TEXT("Skeletal Mesh has no skeleton."));
+			UpdateEditorInputInfo();
+			return;
+		}
+
+		// Add animated curves from all animations.
+		const int32 NumAnims = GetNumTrainingInputAnims();
+		for (int32 Index = 0; Index < NumAnims; ++Index)
+		{
+			FMLDeformerTrainingInputAnim* Anim = GetTrainingInputAnim(Index);
+			if (Anim && Anim->IsEnabled())
+			{
+				AddAnimatedCurvesToCurvesIncludeList(Anim->GetAnimSequence());
+			}
+		}
+
 		UpdateEditorInputInfo();
 	}
 
@@ -1299,12 +1343,27 @@ namespace UE::MLDeformer
 		AddAnimatedCurvesToCurvesIncludeList();
 	}
 
+	FMLDeformerSampler* FMLDeformerEditorModel::GetSamplerForActiveAnim() const
+	{
+		const int32 ActiveAnimIndex = GetActiveTrainingInputAnimIndex();
+		return (ActiveAnimIndex != INDEX_NONE && Samplers.IsValidIndex(ActiveAnimIndex)) ? Samplers[ActiveAnimIndex].Get() : nullptr;
+	}
+
+	void FMLDeformerEditorModel::SetActiveTrainingInputAnimIndex(int32 Index)
+	{
+		check(FMath::IsWithinInclusive<int32>(Index, 0, GetNumTrainingInputAnims()));
+		ActiveTrainingInputAnimIndex = Index;
+	}
+
 	void FMLDeformerEditorModel::Render(const FSceneView* View, FViewport* Viewport, FPrimitiveDrawInterface* PDI)
 	{
+		FMLDeformerSampler* Sampler = GetSamplerForActiveAnim();
+
 		// Make sure that before we render anything, that our sampler is ready.
-		if (!Sampler->IsInitialized())
+		if (Sampler && !Sampler->IsInitialized())
 		{
-			Sampler->Init(this); // This can still fail.
+			const int32 ActiveAnimIndex = GetActiveTrainingInputAnimIndex();
+			Sampler->Init(this, ActiveAnimIndex); // This can still fail.
 			Sampler->SetVertexDeltaSpace(EVertexDeltaSpace::PostSkinning);
 			if (Sampler->IsInitialized()) // If we actually managed to initialize this frame.
 			{
@@ -1313,7 +1372,7 @@ namespace UE::MLDeformer
 		}
 
 		const UMLDeformerVizSettings* VizSettings = Model->GetVizSettings();
-		if (VizSettings->GetVisualizationMode() == EMLDeformerVizMode::TrainingData)
+		if (Sampler && Sampler->IsInitialized() && VizSettings->GetVisualizationMode() == EMLDeformerVizMode::TrainingData)
 		{
 			const TArray<float>& VertexDeltas = Sampler->GetVertexDeltas();
 			const TArray<FVector3f>& LinearSkinnedPositions = Sampler->GetSkinnedVertexPositions();
@@ -1352,13 +1411,16 @@ namespace UE::MLDeformer
 		UMLDeformerVizSettings* VizSettings = Model->GetVizSettings();
 		ClampCurrentTrainingFrameIndex();
 
+		FMLDeformerSampler* Sampler = GetSamplerForActiveAnim();
+
 		// If we have no Persona toolkit yet, then it is not yet safe to init the sampler.
-		if (Editor->GetPersonaToolkitPointer() != nullptr && !Sampler->IsInitialized())
+		if (Editor->GetPersonaToolkitPointer() != nullptr && Sampler && !Sampler->IsInitialized())
 		{
-			Sampler->Init(this);
+			const int32 ActiveAnimIndex = GetActiveTrainingInputAnimIndex();
+			Sampler->Init(this, ActiveAnimIndex);
 		}
 
-		if (Sampler->IsInitialized())
+		if (Sampler && Sampler->IsInitialized())
 		{
 			Sampler->SetVertexDeltaSpace(EVertexDeltaSpace::PostSkinning);
 			Sampler->Sample(VizSettings->GetTrainingFrameNumber());
@@ -1533,7 +1595,11 @@ namespace UE::MLDeformer
 
 	void FMLDeformerEditorModel::OnPostTraining(ETrainingResult TrainingResult, bool bUsePartiallyTrainedWhenAborted)
 	{
-		Sampler->SetVertexDeltaSpace(EVertexDeltaSpace::PostSkinning);
+		FMLDeformerSampler* Sampler = GetSamplerForActiveAnim();
+		if (Sampler)
+		{
+			Sampler->SetVertexDeltaSpace(EVertexDeltaSpace::PostSkinning);
+		}
 		SampleDeltas();
 		Model->InitGPUData();
 		UpdateMemoryUsage();
@@ -1555,9 +1621,24 @@ namespace UE::MLDeformer
 
 	bool FMLDeformerEditorModel::IsEditorReadyForTrainingBasicChecks()
 	{
+		bool bHasTrainingAnim = false;
+		const int32 NumAnims = GetNumTrainingInputAnims();
+		for (int32 AnimIndex = 0; AnimIndex < NumAnims; ++AnimIndex)
+		{
+			const FMLDeformerTrainingInputAnim* Anim = this->GetTrainingInputAnim(AnimIndex);
+			if (Anim && Anim->IsEnabled())
+			{
+				if (Anim->IsValid())
+				{
+					bHasTrainingAnim = true;
+					break;
+				}
+			}
+		}
+
 		// Make sure we have picked required assets.
 		if (!Model->HasTrainingGroundTruth() ||
-			Model->GetAnimSequence() == nullptr ||
+			!bHasTrainingAnim ||
 			Model->GetSkeletalMesh() == nullptr ||
 			GetNumTrainingFrames() == 0)
 		{
@@ -1580,6 +1661,20 @@ namespace UE::MLDeformer
 		Model->InvalidateMemUsage();
 	}
 
+	UAnimSequence* FMLDeformerEditorModel::GetActiveTrainingInputAnimSequence() const
+	{
+		const int32 ActiveIndex = GetActiveTrainingInputAnimIndex();
+		if (ActiveIndex != INDEX_NONE)
+		{
+			FMLDeformerTrainingInputAnim* Anim = GetTrainingInputAnim(ActiveIndex);
+			if (Anim)
+			{
+				return Anim->GetAnimSequence();
+			}
+		}
+		return nullptr;
+	}
+
 	void FMLDeformerEditorModel::TriggerInputAssetChanged(bool bRefreshVizSettings)
 	{
 		OnInputAssetsChanged();
@@ -1723,6 +1818,10 @@ namespace UE::MLDeformer
 		check(!RenderData->LODRenderData.IsEmpty());
 		const int32 NumRenderVertices = RenderData->LODRenderData[LOD].GetNumVertices();
 
+		FMLDeformerSampler* Sampler = GetSamplerForActiveAnim();
+		check(Sampler);
+		check(Sampler->IsInitialized());
+
 		// Calculate the normals for the base mesh.
 		const FSkeletalMeshModel* ImportedModel = SkelMesh->GetImportedModel();
 		check(ImportedModel);
@@ -1885,6 +1984,10 @@ namespace UE::MLDeformer
 			const int32 NumVerts = Model->GetNumBaseMeshVerts();
 			check(MorphDeltas.Num() % NumVerts == 0);
 
+			FMLDeformerSampler* Sampler = GetSamplerForActiveAnim();
+			check(Sampler);
+			check(Sampler->IsInitialized());
+
 			const TArray<FVector3f>& UnskinnedPositions = Sampler->GetUnskinnedVertexPositions();
 			check(NumVerts == UnskinnedPositions.Num());
 
@@ -1958,7 +2061,7 @@ namespace UE::MLDeformer
 			const UMLDeformerVizSettings* VizSettings = GetModel()->GetVizSettings();
 			if (VizSettings->GetVisualizationMode() == EMLDeformerVizMode::TrainingData)
 			{
-				return GetModel()->GetAnimSequence();
+				return GetActiveTrainingInputAnimSequence();
 			}
 			else if (VizSettings->GetVisualizationMode() == EMLDeformerVizMode::TestData)
 			{
@@ -2005,7 +2108,7 @@ namespace UE::MLDeformer
 		if (Mode == EMLDeformerVizMode::TrainingData)
 		{
 			EditorActor = FindEditorActor(ActorID_Train_Base);
-			AnimAsset = Model->GetAnimSequence();
+			AnimAsset = this->GetActiveTrainingInputAnimSequence();
 		}
 		else if (Mode == EMLDeformerVizMode::TestData)
 		{
@@ -2203,8 +2306,13 @@ namespace UE::MLDeformer
 
 	int32 FMLDeformerEditorModel::GetNumTrainingFrames() const
 	{ 
-		UE_LOG(LogMLDeformer, Warning, TEXT("Please override the FMLDeformerEditorModel::GetNumTrainingFrames() method inside your derived editor model class."));
-		return 0;
+		return NumTrainingFrames;
+	}
+
+	void FMLDeformerEditorModel::UpdateNumTrainingFrames()
+	{
+		NumTrainingFrames = 0;
+		UE_LOG(LogMLDeformer, Warning, TEXT("Please override the FMLDeformerEditorModel::UpdateNumTrainingFrames() inside your model."));
 	}
 
 	UMLDeformerComponent* FMLDeformerEditorModel::FindMLDeformerComponent(int32 ActorID) const
@@ -2216,6 +2324,68 @@ namespace UE::MLDeformer
 		}
 		return nullptr;
 	}
+
+	int32 FMLDeformerEditorModel::GetActiveTrainingInputAnimIndex() const
+	{
+		return ActiveTrainingInputAnimIndex;
+	}
+
+	TSharedPtr<FMLDeformerSampler> FMLDeformerEditorModel::CreateSamplerObject() const
+	{
+		return MakeShared<FMLDeformerSampler>();
+	}
+
+	int32 FMLDeformerEditorModel::GetNumTrainingInputAnims() const
+	{
+		UE_LOG(LogMLDeformer, Warning, TEXT("Please override the FMLDeformerEditorModel::GetNumTrainingInputAnims() method inside your derived editor model class."));
+		return 0;
+	}
+
+	FMLDeformerTrainingInputAnim* FMLDeformerEditorModel::GetTrainingInputAnim(int32 Index) const
+	{
+		UE_LOG(LogMLDeformer, Warning, TEXT("Please override the FMLDeformerEditorModel::GetTrainingInputAnim(int32) method inside your derived editor model class."));
+		return nullptr;
+	}
+
+	FMLDeformerSampler* FMLDeformerEditorModel::GetSampler() const
+	{ 
+		return GetSamplerForActiveAnim();
+	}
+
+	void FMLDeformerEditorModel::UpdateTimelineTrainingAnimList()
+	{
+		TArray<TSharedPtr<FMLDeformerTrainingInputAnimName>> NameList;
+
+		// Build the list of names, based on the training inputs.
+		const int32 NumAnims = GetNumTrainingInputAnims();
+		for (int32 AnimIndex = 0; AnimIndex < NumAnims; ++AnimIndex)
+		{
+			const FMLDeformerTrainingInputAnim* Anim = GetTrainingInputAnim(AnimIndex);
+			if (Anim)
+			{
+				if (Anim->IsValid())
+				{
+					TSharedPtr<FMLDeformerTrainingInputAnimName> AnimName = MakeShared<FMLDeformerTrainingInputAnimName>();
+					AnimName->TrainingInputAnimIndex = AnimIndex;
+					AnimName->Name = FString::Printf(TEXT("[#%d] %s"), AnimIndex, *Anim->GetAnimSequence()->GetName());
+					NameList.Emplace(AnimName);
+				}
+			}
+		}
+
+		// Pass this to the timeline.
+		SMLDeformerTimeline* Timeline = GetEditor()->GetTimeSlider();
+		if (Timeline)
+		{
+			Timeline->SetTrainingAnimNames(NameList);
+		}
+	}
+
+	void FMLDeformerEditorModel::InvalidateDeltas()
+	{
+		// Set the current training frame to -1, which will cause the deltas to get updated next frame.
+		CurrentTrainingFrame = -1;
+	}
 }	// namespace UE::MLDeformer
 
 #undef LOCTEXT_NAMESPACE
diff --git a/Engine/Plugins/Animation/MLDeformer/MLDeformerFramework/Source/MLDeformerFrameworkEditor/Private/MLDeformerEditorModule.cpp b/Engine/Plugins/Animation/MLDeformer/MLDeformerFramework/Source/MLDeformerFrameworkEditor/Private/MLDeformerEditorModule.cpp
--- a/Engine/Plugins/Animation/MLDeformer/MLDeformerFramework/Source/MLDeformerFrameworkEditor/Private/MLDeformerEditorModule.cpp
+++ b/Engine/Plugins/Animation/MLDeformer/MLDeformerFramework/Source/MLDeformerFrameworkEditor/Private/MLDeformerEditorModule.cpp
@@ -7,6 +7,7 @@
 #include "MLDeformerModel.h"
 #include "MLDeformerCurveReferenceCustomization.h"
 #include "MLDeformerMorphModelQualityLevelCustomization.h"
+#include "MLDeformerGeomCacheTrainingInputAnimCustomize.h"
 #include "SMLDeformerInputWidget.h"
 #include "Modules/ModuleManager.h"
 #include "EditorModeRegistry.h"
@@ -26,6 +27,7 @@ namespace UE::MLDeformer
 		FPropertyEditorModule& PropertyModule = FModuleManager::GetModuleChecked<FPropertyEditorModule>("PropertyEditor");
 		PropertyModule.RegisterCustomPropertyTypeLayout("MLDeformerCurveReference", FOnGetPropertyTypeCustomizationInstance::CreateStatic(&FMLDeformerCurveReferenceCustomization::MakeInstance) );
 		PropertyModule.RegisterCustomPropertyTypeLayout("MLDeformerMorphModelQualityLevel", FOnGetPropertyTypeCustomizationInstance::CreateStatic(&FMLDeformerMorphModelQualityLevelCustomization::MakeInstance) );
+		PropertyModule.RegisterCustomPropertyTypeLayout("MLDeformerGeomCacheTrainingInputAnim", FOnGetPropertyTypeCustomizationInstance::CreateStatic(&FMLDeformerGeomCacheTrainingInputAnimCustomization::MakeInstance) );
 		PropertyModule.NotifyCustomizationModuleChanged();
 
 		SMLDeformerInputWidget::RegisterCommands();
@@ -41,6 +43,7 @@ namespace UE::MLDeformer
 			FPropertyEditorModule& PropertyModule = FModuleManager::GetModuleChecked<FPropertyEditorModule>(TEXT("PropertyEditor"));
 			PropertyModule.UnregisterCustomPropertyTypeLayout(TEXT("MLDeformerCurveReference"));
 			PropertyModule.UnregisterCustomPropertyTypeLayout(TEXT("MLDeformerMorphModelQualityLevel"));
+			PropertyModule.UnregisterCustomPropertyTypeLayout(TEXT("MLDeformerGeomCacheTrainingInputAnim"));
 			PropertyModule.NotifyCustomizationModuleChanged();
 		}
 	}
diff --git a/Engine/Plugins/Animation/MLDeformer/MLDeformerFramework/Source/MLDeformerFrameworkEditor/Private/MLDeformerEditorToolkit.cpp b/Engine/Plugins/Animation/MLDeformer/MLDeformerFramework/Source/MLDeformerFrameworkEditor/Private/MLDeformerEditorToolkit.cpp
--- a/Engine/Plugins/Animation/MLDeformer/MLDeformerFramework/Source/MLDeformerFrameworkEditor/Private/MLDeformerEditorToolkit.cpp
+++ b/Engine/Plugins/Animation/MLDeformer/MLDeformerFramework/Source/MLDeformerFrameworkEditor/Private/MLDeformerEditorToolkit.cpp
@@ -106,6 +106,7 @@ namespace UE::MLDeformer
 			ActiveModel->UpdateIsReadyForTrainingState();
 			ActiveModel->SetTrainingFrame(ActiveModel->GetModel()->GetVizSettings()->GetTrainingFrameNumber());
 			ActiveModel->SetTestFrame(ActiveModel->GetModel()->GetVizSettings()->GetTestingFrameNumber());
+			ActiveModel->InvalidateDeltas();
 		}
 
 		if (DeformerModel)
@@ -269,7 +270,7 @@ namespace UE::MLDeformer
 		{
 			const EAppReturnType::Type ReturnType = FMessageDialog::Open(
 				EAppMsgType::YesNo, 
-				LOCTEXT("SwitchModelConfirmMessage", "Are you sure you want to switch the current model?\nYou will lose your current setup."),
+				LOCTEXT("SwitchModelConfirmMessage", "Are you sure you want to switch the current model?\nYou will lose your current setup.\n"),
 				LOCTEXT("SwitchModelConfirmTitle", "Switch current model?"));
 
 			if (ReturnType == EAppReturnType::No)
@@ -285,6 +286,19 @@ namespace UE::MLDeformer
 		if (ActiveModel)
 		{
 			ActiveModel->ClearWorldAndPersonaPreviewScene();
+			ActiveModel.Reset();
+
+			if (ModelDetailsView)
+			{
+				ModelDetailsView->SetObject(nullptr);
+				ModelDetailsView->ForceRefresh();
+			}
+
+			if (VizSettingsDetailsView)
+			{
+				VizSettingsDetailsView->SetObject(nullptr);
+				VizSettingsDetailsView->ForceRefresh();
+			}
 		}
 
 		// Get the runtime model type based on the index, and create an instance of it.
@@ -303,11 +317,10 @@ namespace UE::MLDeformer
 		EditorModel->Init(InitSettings);
 
 		// Tell the editor we use this model now.
-		FMLDeformerEditorModel* OldActiveModel = ActiveModel.Get();
 		ActiveModel = TSharedPtr<FMLDeformerEditorModel>(EditorModel);
 
 		// Create the new scene for this model.
-		if (OldActiveModel)
+		if (PersonaToolkit)
 		{
 			HandlePreviewSceneCreated(PersonaToolkit->GetPreviewScene());
 		}
@@ -389,7 +402,7 @@ namespace UE::MLDeformer
 		{
 			const EAppReturnType::Type ConfirmReturnType = FMessageDialog::Open(
 				EAppMsgType::YesNo, 
-				LOCTEXT("RetrainConfirmationMessage", "This asset already has been trained.\n\nAre you sure you would like to re-train the network with your current settings?"),
+				LOCTEXT("RetrainConfirmationMessage", "This asset already has been trained.\n\nAre you sure you would like to re-train the network with your current settings?\n"),
 				LOCTEXT("RetrainConfirmationTitle", "Re-train the network?"));
 
 			if (ConfirmReturnType == EAppReturnType::No || ConfirmReturnType == EAppReturnType::Cancel)
@@ -403,12 +416,6 @@ namespace UE::MLDeformer
 
 		ActiveModel->OnPreTraining();
 
-		// Change the interpolation type for the training sequence to step.
-		if (UAnimSequence* AnimSequence = Model->GetAnimSequence())
-		{
-			AnimSequence->Interpolation = EAnimInterpolationType::Step;
-		}
-
 		// Initialize the training inputs.
 		ActiveModel->UpdateEditorInputInfo();
 
@@ -484,7 +491,7 @@ namespace UE::MLDeformer
 				FUIAction(),
 				FOnGetContent::CreateRaw(this, &FMLDeformerEditorToolkit::GenerateVizModeButtonContents, CommandList.ToSharedRef()),
 				TAttribute<FText>::CreateRaw(this, &FMLDeformerEditorToolkit::GetCurrentVizModeName),			
-				LOCTEXT("VizModeModeTooltip", "The visualization mode, specifying whether you are working on training on testing."),
+				LOCTEXT("VizModeModeTooltip", "The visualization mode, specifying whether you are working on training or testing."),
 				FSlateIcon(FMLDeformerEditorStyle::Get().GetStyleSetName(), "MLDeformer.VizSettings.TabIcon")
 			);
 		}
@@ -548,14 +555,14 @@ namespace UE::MLDeformer
 				if (!ActiveModel->LoadTrainedNetwork())
 				{
 					GEditor->PlayEditorSound(TEXT("/Engine/EditorSounds/Notifications/CompileFailed_Cue.CompileFailed_Cue"));
-					WindowMessage = LOCTEXT("TrainingOnnxLoadFailed", "Training completed but resulting network couldn't be loaded!");
+					WindowMessage = LOCTEXT("TrainingOnnxLoadFailed", "Training completed but resulting network couldn't be loaded!\n");
 				}
 				else
 				{
 					GEditor->PlayEditorSound(TEXT("/Engine/EditorSounds/Notifications/CompileSuccess_Cue.CompileSuccess_Cue"));
 					FFormatNamedArguments SuccessArgs;
 					SuccessArgs.Add(TEXT("Duration"), TrainingDurationText);
-					WindowMessage = FText::Format(LOCTEXT("TrainingSuccess", "Training completed successfully!\n\nTraining time: {Duration}"), SuccessArgs);
+					WindowMessage = FText::Format(LOCTEXT("TrainingSuccess", "Training completed successfully!\n\nTraining time: {Duration}\n"), SuccessArgs);
 					ActiveModel->InitInputInfo(ActiveModel->GetModel()->GetInputInfo());
 					bMarkDirty = true;
 					bOutSuccess = true;
@@ -573,7 +580,7 @@ namespace UE::MLDeformer
 				{
 					ReturnType = FMessageDialog::Open(
 						EAppMsgType::YesNo, 
-						LOCTEXT("TrainingAbortedMessage", "Training has been aborted.\nThe neural network has only been partially trained.\nWould you like to use this partially trained network?"),
+						LOCTEXT("TrainingAbortedMessage", "Training has been aborted.\nThe neural network has only been partially trained.\nWould you like to use this partially trained network?\n"),
 						LOCTEXT("TrainingAbortedMessageTitle", "Use partially trained network?"));
 				}
 
@@ -604,7 +611,7 @@ namespace UE::MLDeformer
 			case ETrainingResult::AbortedCantUse:
 			{
 				ShowNotification(LOCTEXT("TrainingAborted", "Training aborted!"), SNotificationItem::ECompletionState::CS_None, true);
-				WindowMessage = LOCTEXT("TrainingAbortedCantUse", "Training aborted by user.");
+				WindowMessage = LOCTEXT("TrainingAbortedCantUse", "Training aborted by user.\n");
 			}
 			break;
 
@@ -612,7 +619,7 @@ namespace UE::MLDeformer
 			case ETrainingResult::FailOnData:
 			{
 				GEditor->PlayEditorSound(TEXT("/Engine/EditorSounds/Notifications/CompileFailed_Cue.CompileFailed_Cue"));
-				WindowMessage = LOCTEXT("TrainingFailedOnData", "Training failed!\nCheck input parameters or sequence length.");
+				WindowMessage = LOCTEXT("TrainingFailedOnData", "Training failed!\nCheck input parameters or sequence length.\n");
 			}
 			break;
 
@@ -620,7 +627,7 @@ namespace UE::MLDeformer
 			case ETrainingResult::FailPythonError:
 			{
 				GEditor->PlayEditorSound(TEXT("/Engine/EditorSounds/Notifications/CompileFailed_Cue.CompileFailed_Cue"));
-				WindowMessage = LOCTEXT("TrainingPythonError", "Training failed!\nThere is a python error, please check the output log.");
+				WindowMessage = LOCTEXT("TrainingPythonError", "Training failed!\nThere is a python error, please check the output log.\n");
 			}
 			break;
 
@@ -888,7 +895,7 @@ namespace UE::MLDeformer
 		const UMLDeformerVizSettings* VizSettings = ActiveModel->GetModel()->GetVizSettings();
 		if (VizSettings->GetVisualizationMode() == EMLDeformerVizMode::TrainingData)
 		{
-			const UAnimSequence* AnimSeq = ActiveModel.Get() ? ActiveModel->GetModel()->GetAnimSequence() : nullptr;
+			const UAnimSequence* AnimSeq = ActiveModel.Get() ? ActiveModel->GetActiveTrainingInputAnimSequence() : nullptr;
 			const double Duration = AnimSeq ? AnimSeq->GetPlayLength() : 0.0;
 			SetTimeSliderRange(0.0, Duration);
 		}
diff --git a/Engine/Plugins/Animation/MLDeformer/MLDeformerFramework/Source/MLDeformerFrameworkEditor/Private/MLDeformerGeomCacheActor.cpp b/Engine/Plugins/Animation/MLDeformer/MLDeformerFramework/Source/MLDeformerFrameworkEditor/Private/MLDeformerGeomCacheActor.cpp
--- a/Engine/Plugins/Animation/MLDeformer/MLDeformerFramework/Source/MLDeformerFrameworkEditor/Private/MLDeformerGeomCacheActor.cpp
+++ b/Engine/Plugins/Animation/MLDeformer/MLDeformerFramework/Source/MLDeformerFrameworkEditor/Private/MLDeformerGeomCacheActor.cpp
@@ -11,14 +11,6 @@ namespace UE::MLDeformer
 	{
 	}
 
-	FMLDeformerGeomCacheActor::~FMLDeformerGeomCacheActor()
-	{
-		if (GeomCacheComponent)
-		{
-			Actor->RemoveOwnedComponent(GeomCacheComponent);
-		}
-	}
-
 	void FMLDeformerGeomCacheActor::SetVisibility(bool bIsVisible)
 	{
 		FMLDeformerEditorActor::SetVisibility(bIsVisible);
diff --git a/Engine/Plugins/Animation/MLDeformer/MLDeformerFramework/Source/MLDeformerFrameworkEditor/Private/MLDeformerGeomCacheEditorModel.cpp b/Engine/Plugins/Animation/MLDeformer/MLDeformerFramework/Source/MLDeformerFrameworkEditor/Private/MLDeformerGeomCacheEditorModel.cpp
--- a/Engine/Plugins/Animation/MLDeformer/MLDeformerFramework/Source/MLDeformerFrameworkEditor/Private/MLDeformerGeomCacheEditorModel.cpp
+++ b/Engine/Plugins/Animation/MLDeformer/MLDeformerFramework/Source/MLDeformerFrameworkEditor/Private/MLDeformerGeomCacheEditorModel.cpp
@@ -10,6 +10,7 @@
 #include "MLDeformerGeomCacheSampler.h"
 #include "MLDeformerGeomCacheHelpers.h"
 #include "MLDeformerGeomCacheModel.h"
+#include "SMLDeformerTimeline.h"
 #include "GeometryCache.h"
 #include "GeometryCacheComponent.h"
 #include "Animation/DebugSkelMeshComponent.h"
@@ -22,11 +23,21 @@
 
 namespace UE::MLDeformer
 {
-	FMLDeformerSampler* FMLDeformerGeomCacheEditorModel::CreateSampler() const
+	void FMLDeformerGeomCacheEditorModel::Init(const InitSettings& Settings)
 	{
-		FMLDeformerGeomCacheSampler* NewSampler = new FMLDeformerGeomCacheSampler();
-		NewSampler->OnGetGeometryCache().BindLambda([this] { return GetGeomCacheModel()->GetGeometryCache(); });
-		return NewSampler;
+		FMLDeformerEditorModel::Init(Settings);
+
+		// Create an initial training input animation entry.
+		if (GetGeomCacheModel()->GetTrainingInputAnims().IsEmpty())
+		{
+			GetGeomCacheModel()->GetTrainingInputAnims().AddDefaulted();
+			CreateSamplers();
+		}
+	}
+
+	TSharedPtr<FMLDeformerSampler> FMLDeformerGeomCacheEditorModel::CreateSamplerObject() const
+	{
+		return MakeShared<FMLDeformerGeomCacheSampler>();
 	}
 
 	FMLDeformerEditorActor* FMLDeformerGeomCacheEditorModel::CreateEditorActor(const FMLDeformerEditorActor::FConstructSettings& Settings) const
@@ -39,11 +50,6 @@ namespace UE::MLDeformer
 		return Cast<UMLDeformerGeomCacheVizSettings>(Model->GetVizSettings());
 	}
 
-	FMLDeformerGeomCacheSampler* FMLDeformerGeomCacheEditorModel::GetGeomCacheSampler() const
-	{ 
-		return static_cast<FMLDeformerGeomCacheSampler*>(Sampler);
-	}
-
 	UMLDeformerGeomCacheModel* FMLDeformerGeomCacheEditorModel::GetGeomCacheModel() const
 	{ 
 		return Cast<UMLDeformerGeomCacheModel>(Model);
@@ -77,16 +83,30 @@ namespace UE::MLDeformer
 		return FMLDeformerEditorModel::GetTestFrameAtTime(TimeInSeconds);
 	}
 
-	int32 FMLDeformerGeomCacheEditorModel::GetNumTrainingFrames() const
+	UGeometryCache* FMLDeformerGeomCacheEditorModel::GetActiveGeometryCache() const
+	{
+		const int32 ActiveAnimIndex = GetActiveTrainingInputAnimIndex();
+		if (ActiveAnimIndex == INDEX_NONE || ActiveAnimIndex >= GetNumTrainingInputAnims())
+		{
+			return nullptr;
+		}
+
+		FMLDeformerGeomCacheTrainingInputAnim* ActiveAnim = static_cast<FMLDeformerGeomCacheTrainingInputAnim*>(GetTrainingInputAnim(ActiveAnimIndex));
+		return ActiveAnim ? ActiveAnim->GetGeometryCache() : nullptr;
+	}
+
+	void FMLDeformerGeomCacheEditorModel::UpdateNumTrainingFrames()
 	{
-		const UGeometryCache* GeometryCache = GetGeomCacheModel()->GetGeometryCache();
-		if (GeometryCache == nullptr)
+		NumTrainingFrames = 0;
+
+		const TArray<FMLDeformerGeomCacheTrainingInputAnim>& TrainingInputAnims = GetGeomCacheModel()->GetTrainingInputAnims();
+		for (const FMLDeformerGeomCacheTrainingInputAnim& InputAnim : TrainingInputAnims)
 		{
-			return 0;
+			if (InputAnim.IsEnabled() && InputAnim.IsValid())
+			{
+				NumTrainingFrames += InputAnim.GetNumFramesToSample();
+			}
 		}
-		const int32 StartFrame = GeometryCache->GetStartFrame();
-		const int32 EndFrame = GeometryCache->GetEndFrame();
-		return (EndFrame - StartFrame) + 1;
 	}
 
 	void FMLDeformerGeomCacheEditorModel::UpdateIsReadyForTrainingState()
@@ -99,35 +119,47 @@ namespace UE::MLDeformer
 			return;
 		}
 
-		// Now make sure the assets are compatible.
+		// Now make sure the geom caches have no errors.
 		UMLDeformerGeomCacheModel* GeomCacheModel = GetGeomCacheModel();
-		UGeometryCache* GeomCache = GeomCacheModel->GetGeometryCache();
 		USkeletalMesh* SkeletalMesh = GeomCacheModel->GetSkeletalMesh();
-		if (!GetGeomCacheErrorText(SkeletalMesh, GeomCache).IsEmpty())
+		const int32 NumAnims = GetNumTrainingInputAnims();
+		for (int32 AnimIndex = 0; AnimIndex < NumAnims; ++AnimIndex)
 		{
-			return;
-		}
-
-		// Make sure every skeletal imported mesh has some geometry track.
-		const int32 NumGeomCacheTracks = GeomCache ? GeomCache->Tracks.Num() : 0;
-		int32 NumSkelMeshes = 0;
-		check(SkeletalMesh);
-		FSkeletalMeshModel* ImportedModel = SkeletalMesh->GetImportedModel();
-		if (ImportedModel)
-		{
-			NumSkelMeshes = ImportedModel->LODModels[0].ImportedMeshInfos.Num();
+			FMLDeformerGeomCacheTrainingInputAnim* Anim = static_cast<FMLDeformerGeomCacheTrainingInputAnim*>(GetTrainingInputAnim(AnimIndex));
+			if (Anim && Anim->IsEnabled())
+			{
+				UGeometryCache* GeomCache = Anim->GetGeometryCache();
+				if (!GetGeomCacheErrorText(SkeletalMesh, GeomCache).IsEmpty())
+				{
+					return;
+				}
+			}
 		}
 
 		// Check if we have any mappings at all.
-		FMLDeformerGeomCacheSampler* GeomCacheSampler = GetGeomCacheSampler();
-		if (!GeomCacheSampler->IsInitialized())
-		{
-			GeomCacheSampler->Init(this);
-		}
-
-		if (GeomCacheSampler->GetMeshMappings().IsEmpty())
+		for (int32 AnimIndex = 0; AnimIndex < NumAnims; ++AnimIndex)
 		{
-			return;
+			FMLDeformerGeomCacheTrainingInputAnim* Anim = static_cast<FMLDeformerGeomCacheTrainingInputAnim*>(GetTrainingInputAnim(AnimIndex));
+			if (!Anim->IsEnabled())
+			{
+				continue;
+			}
+
+			FMLDeformerGeomCacheSampler* GeomCacheSampler = static_cast<FMLDeformerGeomCacheSampler*>(GetSamplerForTrainingAnim(AnimIndex));
+			if (GeomCacheSampler == nullptr)
+			{
+				return;
+			}
+
+			if (!GeomCacheSampler->IsInitialized())
+			{
+				GeomCacheSampler->Init(this, AnimIndex);
+			}
+
+			if (GeomCacheSampler->GetMeshMappings().IsEmpty())
+			{
+				return;
+			}
 		}
 
 		bIsReadyForTraining = true;
@@ -135,7 +167,17 @@ namespace UE::MLDeformer
 
 	void FMLDeformerGeomCacheEditorModel::CreateTrainingGroundTruthActor(UWorld* World)
 	{
-		UGeometryCache* GeomCache = GetGeomCacheModel()->GetGeometryCache();
+		const int32 ActiveAnimIndex = GetActiveTrainingInputAnimIndex();
+		UGeometryCache* GeomCache = nullptr;
+		if (ActiveAnimIndex != INDEX_NONE)
+		{
+			if (ActiveAnimIndex < GetNumTrainingInputAnims())
+			{
+				FMLDeformerGeomCacheTrainingInputAnim* Anim = static_cast<FMLDeformerGeomCacheTrainingInputAnim*>(GetTrainingInputAnim(ActiveAnimIndex));
+				GeomCache = Anim ? Anim->GetGeometryCache() : nullptr;
+			}
+		}
+
 		const FLinearColor LabelColor = FMLDeformerEditorStyle::Get().GetColor("MLDeformer.TargetMesh.LabelColor");
 		const FLinearColor WireframeColor = FMLDeformerEditorStyle::Get().GetColor("MLDeformer.TargetMesh.WireframeColor");
 		CreateGeomCacheActor(
@@ -211,9 +253,46 @@ namespace UE::MLDeformer
 		// Handle base class property changes.
 		FMLDeformerEditorModel::OnPropertyChanged(PropertyChangedEvent);
 
-		// Properties specific to the geometry cache editor model.
-		if (Property->GetFName() == UMLDeformerGeomCacheModel::GetGeometryCachePropertyName())
+		// Start frame changed.
+		if (PropertyChangedEvent.GetMemberPropertyName() == TEXT("TrainingInputAnims") && 
+			Property->GetFName() == FMLDeformerTrainingInputAnim::GetStartFramePropertyName())
+		{
+			for (FMLDeformerGeomCacheTrainingInputAnim& InputAnim : GetGeomCacheModel()->GetTrainingInputAnims())
+			{
+				if (InputAnim.GetStartFrame() > InputAnim.GetEndFrame())
+				{
+					InputAnim.SetEndFrame(InputAnim.GetStartFrame());
+				}
+			}
+			SetResamplingInputOutputsNeeded(true);
+			TriggerInputAssetChanged(true);
+		}
+		else // End Frame changed.
+		if (PropertyChangedEvent.GetMemberPropertyName() == TEXT("TrainingInputAnims") && 
+		Property->GetFName() == FMLDeformerTrainingInputAnim::GetEndFramePropertyName())
+		{
+			for (FMLDeformerGeomCacheTrainingInputAnim& InputAnim : GetGeomCacheModel()->GetTrainingInputAnims())
+			{
+				if (InputAnim.GetEndFrame() < InputAnim.GetStartFrame())
+				{
+					InputAnim.SetStartFrame(InputAnim.GetEndFrame());
+				}
+			}
+			SetResamplingInputOutputsNeeded(true);
+			TriggerInputAssetChanged(true);
+		}
+		else
+		if (Property->GetFName() == UMLDeformerGeomCacheModel::GetTrainingInputAnimsPropertyName() ||
+		    PropertyChangedEvent.GetMemberPropertyName() == TEXT("TrainingInputAnims"))
 		{
+			if (GetEditor() && GetEditor()->GetTimeSlider())
+			{
+				if (PropertyChangedEvent.ChangeType == EPropertyChangeType::ArrayRemove)
+				{
+					const int32 ArrayIndex = PropertyChangedEvent.GetArrayIndex(PropertyChangedEvent.GetMemberPropertyName().ToString());
+					GetEditor()->GetTimeSlider()->OnDeletedTrainingInputAnim(ArrayIndex);
+				}
+			}
 			SetResamplingInputOutputsNeeded(true);
 			TriggerInputAssetChanged(true);
 		}
@@ -238,7 +317,7 @@ namespace UE::MLDeformer
 		UGeometryCacheComponent* GeometryCacheComponent = TrainGroundTruthActor ? TrainGroundTruthActor->GetGeometryCacheComponent() : nullptr;
 		if (GeometryCacheComponent)
 		{
-			GeometryCacheComponent->SetGeometryCache(GetGeomCacheModel()->GetGeometryCache());
+			GeometryCacheComponent->SetGeometryCache(GetActiveGeometryCache());
 			GeometryCacheComponent->SetLooping(false);
 			GeometryCacheComponent->SetManualTick(true);
 			GeometryCacheComponent->SetPlaybackSpeed(TestAnimSpeed);
@@ -265,11 +344,21 @@ namespace UE::MLDeformer
 		// Handle changes for the skeletal mesh and anim sequence.
 		FMLDeformerEditorModel::OnObjectModified(Object);
 
-		// Check if the object is our training geometry cache.
-		if (GetGeomCacheModel()->GetGeometryCache() == Object)
+		const int32 NumAnims = GetGeomCacheModel()->GetTrainingInputAnims().Num();
+		for (int32 AnimIndex = 0; AnimIndex < NumAnims; ++AnimIndex)
 		{
-			UE_LOG(LogMLDeformer, Display, TEXT("Detected a modification in training geometry cache %s, reinitializing inputs."), *Object->GetName());
-			bNeedsAssetReinit = true;
+			const FMLDeformerGeomCacheTrainingInputAnim& Anim = GetGeomCacheModel()->GetTrainingInputAnims()[AnimIndex];
+			if (Anim.GetAnimSequence() == Object)
+			{
+				UE_LOG(LogMLDeformer, Display, TEXT("Detected a modification in training anim sequence %s, reinitializing inputs."), *Object->GetName());
+				bNeedsAssetReinit = true;
+			}
+
+			if (Anim.GetGeometryCache() == Object)
+			{
+				UE_LOG(LogMLDeformer, Display, TEXT("Detected a modification in training geometry cache %s, reinitializing inputs."), *Object->GetName());
+				bNeedsAssetReinit = true;
+			}
 		}
 
 		// Check if it is the ground truth test geom cache.
@@ -279,6 +368,21 @@ namespace UE::MLDeformer
 			bNeedsAssetReinit = true;
 		}
 	}
+
+	int32 FMLDeformerGeomCacheEditorModel::GetNumTrainingInputAnims() const
+	{
+		return GetGeomCacheModel()->GetTrainingInputAnims().Num();
+	}
+
+	FMLDeformerTrainingInputAnim* FMLDeformerGeomCacheEditorModel::GetTrainingInputAnim(int32 Index) const
+	{
+		if (GetGeomCacheModel()->GetTrainingInputAnims().IsValidIndex(Index))
+		{
+			return &GetGeomCacheModel()->GetTrainingInputAnims()[Index];
+		}
+		
+		return nullptr;
+	}
 }	// namespace UE::MLDeformer
 
 #undef LOCTEXT_NAMESPACE
diff --git a/Engine/Plugins/Animation/MLDeformer/MLDeformerFramework/Source/MLDeformerFrameworkEditor/Private/MLDeformerGeomCacheModelDetails.cpp b/Engine/Plugins/Animation/MLDeformer/MLDeformerFramework/Source/MLDeformerFrameworkEditor/Private/MLDeformerGeomCacheModelDetails.cpp
--- a/Engine/Plugins/Animation/MLDeformer/MLDeformerFramework/Source/MLDeformerFrameworkEditor/Private/MLDeformerGeomCacheModelDetails.cpp
+++ b/Engine/Plugins/Animation/MLDeformer/MLDeformerFramework/Source/MLDeformerFrameworkEditor/Private/MLDeformerGeomCacheModelDetails.cpp
@@ -26,73 +26,10 @@ namespace UE::MLDeformer
 		return (GeomCacheModel != nullptr && GeomCacheEditorModel != nullptr);
 	}
 
-	void FMLDeformerGeomCacheModelDetails::AddAnimSequenceErrors()
+	void FMLDeformerGeomCacheModelDetails::AddTrainingInputAnims()
 	{
-		const FText WarningText = GetGeomCacheAnimSequenceErrorText(GeomCacheModel->GetGeometryCache(), Model->GetAnimSequence());
-		BaseMeshCategoryBuilder->AddCustomRow(FText::FromString("AnimSeqWarning"))
-			.Visibility(!WarningText.IsEmpty() ? EVisibility::Visible : EVisibility::Collapsed)
-			.WholeRowContent()
-			[
-				SNew(SBox)
-				.Padding(FMargin(0.0f, 4.0f))
-				[
-					SNew(SWarningOrErrorBox)
-					.MessageStyle(EMessageStyle::Warning)
-					.Message(WarningText)
-				]
-			];
-	}
-
-	void FMLDeformerGeomCacheModelDetails::AddTargetMesh()
-	{
-		TargetMeshCategoryBuilder->AddProperty(UMLDeformerGeomCacheModel::GetGeometryCachePropertyName(), UMLDeformerGeomCacheModel::StaticClass());
-
-		const FText TargetMeshErrorText = GetGeomCacheErrorText(GeomCacheModel->GetSkeletalMesh(), GeomCacheModel->GetGeometryCache());
-		TargetMeshCategoryBuilder->AddCustomRow(FText::FromString("TargetMeshError"))
-			.Visibility(!TargetMeshErrorText.IsEmpty() ? EVisibility::Visible : EVisibility::Collapsed)
-			.WholeRowContent()
-			[
-				SNew(SBox)
-				.Padding(FMargin(0.0f, 4.0f))
-				[
-					SNew(SWarningOrErrorBox)
-					.MessageStyle(EMessageStyle::Error)
-					.Message(TargetMeshErrorText)
-				]
-			];
-
-		const FText ChangedErrorText = EditorModel->GetTargetAssetChangedErrorText();
-		TargetMeshCategoryBuilder->AddCustomRow(FText::FromString("TargetMeshChangedError"))
-			.Visibility(!ChangedErrorText.IsEmpty() ? EVisibility::Visible : EVisibility::Collapsed)
-			.WholeRowContent()
-			[
-				SNew(SBox)
-				.Padding(FMargin(0.0f, 4.0f))
-				[
-					SNew(SWarningOrErrorBox)
-					.MessageStyle(EMessageStyle::Warning)
-					.Message(ChangedErrorText)
-				]
-			];
-
-		AddGeomCacheMeshMappingWarnings(TargetMeshCategoryBuilder, Model->GetSkeletalMesh(), GeomCacheModel->GetGeometryCache());
-	}
-
-	void FMLDeformerGeomCacheModelDetails::AddGeomCacheMeshMappingWarnings(IDetailCategoryBuilder* InTargetMeshCategoryBuilder, USkeletalMesh* SkeletalMesh, UGeometryCache* GeometryCache)
-	{
-		const FText MeshMappingError = GetGeomCacheMeshMappingErrorText(SkeletalMesh, GeometryCache);
-		InTargetMeshCategoryBuilder->AddCustomRow(FText::FromString("MeshMappingWarning"))
-			.Visibility(!MeshMappingError.IsEmpty() ? EVisibility::Visible : EVisibility::Collapsed)
-			.WholeRowContent()
-			[
-				SNew(SBox)
-				.Padding(FMargin(0.0f, 4.0f))
-				[
-					SNew(SWarningOrErrorBox)
-					.MessageStyle(EMessageStyle::Warning)
-					.Message(MeshMappingError)
-				]
-			];
+		InputOutputCategoryBuilder->AddProperty(UMLDeformerGeomCacheModel::GetTrainingInputAnimsPropertyName(), UMLDeformerGeomCacheModel::StaticClass())
+			.DisplayName( FText::Format(LOCTEXT("TrainingInputAnimsString", "Training Input Anims ({0} Frames)"), EditorModel->GetNumTrainingFrames()) );
 	}
 }	// namespace UE::MLDeformer
 
diff --git a/Engine/Plugins/Animation/MLDeformer/MLDeformerFramework/Source/MLDeformerFrameworkEditor/Private/MLDeformerGeomCacheSampler.cpp b/Engine/Plugins/Animation/MLDeformer/MLDeformerFramework/Source/MLDeformerFrameworkEditor/Private/MLDeformerGeomCacheSampler.cpp
--- a/Engine/Plugins/Animation/MLDeformer/MLDeformerFramework/Source/MLDeformerFrameworkEditor/Private/MLDeformerGeomCacheSampler.cpp
+++ b/Engine/Plugins/Animation/MLDeformer/MLDeformerFramework/Source/MLDeformerFrameworkEditor/Private/MLDeformerGeomCacheSampler.cpp
@@ -7,6 +7,7 @@
 #include "MLDeformerEditorModule.h"
 #include "MLDeformerEditorToolkit.h"
 #include "MLDeformerGeomCacheHelpers.h"
+#include "MLDeformerGeomCacheModel.h"
 #include "Animation/DebugSkelMeshComponent.h"
 #include "Animation/AnimSequence.h"
 #include "Animation/AnimSingleNodeInstance.h"
@@ -23,8 +24,15 @@
 
 namespace UE::MLDeformer
 {
-	void FMLDeformerGeomCacheSampler::RegisterTargetComponents()
+	void FMLDeformerGeomCacheSampler::Init(FMLDeformerEditorModel* InModel, int32 InAnimIndex)
 	{
+		FMLDeformerSampler::Init(InModel, InAnimIndex);
+
+		if (!TargetMeshActor)
+		{
+			return;
+		}
+
 		// Create the geometry cache component.
 		if (GeometryCacheComponent.Get() == nullptr)
 		{
@@ -33,7 +41,10 @@ namespace UE::MLDeformer
 			TargetMeshActor->SetRootComponent(GeometryCacheComponent);
 		}
 
-		UGeometryCache* GeomCache = OnGetGeometryCache().IsBound() ? OnGetGeometryCache().Execute() : nullptr;
+		UMLDeformerGeomCacheModel* GeomCacheModel = Cast<UMLDeformerGeomCacheModel>(InModel->GetModel());
+		check(GeomCacheModel);
+
+		UGeometryCache* GeomCache = GeomCacheModel->GetTrainingInputAnims()[InAnimIndex].GetGeometryCache();
 		GeometryCacheComponent->SetGeometryCache(GeomCache);
 		GeometryCacheComponent->SetManualTick(true);
 		GeometryCacheComponent->SetVisibility(false);
diff --git a/Engine/Plugins/Animation/MLDeformer/MLDeformerFramework/Source/MLDeformerFrameworkEditor/Private/MLDeformerGeomCacheTrainingInputAnimCustomize.cpp b/Engine/Plugins/Animation/MLDeformer/MLDeformerFramework/Source/MLDeformerFrameworkEditor/Private/MLDeformerGeomCacheTrainingInputAnimCustomize.cpp
new file mode 100644
--- /dev/null
+++ b/Engine/Plugins/Animation/MLDeformer/MLDeformerFramework/Source/MLDeformerFrameworkEditor/Private/MLDeformerGeomCacheTrainingInputAnimCustomize.cpp
@@ -0,0 +1,258 @@
+// Copyright Epic Games, Inc. All Rights Reserved.
+
+#include "MLDeformerGeomCacheTrainingInputAnimCustomize.h"
+#include "MLDeformerGeomCacheTrainingInputAnim.h"
+#include "MLDeformerGeomCacheModel.h"
+#include "MLDeformerGeomCacheHelpers.h"
+#include "MLDeformerEditorModel.h"
+#include "MLDeformerEditorModule.h"
+#include "Animation/Skeleton.h"
+#include "PropertyCustomizationHelpers.h"
+#include "DetailLayoutBuilder.h"
+#include "IDetailChildrenBuilder.h"
+#include "Widgets/SBoxPanel.h"
+#include "Widgets/Text/STextBlock.h"
+#include "Widgets/Layout/SBox.h"
+#include "SWarningOrErrorBox.h"
+#include "DetailWidgetRow.h"
+#include "Engine/SkeletalMesh.h"
+
+#define LOCTEXT_NAMESPACE "MLDeformerGeomCacheTrainingInputAnimCustomize"
+
+namespace UE::MLDeformer
+{
+	TSharedRef<IPropertyTypeCustomization> FMLDeformerGeomCacheTrainingInputAnimCustomization::MakeInstance()
+	{
+		return MakeShareable(new FMLDeformerGeomCacheTrainingInputAnimCustomization());
+	}
+
+	void FMLDeformerGeomCacheTrainingInputAnimCustomization::CustomizeHeader(TSharedRef<IPropertyHandle> StructPropertyHandle, class FDetailWidgetRow& HeaderRow, IPropertyTypeCustomizationUtils& StructCustomizationUtils)
+	{
+		const int32 ArrayIndex = StructPropertyHandle->GetIndexInArray();
+		check(ArrayIndex != INDEX_NONE);
+
+		UMLDeformerGeomCacheModel* Model = FindMLDeformerModel(StructPropertyHandle);
+		check(Model);
+		int32 NumAnimFrames = 0;
+		int32 NumFramesToSample = 0;
+		FMLDeformerGeomCacheTrainingInputAnim& InputAnim = Model->GetTrainingInputAnims()[ArrayIndex];
+		const UGeometryCache* GeomCache = InputAnim.GetGeometryCache();
+		if (GeomCache)
+		{
+			NumAnimFrames = GeomCache->GetEndFrame() - GeomCache->GetStartFrame() + 1;
+		}
+		NumFramesToSample = InputAnim.GetNumFramesToSample();
+
+		HeaderRow
+		.NameContent()
+		.VAlign(EVerticalAlignment::VAlign_Center)
+		.HAlign(EHorizontalAlignment::HAlign_Left)
+		[
+			SNew(STextBlock)
+			.Text(FText::Format(LOCTEXT("HeaderTextFmt", "Animation #{0}"), ArrayIndex))
+			.Font(StructCustomizationUtils.GetRegularFont())
+			.ColorAndOpacity(InputAnim.IsEnabled() ? FSlateColor::UseForeground() : FSlateColor::UseSubduedForeground())
+		]
+		.ValueContent()
+		[
+			SNew(SHorizontalBox)	// Required to work around a text alignment issue. Otherwise the text block will not center align.
+			+SHorizontalBox::Slot()
+			.VAlign(EVerticalAlignment::VAlign_Center)
+			.HAlign(EHorizontalAlignment::HAlign_Left)
+			[
+				SNew(STextBlock)
+				.Text(FText::Format(LOCTEXT("HeaderTextValue", "Sample {0} / {1} Frames"), NumFramesToSample, NumAnimFrames))
+				.Font(StructCustomizationUtils.GetRegularFont())
+				.ColorAndOpacity(InputAnim.IsEnabled() ? FSlateColor::UseForeground() : FSlateColor::UseSubduedForeground())
+			]
+		];
+	}
+
+	void FMLDeformerGeomCacheTrainingInputAnimCustomization::CustomizeChildren(TSharedRef<IPropertyHandle> StructPropertyHandle, IDetailChildrenBuilder& ChildBuilder, IPropertyTypeCustomizationUtils& StructCustomizationUtils)
+	{
+		// Get all the child properties.
+		uint32 NumChildren;
+		StructPropertyHandle->GetNumChildren(NumChildren);
+		TMap<FName, TSharedPtr<IPropertyHandle>> PropertyHandles;
+		PropertyHandles.Reserve(NumChildren);
+		for (uint32 ChildIndex = 0; ChildIndex < NumChildren; ++ChildIndex)
+		{
+			TSharedRef<IPropertyHandle> ChildHandle = StructPropertyHandle->GetChildHandle(ChildIndex).ToSharedRef();
+			const FName PropertyName = ChildHandle->GetProperty()->GetFName();
+			PropertyHandles.Add(PropertyName, ChildHandle);
+		}
+
+		const int32 ArrayIndex = StructPropertyHandle->GetIndexInArray();
+		check(ArrayIndex != INDEX_NONE);
+
+		TSharedPtr<IPropertyHandle> AnimSequenceHandle		= PropertyHandles.FindChecked(FMLDeformerGeomCacheTrainingInputAnim::GetAnimSequencePropertyName());
+		TSharedPtr<IPropertyHandle> GeomCacheHandle			= PropertyHandles.FindChecked(FMLDeformerGeomCacheTrainingInputAnim::GetGeomCachePropertyName());
+		TSharedPtr<IPropertyHandle> EnabledHandle			= PropertyHandles.FindChecked(FMLDeformerGeomCacheTrainingInputAnim::GetEnabledPropertyName());
+		TSharedPtr<IPropertyHandle> UseCustomRangeHandle	= PropertyHandles.FindChecked(FMLDeformerGeomCacheTrainingInputAnim::GetUseCustomRangePropertyName());
+		TSharedPtr<IPropertyHandle> StartFrameHandle		= PropertyHandles.FindChecked(FMLDeformerGeomCacheTrainingInputAnim::GetStartFramePropertyName());
+		TSharedPtr<IPropertyHandle> EndFrameHandle			= PropertyHandles.FindChecked(FMLDeformerGeomCacheTrainingInputAnim::GetEndFramePropertyName());
+		check(AnimSequenceHandle.IsValid());
+		check(GeomCacheHandle.IsValid());
+		check(EnabledHandle.IsValid());
+		check(UseCustomRangeHandle.IsValid());
+		check(StartFrameHandle.IsValid());
+		check(EndFrameHandle.IsValid());
+
+		UMLDeformerGeomCacheModel* Model = FindMLDeformerModel(StructPropertyHandle);
+		check(Model);
+		FMLDeformerGeomCacheTrainingInputAnim& InputAnim = Model->GetTrainingInputAnims()[ArrayIndex];
+		UAnimSequence* AnimSequence = InputAnim.GetAnimSequence();
+		UGeometryCache* GeomCache = InputAnim.GetGeometryCache();
+
+		USkeleton* Skeleton = nullptr;
+		FMLDeformerEditorModel* EditorModel = nullptr;
+		if (Model)
+		{
+			if (Model->GetSkeletalMesh())
+			{
+				Skeleton = Model->GetSkeletalMesh()->GetSkeleton();
+			}
+
+			// Get the editor model for this runtime model.
+			FMLDeformerEditorModule& EditorModule = FModuleManager::LoadModuleChecked<FMLDeformerEditorModule>("MLDeformerFrameworkEditor");
+			EditorModel = EditorModule.GetModelRegistry().GetEditorModel(Model);
+		}
+
+		// Add the properties
+		IDetailPropertyRow& AnimRow = ChildBuilder.AddProperty(AnimSequenceHandle.ToSharedRef());
+		AnimRow.EditCondition(InputAnim.IsEnabled(), nullptr);
+		AnimRow.CustomWidget()
+		.NameContent()
+		[
+			AnimRow.GetPropertyHandle()->CreatePropertyNameWidget()					
+		]
+		.ValueContent()
+		[
+			SNew(SObjectPropertyEntryBox)
+			.PropertyHandle(AnimRow.GetPropertyHandle())
+			.AllowedClass(UAnimSequence::StaticClass())
+			.ObjectPath(AnimSequence ? AnimSequence->GetPathName() : FString())
+			.ThumbnailPool(StructCustomizationUtils.GetThumbnailPool())
+			.OnShouldFilterAsset
+			(
+				this, 
+				&FMLDeformerGeomCacheTrainingInputAnimCustomization::FilterAnimSequences, 
+				Skeleton
+			)
+		];
+
+		ChildBuilder.AddProperty(GeomCacheHandle.ToSharedRef())
+			.EditCondition(InputAnim.IsEnabled(), nullptr);
+
+		ChildBuilder.AddProperty(UseCustomRangeHandle.ToSharedRef())
+			.EditCondition(InputAnim.IsEnabled(), nullptr);
+
+		ChildBuilder.AddProperty(StartFrameHandle.ToSharedRef())
+			.EditCondition(InputAnim.IsEnabled() && InputAnim.GetUseCustomRange(), nullptr);
+
+		ChildBuilder.AddProperty(EndFrameHandle.ToSharedRef())
+			.EditCondition(InputAnim.IsEnabled() && InputAnim.GetUseCustomRange(), nullptr);
+
+		ChildBuilder.AddProperty(EnabledHandle.ToSharedRef());
+
+		// Verify the geometry cache against the anim sequence.
+		const FText GeomCacheErrorText = GetGeomCacheAnimSequenceErrorText(GeomCache, AnimSequence);
+		ChildBuilder.AddCustomRow(FText::FromString("AnimSeqWarning"))
+			.Visibility(!GeomCacheErrorText.IsEmpty() ? EVisibility::Visible : EVisibility::Collapsed)
+			.WholeRowContent()
+			[
+				SNew(SBox)
+				.Padding(FMargin(0.0f, 4.0f))
+				[
+					SNew(SWarningOrErrorBox)
+					.MessageStyle(EMessageStyle::Warning)
+					.Message(GeomCacheErrorText)
+				]
+			];
+
+		// Show some additional geometry cache errors.
+		USkeletalMesh* SkelMesh = Model->GetSkeletalMesh();
+		const FText TargetMeshErrorText = GetGeomCacheErrorText(SkelMesh, GeomCache);
+		ChildBuilder.AddCustomRow(FText::FromString("TargetMeshError"))
+			.Visibility(!TargetMeshErrorText.IsEmpty() ? EVisibility::Visible : EVisibility::Collapsed)
+			.WholeRowContent()
+			[
+				SNew(SBox)
+				.Padding(FMargin(0.0f, 4.0f))
+				[
+					SNew(SWarningOrErrorBox)
+					.MessageStyle(EMessageStyle::Error)
+					.Message(TargetMeshErrorText)
+				]
+			];
+
+		// Check if our vertex counts changed versus the trained model.
+		if (EditorModel)
+		{
+			const FText ChangedErrorText = EditorModel->GetTargetAssetChangedErrorText();
+			ChildBuilder.AddCustomRow(FText::FromString("TargetMeshChangedError"))
+				.Visibility(!ChangedErrorText.IsEmpty() ? EVisibility::Visible : EVisibility::Collapsed)
+				.WholeRowContent()
+				[
+					SNew(SBox)
+					.Padding(FMargin(0.0f, 4.0f))
+					[
+						SNew(SWarningOrErrorBox)
+						.MessageStyle(EMessageStyle::Warning)
+						.Message(ChangedErrorText)
+					]
+				];
+		}
+
+		// Check if the skeleton is incompatible or not.
+		const FText SkeletonError = EditorModel->GetIncompatibleSkeletonErrorText(Model->GetSkeletalMesh(), InputAnim.GetAnimSequence());
+		ChildBuilder.AddCustomRow(FText::FromString("AnimSkeletonMisMatchError"))
+			.Visibility(!SkeletonError.IsEmpty() ? EVisibility::Visible : EVisibility::Collapsed)
+			.WholeRowContent()
+			[
+				SNew(SBox)
+				.Padding(FMargin(0.0f, 4.0f))
+				[
+					SNew(SWarningOrErrorBox)
+					.MessageStyle(EMessageStyle::Warning)
+					.Message(SkeletonError)
+				]
+			];
+
+		// Show geometry cache mapping errors, which maps geom cache tracks to FBX meshes.
+		const FText MappingWarningText = GetGeomCacheMeshMappingErrorText(SkelMesh, GeomCache);
+		ChildBuilder.AddCustomRow(FText::FromString("MeshMappingWarning"))
+			.Visibility(!MappingWarningText.IsEmpty() ? EVisibility::Visible : EVisibility::Collapsed)
+			.WholeRowContent()
+			[
+				SNew(SBox)
+				.Padding(FMargin(0.0f, 4.0f))
+				[
+					SNew(SWarningOrErrorBox)
+					.MessageStyle(EMessageStyle::Warning)
+					.Message(MappingWarningText)
+				]
+			];
+	}
+
+	bool FMLDeformerGeomCacheTrainingInputAnimCustomization::FilterAnimSequences(const FAssetData& AssetData, USkeleton* Skeleton)
+	{
+		if (Skeleton && Skeleton->IsCompatibleForEditor(AssetData))
+		{
+			return false;
+		}
+
+		return true;
+	}
+
+	UMLDeformerGeomCacheModel* FMLDeformerGeomCacheTrainingInputAnimCustomization::FindMLDeformerModel(TSharedRef<IPropertyHandle> StructPropertyHandle) const
+	{
+		TArray<UObject*> Objects;
+		StructPropertyHandle->GetOuterObjects(Objects);
+		UObject* OwnerObject = !Objects.IsEmpty() ? Objects[0] : nullptr;
+		UMLDeformerGeomCacheModel* Model = Cast<UMLDeformerGeomCacheModel>(OwnerObject);
+		return Model;
+	}
+}	// namespace UE::MLDeformer
+
+#undef LOCTEXT_NAMESPACE
diff --git a/Engine/Plugins/Animation/MLDeformer/MLDeformerFramework/Source/MLDeformerFrameworkEditor/Private/MLDeformerGeomCacheTrainingModel.cpp b/Engine/Plugins/Animation/MLDeformer/MLDeformerFramework/Source/MLDeformerFrameworkEditor/Private/MLDeformerGeomCacheTrainingModel.cpp
new file mode 100644
--- /dev/null
+++ b/Engine/Plugins/Animation/MLDeformer/MLDeformerFramework/Source/MLDeformerFrameworkEditor/Private/MLDeformerGeomCacheTrainingModel.cpp
@@ -0,0 +1,96 @@
+// Copyright Epic Games, Inc. All Rights Reserved.
+
+#include "MLDeformerGeomCacheTrainingModel.h"
+#include "MLDeformerGeomCacheTrainingInputAnim.h"
+#include "MLDeformerGeomCacheEditorModel.h"
+#include "MLDeformerGeomCacheModel.h"
+
+void UMLDeformerGeomCacheTrainingModel::Init(UE::MLDeformer::FMLDeformerEditorModel* InEditorModel)
+{
+	Super::Init(InEditorModel);
+
+	// Find the first valid input anim index to sample from.
+	// This modifies the SampleAnimIndex value.
+	bFinishedSampling = !FindNextAnimToSample(SampleAnimIndex);
+}
+
+bool UMLDeformerGeomCacheTrainingModel::FindNextAnimToSample(int32& OutNextAnimIndex) const
+{
+	using namespace UE::MLDeformer;
+
+	UMLDeformerGeomCacheModel* GeomCacheModel = Cast<UMLDeformerGeomCacheModel>(EditorModel->GetModel());
+	int32 NumTries = 0;
+	int32 AnimIndex = SampleAnimIndex;
+
+	const int32 NumInputAnims = EditorModel->GetNumTrainingInputAnims();
+	while (NumTries < NumInputAnims)	// Try all input animations at worst case.
+	{
+		const FMLDeformerGeomCacheTrainingInputAnim& InputAnim = GeomCacheModel->GetTrainingInputAnims()[AnimIndex];
+		if (InputAnim.IsEnabled())
+		{
+			if (NumTimesSampled[AnimIndex] < InputAnim.GetNumFramesToSample())
+			{
+				OutNextAnimIndex = AnimIndex;
+				return true;
+			}
+		}
+
+		// Get the next animation index.
+		AnimIndex++;
+		AnimIndex %= EditorModel->GetNumTrainingInputAnims();
+
+		NumTries++;
+	}
+
+	OutNextAnimIndex = INDEX_NONE;
+    return false;
+}
+
+bool UMLDeformerGeomCacheTrainingModel::SampleNextFrame()
+{
+	using namespace UE::MLDeformer;
+
+	// Make sure that there is more left to sample.
+	if (bFinishedSampling)
+	{
+		return false;
+	}
+
+	FMLDeformerGeomCacheEditorModel* GeomEditorModel = static_cast<FMLDeformerGeomCacheEditorModel*>(EditorModel);
+	UMLDeformerGeomCacheModel* GeomCacheModel = Cast<UMLDeformerGeomCacheModel>(EditorModel->GetModel());
+
+	// Get the animation to sample from and validate some things.
+	const FMLDeformerGeomCacheTrainingInputAnim& InputAnim = GeomCacheModel->GetTrainingInputAnims()[SampleAnimIndex];
+	check(InputAnim.IsEnabled());
+	check(InputAnim.GetAnimSequence());
+	check(InputAnim.GetGeometryCache());
+
+	const int32 StartFrame = InputAnim.GetUseCustomRange() ? FMath::Min<int32>(InputAnim.GetStartFrame(), InputAnim.GetEndFrame()) : 0;
+	const int32 CurFrameToSample = StartFrame + NumTimesSampled[SampleAnimIndex];
+	check(CurFrameToSample < StartFrame + InputAnim.GetNumFramesToSample());	// We should never sample more frames than the animation has.
+	NumTimesSampled[SampleAnimIndex]++;
+
+	// Perform the actual sampling.
+	FMLDeformerSampler* Sampler = EditorModel->GetSamplerForTrainingAnim(SampleAnimIndex);
+	Sampler->SetVertexDeltaSpace(EVertexDeltaSpace::PreSkinning);
+	Sampler->Sample(CurFrameToSample);
+	UE_LOG(LogMLDeformer, Display, TEXT("Sampling frame %d of anim %d"), CurFrameToSample, SampleAnimIndex);
+
+	// Copy sampled values.
+	SampleDeltas = Sampler->GetVertexDeltas();
+	SampleBoneRotations = Sampler->GetBoneRotations();
+	SampleCurveValues = Sampler->GetCurveValues();
+
+	// Now find the next sample we should take.
+	// This will return false when we finished sampling everything.
+	SampleAnimIndex++;
+	SampleAnimIndex %= GeomCacheModel->GetTrainingInputAnims().Num();
+	if (!FindNextAnimToSample(SampleAnimIndex))
+	{
+		bFinishedSampling = true;
+		return false;
+	}
+
+	// Let the caller know we can still sample more.
+	return true;
+}
diff --git a/Engine/Plugins/Animation/MLDeformer/MLDeformerFramework/Source/MLDeformerFrameworkEditor/Private/MLDeformerModelDetails.cpp b/Engine/Plugins/Animation/MLDeformer/MLDeformerFramework/Source/MLDeformerFrameworkEditor/Private/MLDeformerModelDetails.cpp
--- a/Engine/Plugins/Animation/MLDeformer/MLDeformerFramework/Source/MLDeformerFrameworkEditor/Private/MLDeformerModelDetails.cpp
+++ b/Engine/Plugins/Animation/MLDeformer/MLDeformerFramework/Source/MLDeformerFrameworkEditor/Private/MLDeformerModelDetails.cpp
@@ -9,6 +9,7 @@
 #include "MLDeformerEditorModel.h"
 #include "MLDeformerEditorModule.h"
 #include "MLDeformerGeomCacheHelpers.h"
+#include "MLDeformerTrainingInputAnim.h"
 #include "GeometryCache.h"
 #include "GeometryCacheTrack.h"
 #include "Animation/AnimSequence.h"
@@ -55,7 +56,6 @@ namespace UE::MLDeformer
 	void FMLDeformerModelDetails::CreateCategories()
 	{
 		BaseMeshCategoryBuilder = &DetailLayoutBuilder->EditCategory("Base Mesh", FText::GetEmpty(), ECategoryPriority::Important);
-		TargetMeshCategoryBuilder = &DetailLayoutBuilder->EditCategory("Target Mesh", FText::GetEmpty(), ECategoryPriority::Important);
 		InputOutputCategoryBuilder = &DetailLayoutBuilder->EditCategory("Inputs", FText::GetEmpty(), ECategoryPriority::Important);
 		TrainingSettingsCategoryBuilder = &DetailLayoutBuilder->EditCategory("Training Settings", FText::GetEmpty(), ECategoryPriority::Important);
 	}
@@ -124,47 +124,9 @@ namespace UE::MLDeformer
 				]
 			];
 
-		// Animation sequence.
-		IDetailPropertyRow& AnimRow = BaseMeshCategoryBuilder->AddProperty(UMLDeformerModel::GetAnimSequencePropertyName(), UMLDeformerModel::StaticClass());
-		AnimRow.CustomWidget()
-		.NameContent()
-		[
-			AnimRow.GetPropertyHandle()->CreatePropertyNameWidget()
-		]
-		.ValueContent()
-		[
-			SNew(SObjectPropertyEntryBox)
-			.PropertyHandle(AnimRow.GetPropertyHandle())
-			.AllowedClass(UAnimSequence::StaticClass())
-			.ObjectPath( (Model && Model->GetAnimSequence()) ? Model->GetAnimSequence()->GetPathName() : FString())
-			.ThumbnailPool(DetailBuilder.GetThumbnailPool())
-			.OnShouldFilterAsset(
-				this, 
-				&FMLDeformerModelDetails::FilterAnimSequences, 
-				Model->GetSkeletalMesh() ? Model->GetSkeletalMesh()->GetSkeleton() : nullptr
-			)
-		];
+		AddTrainingInputAnims();
 
-		AddAnimSequenceErrors();
-
-		const FText AnimErrorText = EditorModel->GetIncompatibleSkeletonErrorText(Model->GetSkeletalMesh(), Model->GetAnimSequence());
-		FDetailWidgetRow& AnimErrorRow = BaseMeshCategoryBuilder->AddCustomRow(FText::FromString("AnimSkeletonMisMatchError"))
-			.Visibility(!AnimErrorText.IsEmpty() ? EVisibility::Visible : EVisibility::Collapsed)
-			.WholeRowContent()
-			[
-				SNew(SBox)
-				.Padding(FMargin(0.0f, 4.0f))
-				[
-					SNew(SWarningOrErrorBox)
-					.MessageStyle(EMessageStyle::Warning)
-					.Message(AnimErrorText)
-				]
-			];
-
-
-		AddTargetMesh();
-
-		TargetMeshCategoryBuilder->AddProperty(UMLDeformerModel::GetAlignmentTransformPropertyName(), UMLDeformerModel::StaticClass());
+		InputOutputCategoryBuilder->AddProperty(UMLDeformerModel::GetAlignmentTransformPropertyName(), UMLDeformerModel::StaticClass());
 
 		AddTrainingInputFlags();
 		AddTrainingInputErrors();
@@ -189,14 +151,12 @@ namespace UE::MLDeformer
 		// Create the inputs widget.
 		TSharedPtr<SMLDeformerInputWidget> InputWidget = EditorModel->CreateInputWidget();
 		EditorModel->SetInputWidget(InputWidget);
-		InputOutputCategoryBuilder->AddCustomRow(FText::FromString("Inputs")).WholeRowContent().Widget = InputWidget.ToSharedRef();
+		IDetailGroup& NetworkInputsGroup = InputOutputCategoryBuilder->AddGroup("NetworkInputs", LOCTEXT("NetworkInputs", "Network Inputs"), false, true);
+		NetworkInputsGroup.AddWidgetRow().WholeRowContent().Widget = InputWidget.ToSharedRef();
 
 		AddBoneInputErrors();
 		AddCurveInputErrors();
 
-		InputOutputCategoryBuilder->AddProperty(UMLDeformerModel::GetMaxTrainingFramesPropertyName(), UMLDeformerModel::StaticClass());
-		InputOutputCategoryBuilder->AddProperty(UMLDeformerModel::GetDeltaCutoffLengthPropertyName(), UMLDeformerModel::StaticClass());
-
 		// Show a warning when no neural network has been set.
 		{		
 			FDetailWidgetRow& NeuralNetErrorRow = TrainingSettingsCategoryBuilder->AddCustomRow(FText::FromString("NeuralNetError"))
@@ -230,6 +190,8 @@ namespace UE::MLDeformer
 			}
 		}
 
+		TrainingSettingsCategoryBuilder->AddProperty(UMLDeformerModel::GetMaxTrainingFramesPropertyName(), UMLDeformerModel::StaticClass());
+		TrainingSettingsCategoryBuilder->AddProperty(UMLDeformerModel::GetDeltaCutoffLengthPropertyName(), UMLDeformerModel::StaticClass());
 		AddTrainingSettingsErrors();
 	}
 
@@ -243,30 +205,9 @@ namespace UE::MLDeformer
 		return true;
 	}
 
-	FReply FMLDeformerModelDetails::OnFilterAnimatedBonesOnly() const
-	{
-		EditorModel->InitBoneIncludeListToAnimatedBonesOnly();
-		EditorModel->SetResamplingInputOutputsNeeded(true);
-		DetailLayoutBuilder->ForceRefreshDetails();
-		return FReply::Handled();
-	}
-
-	FReply FMLDeformerModelDetails::OnFilterAnimatedCurvesOnly() const
-	{
-		EditorModel->InitCurveIncludeListToAnimatedCurvesOnly();
-		EditorModel->SetResamplingInputOutputsNeeded(true);
-		DetailLayoutBuilder->ForceRefreshDetails();
-		return FReply::Handled();
-	}
-
-	bool FMLDeformerModelDetails::IsBonesFlagVisible() const
-	{
-		return (Model->DoesSupportBones() && Model->DoesSupportCurves());
-	}
-
-	bool FMLDeformerModelDetails::IsCurvesFlagVisible() const
+	void FMLDeformerModelDetails::AddTrainingInputAnims()
 	{
-		return (Model->DoesSupportBones() && Model->DoesSupportCurves() && EditorModel->GetNumCurvesOnSkeletalMesh(Model->GetSkeletalMesh()) > 0);
+		UE_LOG(LogMLDeformer, Warning, TEXT("Please implement the AddTrainingInputAnims method in your model details class."));
 	}
 }	// namespace UE::MLDeformer
 
diff --git a/Engine/Plugins/Animation/MLDeformer/MLDeformerFramework/Source/MLDeformerFrameworkEditor/Private/MLDeformerMorphModelEditorModel.cpp b/Engine/Plugins/Animation/MLDeformer/MLDeformerFramework/Source/MLDeformerFrameworkEditor/Private/MLDeformerMorphModelEditorModel.cpp
--- a/Engine/Plugins/Animation/MLDeformer/MLDeformerFramework/Source/MLDeformerFrameworkEditor/Private/MLDeformerMorphModelEditorModel.cpp
+++ b/Engine/Plugins/Animation/MLDeformer/MLDeformerFramework/Source/MLDeformerFrameworkEditor/Private/MLDeformerMorphModelEditorModel.cpp
@@ -402,8 +402,10 @@ namespace UE::MLDeformer
 		EMLDeformerMaskChannel MaskChannel,
 		bool bInvertMaskChannel)
 	{
+		FMLDeformerSampler* Sampler = GetSamplerForActiveAnim();
+
 		OutMorphTargets.Reset();
-		if (Deltas.IsEmpty())
+		if (Deltas.IsEmpty() || Sampler == nullptr)
 		{
 			return;
 		}
@@ -566,6 +568,12 @@ namespace UE::MLDeformer
 
 	void FMLDeformerMorphModelEditorModel::DebugDrawMorphTarget(FPrimitiveDrawInterface* PDI, const TArray<FVector3f>& MorphDeltas, float DeltaThreshold, int32 MorphTargetIndex, const FVector& DrawOffset)
 	{
+		FMLDeformerSampler* Sampler = GetSamplerForActiveAnim();
+		if (Sampler == nullptr)
+		{
+			return;
+		}
+
 		UMLDeformerVizSettings* VizSettings = Model->GetVizSettings();
 		const int32 NumVerts = Model->GetNumBaseMeshVerts();
 		const TArray<FVector3f>& UnskinnedPositions = Sampler->GetUnskinnedVertexPositions();
diff --git a/Engine/Plugins/Animation/MLDeformer/MLDeformerFramework/Source/MLDeformerFrameworkEditor/Private/MLDeformerSampler.cpp b/Engine/Plugins/Animation/MLDeformer/MLDeformerFramework/Source/MLDeformerFrameworkEditor/Private/MLDeformerSampler.cpp
--- a/Engine/Plugins/Animation/MLDeformer/MLDeformerFramework/Source/MLDeformerFrameworkEditor/Private/MLDeformerSampler.cpp
+++ b/Engine/Plugins/Animation/MLDeformer/MLDeformerFramework/Source/MLDeformerFrameworkEditor/Private/MLDeformerSampler.cpp
@@ -1,11 +1,10 @@
 // Copyright Epic Games, Inc. All Rights Reserved.
 
 #include "MLDeformerSampler.h"
-
 #include "MLDeformerInputInfo.h"
 #include "MLDeformerModel.h"
 #include "MLDeformerEditorToolkit.h"
-
+#include "MLDeformerTrainingInputAnim.h"
 #include "Animation/DebugSkelMeshComponent.h"
 #include "Animation/AnimSequence.h"
 #include "Animation/AnimInstance.h"
@@ -34,8 +33,13 @@ namespace UE::MLDeformer
 		}
 	}
 
-	// Initializer a sampler item.
 	void FMLDeformerSampler::Init(FMLDeformerEditorModel* InEditorModel)
+	{
+		Init(InEditorModel, 0);
+	}
+
+	// Initializer a sampler item.
+	void FMLDeformerSampler::Init(FMLDeformerEditorModel* InEditorModel, int32 InAnimIndex)
 	{
 		check(InEditorModel);
 		if (InEditorModel->GetEditor()->GetPersonaToolkitPointer() == nullptr)
@@ -47,13 +51,13 @@ namespace UE::MLDeformer
 		Model = EditorModel->GetModel();
 		NumImportedVertices = UMLDeformerModel::ExtractNumImportedSkinnedVertices(Model->GetSkeletalMesh());
 		AnimFrameIndex = 0;
+		AnimIndex = InAnimIndex;
 		SampleTime = 0.0f;
 		NumFloatsPerCurve = 1;
 
 		// Create the actors and components.
 		// This internally skips creating them if they already exist.
 		CreateActors();
-		RegisterTargetComponents();
 
 		SkinnedVertexPositions.Reset();
 		SkinnedVertexPositions.AddUninitialized(NumImportedVertices);
@@ -102,6 +106,11 @@ namespace UE::MLDeformer
 
 	void FMLDeformerSampler::Sample(int32 InAnimFrameIndex)
 	{
+		UAnimSequence* TrainingAnimSequence = EditorModel->GetTrainingInputAnim(AnimIndex)->GetAnimSequence();
+		check(TrainingAnimSequence);
+		const EAnimInterpolationType InterpolationTypeBackup = TrainingAnimSequence->Interpolation;
+		TrainingAnimSequence->Interpolation = EAnimInterpolationType::Step;
+
 		AnimFrameIndex = InAnimFrameIndex;
 		SampleTime = GetTimeAtFrame(InAnimFrameIndex);
 
@@ -121,6 +130,8 @@ namespace UE::MLDeformer
 			VertexDeltas.Reset(NumFloats);
 			VertexDeltas.AddZeroed(NumFloats);
 		}
+
+		TrainingAnimSequence->Interpolation = InterpolationTypeBackup;
 	}
 
 	// Calculate the inverse skinning transform. This is basically inv(sum(BoneTransform_i * inv(BoneRestTransform_i) * Weight_i)), where i is for each skinning influence for the given vertex.
@@ -245,7 +256,7 @@ namespace UE::MLDeformer
 
 		// Create the skeletal mesh component.
 		USkeletalMesh* SkeletalMesh = Model->GetSkeletalMesh();
-		UAnimSequence* TrainingAnimSequence = Model->GetAnimSequence();
+		UAnimSequence* TrainingAnimSequence = EditorModel->GetTrainingInputAnim(AnimIndex)->GetAnimSequence();
 		if (SkeletalMeshComponent.Get() == nullptr)
 		{
 			SkeletalMeshComponent = NewObject<UDebugSkelMeshComponent>(SkelMeshActor);
@@ -275,7 +286,7 @@ namespace UE::MLDeformer
 	AActor* FMLDeformerSampler::CreateNewActor(UWorld* InWorld, const FName& Name) const
 	{
 		FActorSpawnParameters SpawnParams;
-		SpawnParams.Name = NAME_None;
+		SpawnParams.Name = MakeUniqueObjectName(InWorld, AActor::StaticClass(), Name);
 		AActor* Actor = InWorld->SpawnActor<AActor>(SpawnParams);
 		Actor->SetFlags(RF_Transient);
 		return Actor;
diff --git a/Engine/Plugins/Animation/MLDeformer/MLDeformerFramework/Source/MLDeformerFrameworkEditor/Private/MLDeformerTrainingModel.cpp b/Engine/Plugins/Animation/MLDeformer/MLDeformerFramework/Source/MLDeformerFrameworkEditor/Private/MLDeformerTrainingModel.cpp
--- a/Engine/Plugins/Animation/MLDeformer/MLDeformerFramework/Source/MLDeformerFrameworkEditor/Private/MLDeformerTrainingModel.cpp
+++ b/Engine/Plugins/Animation/MLDeformer/MLDeformerFramework/Source/MLDeformerFrameworkEditor/Private/MLDeformerTrainingModel.cpp
@@ -14,6 +14,7 @@ UMLDeformerModel* UMLDeformerTrainingModel::GetModel() const
 void UMLDeformerTrainingModel::Init(UE::MLDeformer::FMLDeformerEditorModel* InEditorModel)
 {
 	EditorModel = InEditorModel;
+	ResetSampling();
 }
 
 void UMLDeformerTrainingModel::SetEditorModel(UE::MLDeformer::FMLDeformerEditorModel* InModel)
@@ -41,6 +42,14 @@ int32 UMLDeformerTrainingModel::NumSamples() const
 	return EditorModel->GetNumFramesForTraining();
 }
 
+void UMLDeformerTrainingModel::ResetSampling()
+{
+	NumTimesSampled.Reset();
+	NumTimesSampled.AddZeroed(EditorModel->GetNumTrainingInputAnims());
+	SampleAnimIndex = 0;
+	bFinishedSampling = !FindNextAnimToSample(SampleAnimIndex);
+}
+
 int32 UMLDeformerTrainingModel::GetNumberSampleDeltas() const
 {
 	return EditorModel->GetEditorInputInfo()->GetNumBaseMeshVertices();
@@ -48,12 +57,20 @@ int32 UMLDeformerTrainingModel::GetNumberSampleDeltas() const
 
 void UMLDeformerTrainingModel::SetNumFloatsPerCurve(int32 NumFloatsPerCurve)
 {
-	EditorModel->GetSampler()->SetNumFloatsPerCurve(NumFloatsPerCurve);
+	const int32 NumAnims = EditorModel->GetNumTrainingInputAnims();
+	for (int32 AnimIndex = 0; AnimIndex < NumAnims; ++AnimIndex)
+	{
+		UE::MLDeformer::FMLDeformerSampler* Sampler = EditorModel->GetSamplerForTrainingAnim(AnimIndex);
+		if (Sampler)
+		{
+			Sampler->SetNumFloatsPerCurve(NumFloatsPerCurve);
+		}
+	}
 }
 
 bool UMLDeformerTrainingModel::SetCurrentSampleIndex(int32 Index)
 {
-	return SampleFrame(Index);
+	return NextSample();
 }
 
 bool UMLDeformerTrainingModel::GetNeedsResampling() const
@@ -66,26 +83,19 @@ void UMLDeformerTrainingModel::SetNeedsResampling(bool bNeedsResampling)
 	EditorModel->SetResamplingInputOutputsNeeded(bNeedsResampling);
 }
 
-bool UMLDeformerTrainingModel::SampleFrame(int32 Index)
+bool UMLDeformerTrainingModel::NextSample()
 {
-	using namespace UE::MLDeformer;
-
-	// Make sure we have a valid frame number.
-	if (Index < 0 || Index >= EditorModel->GetNumTrainingFrames())
-	{
-		UE_LOG(LogMLDeformer, Warning, TEXT("Sample index must range from %d to %d, but a value of %d was provided."), 0, EditorModel->GetNumTrainingFrames()-1, Index);
-		return false;
-	}
-
-	// Sample the frame.
-	FMLDeformerSampler* Sampler = EditorModel->GetSampler();
-	Sampler->SetVertexDeltaSpace(EVertexDeltaSpace::PreSkinning);
-	Sampler->Sample(Index);
+	return SampleNextFrame();
+}
 
-	// Copy sampled values.
-	SampleDeltas = Sampler->GetVertexDeltas();
-	SampleBoneRotations = Sampler->GetBoneRotations();
-	SampleCurveValues = Sampler->GetCurveValues();
+bool UMLDeformerTrainingModel::SampleNextFrame()
+{
+	UE_LOG(LogMLDeformer, Warning, TEXT("Please override the SampleNextFrame method in your UMLDeformerTrainingModel inherited class."));
+	return false;
+}
 
-	return true;
+bool UMLDeformerTrainingModel::SampleFrame(int32 Index)
+{
+	UE_LOG(LogMLDeformer, Warning, TEXT("Please use UMLDeformerTrainingModel::NextSample() instead."));
+	return false;
 }
diff --git a/Engine/Plugins/Animation/MLDeformer/MLDeformerFramework/Source/MLDeformerFrameworkEditor/Private/SMLDeformerInputWidget.cpp b/Engine/Plugins/Animation/MLDeformer/MLDeformerFramework/Source/MLDeformerFrameworkEditor/Private/SMLDeformerInputWidget.cpp
--- a/Engine/Plugins/Animation/MLDeformer/MLDeformerFramework/Source/MLDeformerFrameworkEditor/Private/SMLDeformerInputWidget.cpp
+++ b/Engine/Plugins/Animation/MLDeformer/MLDeformerFramework/Source/MLDeformerFrameworkEditor/Private/SMLDeformerInputWidget.cpp
@@ -33,56 +33,51 @@ namespace UE::MLDeformer
 		SectionVerticalBox->AddSlot()
 		.AutoHeight()
 		[
-			SNew(SBorder)
+			SNew(SExpandableArea)
 			.BorderImage(FAppStyle::Get().GetBrush("Brushes.Header"))
-			.Padding(0.0f, 2.0f, 0.0f, 0.0f)
+			.BodyBorderImage(FAppStyle::Get().GetBrush("Brushes.Recessed"))
+			.HeaderPadding(FMargin(4.0f, 2.0f))
+			.Padding(0.0f)
+			.MaxHeight(300.0f)
+			.AllowAnimatedTransition(false)
+			.HeaderContent()
 			[
-				SNew(SExpandableArea)
-				.BorderImage(FAppStyle::Get().GetBrush("Brushes.Header"))
-				.BodyBorderImage(FAppStyle::Get().GetBrush("Brushes.Recessed"))
-				.HeaderPadding(FMargin(4.0f, 2.0f))
-				.Padding(0.0f)
-				.MaxHeight(300.0f)
-				.AllowAnimatedTransition(false)
-				.HeaderContent()
+				SNew(SHorizontalBox)
+				+SHorizontalBox::Slot()
+				.VAlign(VAlign_Center)
 				[
-					SNew(SHorizontalBox)
-					+SHorizontalBox::Slot()
-					.VAlign(VAlign_Center)
-					[
-						SNew(STextBlock)
-						.Text(SectionInfo.SectionTitle)
-						.TextStyle(FAppStyle::Get(), "ButtonText")
-						.Font(FAppStyle::Get().GetFontStyle("NormalFontBold"))
-					]
-					+SHorizontalBox::Slot()
-					.VAlign(VAlign_Center)
-					.HAlign(HAlign_Right)
-					.AutoWidth()
-					.Padding(4.0f, 0.0f, 0.0f, 0.0f)
-					[
-						SNew(SButton)
-						.ButtonStyle(FAppStyle::Get(), "SimpleButton")
-						.ToolTipText(SectionInfo.PlusButtonTooltip)
-						.OnClicked(SectionInfo.PlusButtonPressed)
-						.ContentPadding(FMargin(1, 0))					
-						[
-							SNew(SImage)
-							.Image(FAppStyle::Get().GetBrush("Icons.PlusCircle"))
-							.ColorAndOpacity(FSlateColor::UseForeground())
-						]
-					]
+					SNew(STextBlock)
+					.Text(SectionInfo.SectionTitle)
+					.TextStyle(FAppStyle::Get(), "ButtonText")
+					.Font(FAppStyle::Get().GetFontStyle("NormalFontBold"))
 				]
-				.BodyContent()			
+				+SHorizontalBox::Slot()
+				.VAlign(VAlign_Center)
+				.HAlign(HAlign_Right)
+				.AutoWidth()
+				.Padding(4.0f, 0.0f, 0.0f, 0.0f)
 				[
-					SNew(SBorder)
-					.BorderImage(FAppStyle::Get().GetBrush("Brushes.Header"))
-					.Padding(FMargin(8.0f, 2.0f))
+					SNew(SButton)
+					.ButtonStyle(FAppStyle::Get(), "SimpleButton")
+					.ToolTipText(SectionInfo.PlusButtonTooltip)
+					.OnClicked(SectionInfo.PlusButtonPressed)
+					.ContentPadding(FMargin(1, 0))					
 					[
-						Widget.ToSharedRef()
+						SNew(SImage)
+						.Image(FAppStyle::Get().GetBrush("Icons.PlusCircle"))
+						.ColorAndOpacity(FSlateColor::UseForeground())
 					]
 				]
 			]
+			.BodyContent()			
+			[
+				SNew(SBorder)
+				.BorderImage(FAppStyle::Get().GetBrush("Brushes.Header"))
+				.Padding(FMargin(8.0f, 2.0f))
+				[
+					Widget.ToSharedRef()
+				]
+			]
 		];
 	}
 
diff --git a/Engine/Plugins/Animation/MLDeformer/MLDeformerFramework/Source/MLDeformerFrameworkEditor/Private/SMLDeformerTimeline.cpp b/Engine/Plugins/Animation/MLDeformer/MLDeformerFramework/Source/MLDeformerFrameworkEditor/Private/SMLDeformerTimeline.cpp
--- a/Engine/Plugins/Animation/MLDeformer/MLDeformerFramework/Source/MLDeformerFrameworkEditor/Private/SMLDeformerTimeline.cpp
+++ b/Engine/Plugins/Animation/MLDeformer/MLDeformerFramework/Source/MLDeformerFrameworkEditor/Private/SMLDeformerTimeline.cpp
@@ -23,6 +23,7 @@
 #include "Styling/AppStyle.h"
 #include "Styling/ISlateStyle.h"
 #include "TimeSliderArgs.h"
+#include "Widgets/Input/SComboBox.h"
 #include "Widgets/Input/SButton.h"
 #include "Widgets/Input/SSearchBox.h"
 #include "Widgets/Input/SSpinBox.h"
@@ -1478,6 +1479,45 @@ namespace UE::MLDeformer
 						[
 							SNew(SHorizontalBox)
 							+SHorizontalBox::Slot()
+							.AutoWidth()
+							.Padding(2.0f)
+							[
+								SNew(SComboBox<TSharedPtr<FMLDeformerTrainingInputAnimName>>)
+								.OptionsSource(&TrainingAnimNames)
+								.OnGenerateWidget(this, &SMLDeformerTimeline::OnGenerateAnimListComboWidget)
+								.OnSelectionChanged(this, &SMLDeformerTimeline::OnSelectTrainingAnim)
+								.Visibility_Lambda
+								(
+									[this]()
+									{
+										if (!Model.IsValid())
+										{
+											return EVisibility::Hidden;
+										}
+
+										TSharedPtr<FMLDeformerEditorModel> EditorModel = Model.Pin();
+										check(EditorModel.IsValid());										
+										const EMLDeformerVizMode VizMode = EditorModel->GetModel()->GetVizSettings()->GetVisualizationMode();
+										return (VizMode == EMLDeformerVizMode::TrainingData && EditorModel->GetNumTrainingInputAnims() > 0) ? EVisibility::Visible : EVisibility::Hidden;
+									}
+								)
+								[
+									SNew(STextBlock)
+									.Text_Lambda
+									(
+										[this]()
+										{
+											if (SelectedTrainingAnimNameIndex != INDEX_NONE && TrainingAnimNames.IsValidIndex(SelectedTrainingAnimNameIndex))
+											{
+												return FText::FromString(TrainingAnimNames[SelectedTrainingAnimNameIndex]->Name);
+											}
+
+											return FText::FromString("Select Animation");
+										}
+									)
+								]
+							]
+							+SHorizontalBox::Slot()
 							.FillWidth(1.0f)
 							.VAlign(VAlign_Center)
 							+SHorizontalBox::Slot()
@@ -1835,6 +1875,27 @@ namespace UE::MLDeformer
 		TransportControls->SetModel(Model);
 	}
 
+	void SMLDeformerTimeline::SetTrainingAnimNames(const TArray<TSharedPtr<FMLDeformerTrainingInputAnimName>>& Names)
+	{
+		TrainingAnimNames = Names;
+
+		const int32 OldSelectedItemIndex = SelectedTrainingAnimNameIndex;
+		if (!Names.IsEmpty() && SelectedTrainingAnimNameIndex == INDEX_NONE)
+		{
+			SelectedTrainingAnimNameIndex = 0;
+		}
+
+		if (Names.IsEmpty())
+		{
+			SelectedTrainingAnimNameIndex = INDEX_NONE;
+		}
+
+		if (OldSelectedItemIndex != SelectedTrainingAnimNameIndex && SelectedTrainingAnimNameIndex != INDEX_NONE)
+		{
+			OnSelectTrainingAnim(TrainingAnimNames[SelectedTrainingAnimNameIndex], ESelectInfo::Type::Direct);
+		}
+	}
+
 	void SMLDeformerTimeline::OnColumnFillCoefficientChanged(float FillCoefficient, int32 ColumnIndex)
 	{
 		ColumnFillCoefficients[ColumnIndex] = FillCoefficient;
@@ -1873,6 +1934,60 @@ namespace UE::MLDeformer
 			Model.Pin()->SetScrubPosition(FrameTime);
 		}
 	}
+
+	TSharedRef<SWidget> SMLDeformerTimeline::OnGenerateAnimListComboWidget(TSharedPtr<FMLDeformerTrainingInputAnimName> Item) const
+	{
+		return SNew(STextBlock)
+			.Text(FText::FromString(Item->Name));
+	}
+
+	void SMLDeformerTimeline::OnSelectTrainingAnim(TSharedPtr<FMLDeformerTrainingInputAnimName> Item, ESelectInfo::Type SelectInfo)
+	{
+		if (Item.IsValid())
+		{
+			SelectedTrainingAnimNameIndex = TrainingAnimNames.Find(Item);
+
+			if (Model.IsValid())
+			{
+				TSharedPtr<FMLDeformerEditorModel> EditorModel = Model.Pin();
+				EditorModel->SetActiveTrainingInputAnimIndex(Item->TrainingInputAnimIndex);
+				if (EditorModel->GetEditor())
+				{
+					EditorModel->GetEditor()->UpdateTimeSliderRange();					
+				}
+				EditorModel->TriggerInputAssetChanged();
+			}
+		}
+	}
+
+	void SMLDeformerTimeline::OnDeletedTrainingInputAnim(int32 Index)
+	{
+		const int32 OldSelectedItemIndex = SelectedTrainingAnimNameIndex;
+		if (TrainingAnimNames[SelectedTrainingAnimNameIndex]->TrainingInputAnimIndex >= Index)
+		{
+			// Find a new valid selected item index.
+			SelectedTrainingAnimNameIndex--;
+			if (SelectedTrainingAnimNameIndex < 0 && !TrainingAnimNames.IsEmpty())
+			{
+				SelectedTrainingAnimNameIndex = 0;
+			}
+
+			if (SelectedTrainingAnimNameIndex < 0)
+			{
+				SelectedTrainingAnimNameIndex = INDEX_NONE;
+			}
+
+			if (SelectedTrainingAnimNameIndex > TrainingAnimNames.Num() - 1)
+			{
+				SelectedTrainingAnimNameIndex = TrainingAnimNames.Num() - 1;
+			}
+		}
+
+		if (OldSelectedItemIndex != SelectedTrainingAnimNameIndex)
+		{
+			OnSelectTrainingAnim(TrainingAnimNames[SelectedTrainingAnimNameIndex], ESelectInfo::Type::Direct);
+		}
+	}
 }	// namespace UE::MLDeformer
 
 #undef LOCTEXT_NAMESPACE
diff --git a/Engine/Plugins/Animation/MLDeformer/MLDeformerFramework/Source/MLDeformerFrameworkEditor/Private/SMLDeformerTimeline.h b/Engine/Plugins/Animation/MLDeformer/MLDeformerFramework/Source/MLDeformerFrameworkEditor/Private/SMLDeformerTimeline.h
--- a/Engine/Plugins/Animation/MLDeformer/MLDeformerFramework/Source/MLDeformerFrameworkEditor/Private/SMLDeformerTimeline.h
+++ b/Engine/Plugins/Animation/MLDeformer/MLDeformerFramework/Source/MLDeformerFrameworkEditor/Private/SMLDeformerTimeline.h
@@ -13,6 +13,7 @@
 #include "ITransportControl.h"
 #include "Framework/Commands/Commands.h"
 #include "MLDeformerEditorToolkit.h"
+#include "MLDeformerTrainingInputAnim.h"
 
 struct FPaintPlaybackRangeArgs;
 class SSearchBox;
@@ -60,12 +61,20 @@ namespace UE::MLDeformer
 
 		void SetModel(TWeakPtr<FMLDeformerEditorModel> InModel);
 
+		/** Set the list of names to show in the timeline animation list combobox. */
+		void SetTrainingAnimNames(const TArray<TSharedPtr<FMLDeformerTrainingInputAnimName>>& Names);
+
+		void OnDeletedTrainingInputAnim(int32 Index);
+
 	private:
 		/**
 		 * @return The fill percentage of the animation outliner.
 		 */
 		float GetColumnFillCoefficient(int32 ColumnIndex) const { return ColumnFillCoefficients[ColumnIndex]; }
 
+		TSharedRef<SWidget> OnGenerateAnimListComboWidget(TSharedPtr<FMLDeformerTrainingInputAnimName> Item) const;
+		void OnSelectTrainingAnim(TSharedPtr<FMLDeformerTrainingInputAnimName> Item, ESelectInfo::Type SelectInfo);
+
 		/** Get numeric Type interface for converting between frame numbers and display formats. */
 		TSharedRef<INumericTypeInterface<double>> GetNumericTypeInterface() const;
 
@@ -117,6 +126,12 @@ namespace UE::MLDeformer
 
 		/** Filter text used to search the tree. */
 		FText FilterText;
+
+		/** The training animation names as they will appear inside the timeline combobox on the left (in training mode only). */
+		TArray<TSharedPtr<FMLDeformerTrainingInputAnimName>> TrainingAnimNames;
+
+		/** The currently selected animation index. Index into the TrainingAnimNames array.*/
+		int32 SelectedTrainingAnimNameIndex = INDEX_NONE;
 	};
 
 }// namespace UE::MLDeformer
diff --git a/Engine/Plugins/Animation/MLDeformer/MLDeformerFramework/Source/MLDeformerFrameworkEditor/Public/MLDeformerEditorModel.h b/Engine/Plugins/Animation/MLDeformer/MLDeformerFramework/Source/MLDeformerFrameworkEditor/Public/MLDeformerEditorModel.h
--- a/Engine/Plugins/Animation/MLDeformer/MLDeformerFramework/Source/MLDeformerFrameworkEditor/Public/MLDeformerEditorModel.h
+++ b/Engine/Plugins/Animation/MLDeformer/MLDeformerFramework/Source/MLDeformerFrameworkEditor/Public/MLDeformerEditorModel.h
@@ -12,7 +12,7 @@
 #include "MLDeformerVizSettings.h"
 #include "MLDeformerModule.h"
 #include "MLDeformerModel.h"
-
+#include "MLDeformerSampler.h"
 #include "MLDeformerEditorModel.generated.h"
 
 class UMLDeformerModel;
@@ -31,6 +31,7 @@ class UMaterial;
 class UGeometryCache;
 class UMorphTarget;
 class FMorphTargetVertexInfoBuffers;
+struct FMLDeformerTrainingInputAnim;
 
 /** Training process return codes. */
 UENUM()
@@ -114,6 +115,14 @@ namespace UE::MLDeformer
 		 */
 		virtual int32 GetNumTrainingFrames() const;
 
+		/**
+		 * Updates the number of available training frames as returned by GetNumTrainingFrames().
+		 * This should update the NumTrainingFrames class member with the total number of frames that can be included during training.
+		 * So this must be the sum of all frames of all enabled training input animations.
+		 * This should not take into account the maximum number of frames we want to train with.
+		 */
+		virtual void UpdateNumTrainingFrames();
+
 		/**
 		 * Launch the training. This gets executed when the Train button is pressed.
 		 * Training can succeed, or be aborted, and there can be training errors. The result is returned by this method.
@@ -191,13 +200,20 @@ namespace UE::MLDeformer
 		 */
 		virtual FMLDeformerEditorActor* CreateEditorActor(const FMLDeformerEditorActor::FConstructSettings& Settings) const;
 
+		UE_DEPRECATED(5.4, "This method will be removed. Please use the CreateSamplerObject that returns a shared pointer instead.")
+		virtual FMLDeformerSampler* CreateSampler() const { return nullptr; }
+
 		/**
 		 * Create the vertex delta sampler object.
 		 * You can create your own sampler in case you use anything else than say a Geometry Cache as ground truth target data.
 		 * A geometry cache based editor model would create a new FMLDeformerGeomCacheSampler for example.
 		 * @return A pointer to the newly created sampler.
 		 */
-		virtual FMLDeformerSampler* CreateSampler() const;
+		virtual TSharedPtr<FMLDeformerSampler> CreateSamplerObject() const;
+
+		virtual int32 GetNumTrainingInputAnims() const;
+		virtual FMLDeformerTrainingInputAnim* GetTrainingInputAnim(int32 Index) const;
+		virtual void UpdateTimelineTrainingAnimList();
 
 		/**
 		 * Get the time in seconds, for a given frame in the training data.
@@ -608,7 +624,10 @@ namespace UE::MLDeformer
 		virtual void GenerateNormalsForMorphTarget(int32 LOD, USkeletalMesh* SkelMesh, int32 MorphTargetIndex, TArrayView<const FVector3f> Deltas, TArrayView<const FVector3f> BaseVertexPositions, TArrayView<FVector3f> BaseNormals, TArray<FVector3f>& OutDeltaNormals);
 
 		/** Invalidate the memory usage, so it gets updated in the UI again. */
-		void UpdateMemoryUsage();	
+		void UpdateMemoryUsage();
+
+		/** Get the currently active training input animation sequence, which is the one that the viewport is showing the timeline for in training mode. */
+		UAnimSequence* GetActiveTrainingInputAnimSequence() const;
 
 		/** Get the current view range. */
 		TRange<double> GetViewRange() const;
@@ -680,7 +699,10 @@ namespace UE::MLDeformer
 		bool IsReadyForTraining() const { return bIsReadyForTraining; }
 
 		/** Get the sampler we use to calculate vertex deltas when we are in training mdoe and enable deltas. */
-		FMLDeformerSampler* GetSampler() const { return Sampler; }
+		UE_DEPRECATED(5.4, "Please use GetSamplerForTrainingAnim instead.")
+		FMLDeformerSampler* GetSampler() const;
+
+		FMLDeformerSampler* GetSamplerForTrainingAnim(int32 AnimIndex) const	{ return Samplers.IsValidIndex(AnimIndex) ? Samplers[AnimIndex].Get() : nullptr; }
 
 		/**
 		 * Set whether we need to resample inputs or not. This can be used by the training code to determine if we need to resample inputs and outputs, or if
@@ -720,7 +742,7 @@ namespace UE::MLDeformer
 		 * Check for incompatible skeletons between the skeletal mesh and anim sequence skeleton.
 		 * @return An error string if there is an error, otherwise an empty text object is returned.
 		 */
-		FText GetIncompatibleSkeletonErrorText(USkeletalMesh* InSkelMesh, UAnimSequence* InAnimSeq) const;
+		FText GetIncompatibleSkeletonErrorText(const USkeletalMesh* InSkelMesh, const UAnimSequence* InAnimSeq) const;
 
 		/**
 		 * Check to see if the skeletal mesh has to be reimported because it misses some newly added data.
@@ -821,7 +843,6 @@ namespace UE::MLDeformer
 		 */
 		UMLDeformerComponent* FindMLDeformerComponent(int32 ActorID = ActorID_Test_MLDeformed) const;
 
-
 		/**
 		 * Correct floating point errors that can cause issues when sampling animation using step timing.
 		 * @param FrameNumber The desired frame number
@@ -834,7 +855,20 @@ namespace UE::MLDeformer
 		TSharedPtr<SMLDeformerInputWidget> GetInputWidget() const			{ return InputWidget; }
 		void SetInputWidget(TSharedPtr<SMLDeformerInputWidget> Widget)		{ InputWidget = Widget; }
 
+		FMLDeformerSampler* GetSamplerForActiveAnim() const;
+
+		void SetActiveTrainingInputAnimIndex(int32 Index);
+		int32 GetActiveTrainingInputAnimIndex() const;
+
+		/** Mark the deltas to be updated on next Tick. */
+		void InvalidateDeltas();
+
 	protected:
+		void CreateSamplers();
+
+		void AddAnimatedBonesToBonesIncludeList(const UAnimSequence* AnimSequence);
+		void AddAnimatedCurvesToCurvesIncludeList(const UAnimSequence* AnimSequence);
+
 		/**
 		 * Executed when a property changes by a change in the UI.
 		 * This internally calls OnPropertyChanged.
@@ -951,8 +985,8 @@ namespace UE::MLDeformer
 		/** A pointer to the editor toolkit. */
 		FMLDeformerEditorToolkit* Editor = nullptr;
 
-		/** A pointer to the sampler, which can sample target meshes to calculate deltas. */
-		FMLDeformerSampler* Sampler = nullptr;
+		/** A sampler for every input animation. These samplers are used to calculate the vertex deltas. */
+		TArray<TSharedPtr<FMLDeformerSampler>> Samplers;
 
 		/** The inputs widget. */
 		TSharedPtr<SMLDeformerInputWidget> InputWidget;
@@ -1007,6 +1041,15 @@ namespace UE::MLDeformer
 
 		/** Set to true when on next tick we need to trigger an input assets changed event. */
 		bool bNeedsAssetReinit = false;
+
+		/**
+		 * The total number of frames that can be used for training, based on the training input animations list.
+		 * Call UpdateNumTrainingFrames() to refresh this value.
+		 */
+		int32 NumTrainingFrames = 0;
+
+		/** The training input animation that is selected in the timeline. */
+		int32 ActiveTrainingInputAnimIndex = INDEX_NONE;
 	};
 
 	/**
diff --git a/Engine/Plugins/Animation/MLDeformer/MLDeformerFramework/Source/MLDeformerFrameworkEditor/Public/MLDeformerGeomCacheActor.h b/Engine/Plugins/Animation/MLDeformer/MLDeformerFramework/Source/MLDeformerFrameworkEditor/Public/MLDeformerGeomCacheActor.h
--- a/Engine/Plugins/Animation/MLDeformer/MLDeformerFramework/Source/MLDeformerFrameworkEditor/Public/MLDeformerGeomCacheActor.h
+++ b/Engine/Plugins/Animation/MLDeformer/MLDeformerFramework/Source/MLDeformerFrameworkEditor/Public/MLDeformerGeomCacheActor.h
@@ -18,7 +18,6 @@ namespace UE::MLDeformer
 	{
 	public:
 		FMLDeformerGeomCacheActor(const FConstructSettings& Settings);
-		virtual ~FMLDeformerGeomCacheActor() override;
 
 		void SetGeometryCacheComponent(UGeometryCacheComponent* Component)	{ GeomCacheComponent = Component; }
 		UGeometryCacheComponent* GetGeometryCacheComponent() const			{ return GeomCacheComponent; }
diff --git a/Engine/Plugins/Animation/MLDeformer/MLDeformerFramework/Source/MLDeformerFrameworkEditor/Public/MLDeformerGeomCacheEditorModel.h b/Engine/Plugins/Animation/MLDeformer/MLDeformerFramework/Source/MLDeformerFrameworkEditor/Public/MLDeformerGeomCacheEditorModel.h
--- a/Engine/Plugins/Animation/MLDeformer/MLDeformerFramework/Source/MLDeformerFrameworkEditor/Public/MLDeformerGeomCacheEditorModel.h
+++ b/Engine/Plugins/Animation/MLDeformer/MLDeformerFramework/Source/MLDeformerFrameworkEditor/Public/MLDeformerGeomCacheEditorModel.h
@@ -27,11 +27,14 @@ namespace UE::MLDeformer
 		// ~END FGCObject overrides.
 
 		// FMLDeformerEditorModel overrides.
+		virtual void Init(const InitSettings& Settings) override;
 		virtual FMLDeformerEditorActor* CreateEditorActor(const FMLDeformerEditorActor::FConstructSettings& Settings) const override;
-		virtual FMLDeformerSampler* CreateSampler() const override;
+		virtual TSharedPtr<FMLDeformerSampler> CreateSamplerObject() const override;
 		virtual void CreateTrainingGroundTruthActor(UWorld* World) override;
 		virtual void CreateTestGroundTruthActor(UWorld* World) override;
-		virtual int32 GetNumTrainingFrames() const override;
+		virtual int32 GetNumTrainingInputAnims() const override;
+		virtual FMLDeformerTrainingInputAnim* GetTrainingInputAnim(int32 Index) const override;
+		virtual void UpdateNumTrainingFrames() override;
 		virtual double GetTestTimeAtFrame(int32 FrameNumber) const override;
 		virtual int32 GetTestFrameAtTime(double TimeInSeconds) const override;
 		virtual void UpdateIsReadyForTrainingState() override;
@@ -42,10 +45,14 @@ namespace UE::MLDeformer
 		// ~END FMLDeformerEditorModel overrides.
 		
 		// Helpers.
-		FMLDeformerGeomCacheSampler* GetGeomCacheSampler() const;
+		UE_DEPRECATED(5.4, "Please use GetSamplerForTrainingAnim() or GetSamplerForActiveAnim() instead and cast those manually.")
+		FMLDeformerGeomCacheSampler* GetGeomCacheSampler() const { return nullptr; }
+
 		UMLDeformerGeomCacheModel* GetGeomCacheModel() const;
 		UMLDeformerGeomCacheVizSettings* GetGeomCacheVizSettings() const;
 		FMLDeformerGeomCacheActor* FindGeomCacheEditorActor(int32 ID) const;
+		UGeometryCache* GetActiveGeometryCache() const;
+
 
 	protected:
 		/**
diff --git a/Engine/Plugins/Animation/MLDeformer/MLDeformerFramework/Source/MLDeformerFrameworkEditor/Public/MLDeformerGeomCacheModelDetails.h b/Engine/Plugins/Animation/MLDeformer/MLDeformerFramework/Source/MLDeformerFrameworkEditor/Public/MLDeformerGeomCacheModelDetails.h
--- a/Engine/Plugins/Animation/MLDeformer/MLDeformerFramework/Source/MLDeformerFrameworkEditor/Public/MLDeformerGeomCacheModelDetails.h
+++ b/Engine/Plugins/Animation/MLDeformer/MLDeformerFramework/Source/MLDeformerFrameworkEditor/Public/MLDeformerGeomCacheModelDetails.h
@@ -28,13 +28,13 @@ namespace UE::MLDeformer
 	public:
 		// FMLDeformerModelDetails overrides.
 		virtual bool UpdateMemberPointers(const TArray<TWeakObjectPtr<UObject>>& Objects) override;
-		virtual void AddTargetMesh() override;
-		virtual void AddAnimSequenceErrors() override;
+		virtual void AddTrainingInputAnims() override;
 		// ~END FMLDeformerModelDetails overrides.
 
 	protected:
 		/** Add warnings to the UI about mesh mapping errors when mapping between geometry caches and a skeletal mesh. */
-		void AddGeomCacheMeshMappingWarnings(IDetailCategoryBuilder* InTargetMeshCategoryBuilder, USkeletalMesh* SkeletalMesh, UGeometryCache* GeometryCache);
+		UE_DEPRECATED(5.4, "This method will be removed.")
+		void AddGeomCacheMeshMappingWarnings(IDetailCategoryBuilder* InTargetMeshCategoryBuilder, USkeletalMesh* SkeletalMesh, UGeometryCache* GeometryCache) {}
 
 	protected:
 		/** The geometry cache based runtime model. This value gets updated once UpdateMemberPointers gets called. */
diff --git a/Engine/Plugins/Animation/MLDeformer/MLDeformerFramework/Source/MLDeformerFrameworkEditor/Public/MLDeformerGeomCacheSampler.h b/Engine/Plugins/Animation/MLDeformer/MLDeformerFramework/Source/MLDeformerFrameworkEditor/Public/MLDeformerGeomCacheSampler.h
--- a/Engine/Plugins/Animation/MLDeformer/MLDeformerFramework/Source/MLDeformerFrameworkEditor/Public/MLDeformerGeomCacheSampler.h
+++ b/Engine/Plugins/Animation/MLDeformer/MLDeformerFramework/Source/MLDeformerFrameworkEditor/Public/MLDeformerGeomCacheSampler.h
@@ -25,8 +25,8 @@ namespace UE::MLDeformer
 	{
 	public:
 		// FVertexDeltaSampler overrides.
+		virtual void Init(FMLDeformerEditorModel* InModel, int32 InAnimIndex) override;
 		virtual void Sample(int32 AnimFrameIndex) override;
-		virtual void RegisterTargetComponents() override;
 		virtual float GetTimeAtFrame(int32 InAnimFrameIndex) const override;
 		// ~END FVertexDeltaSampler overrides.
 
@@ -43,6 +43,7 @@ namespace UE::MLDeformer
 		 * Executed when we need a pointer to the geometry cache object that we want to sample from.
 		 * @return A reference to the delegate.
 		 */
+		UE_DEPRECATED(5.4, "This method will be removed. To get the geom cache use FMLDeformerGeomCacheEditorModel::GetTrainingInputAnims() together with the AnimIndex as index value.")
 		FMLDeformerGetGeomCacheEvent& OnGetGeometryCache()			{ return GetGeometryCacheEvent; }
 
 		UGeometryCacheComponent* GetGeometryCacheComponent()		{ return GeometryCacheComponent; }
@@ -62,7 +63,8 @@ namespace UE::MLDeformer
 		 * @param OutVertexDeltas The array that will receive the deltas, in a float buffer. The buffer will contain 3*NumVerts number of elements, and the layout is
 		 *                        (x, y, z, x, y, z, x, y, z...) so 3 values per vertex: x, y and z.
 		 */
-		void CalculateVertexDeltas(const TArray<FVector3f>& SkinnedPositions, float DeltaCutoffLength, TArray<float>& OutVertexDeltas);
+		UE_DEPRECATED(5.4, "This method will be removed.")
+		void CalculateVertexDeltas(const TArray<FVector3f>& SkinnedPositions, float DeltaCutoffLength, TArray<float>& OutVertexDeltas) {}
 
 	protected:
 		/** The geometry cache component used to sample the geometry cache. */
diff --git a/Engine/Plugins/Animation/MLDeformer/MLDeformerFramework/Source/MLDeformerFrameworkEditor/Public/MLDeformerGeomCacheTrainingInputAnimCustomize.h b/Engine/Plugins/Animation/MLDeformer/MLDeformerFramework/Source/MLDeformerFrameworkEditor/Public/MLDeformerGeomCacheTrainingInputAnimCustomize.h
new file mode 100644
--- /dev/null
+++ b/Engine/Plugins/Animation/MLDeformer/MLDeformerFramework/Source/MLDeformerFrameworkEditor/Public/MLDeformerGeomCacheTrainingInputAnimCustomize.h
@@ -0,0 +1,37 @@
+// Copyright Epic Games, Inc. All Rights Reserved.
+
+#pragma once
+
+#include "CoreTypes.h"
+#include "IPropertyTypeCustomization.h"
+#include "PropertyHandle.h"
+#include "DetailWidgetRow.h"
+#include "AssetRegistry/AssetData.h"
+
+class USkeleton;
+class UMLDeformerGeomCacheModel;
+
+namespace UE::MLDeformer
+{
+	/**
+	 * The ML Deformer's training input animation (that uses a geometry cache) type customization.
+	 * This is used to show a custom UI when displaying the FMLDeformerGeomCacheTrainingInputAnim struct.
+	 * That struct is shown in the list of training input animations in the details panel.
+	 * @see FMLDeformerGeomCacheTrainingInputAnim
+	 */
+	class MLDEFORMERFRAMEWORKEDITOR_API FMLDeformerGeomCacheTrainingInputAnimCustomization
+		: public IPropertyTypeCustomization
+	{
+	public:
+		static TSharedRef<IPropertyTypeCustomization> MakeInstance();
+
+		// IPropertyTypeCustomization overrides.
+		virtual void CustomizeHeader(TSharedRef<IPropertyHandle> StructPropertyHandle, class FDetailWidgetRow& HeaderRow, IPropertyTypeCustomizationUtils& StructCustomizationUtils) override;
+		virtual void CustomizeChildren(TSharedRef<IPropertyHandle> StructPropertyHandle, class IDetailChildrenBuilder& ChildBuilder, IPropertyTypeCustomizationUtils& StructCustomizationUtils) override;
+		// ~END IPropertyTypeCustomization overrides.
+
+	private:
+		bool FilterAnimSequences(const FAssetData& AssetData, USkeleton* Skeleton);
+		UMLDeformerGeomCacheModel* FindMLDeformerModel(TSharedRef<IPropertyHandle> StructPropertyHandle) const;
+	};
+}	// namespace UE::MLDeformer
diff --git a/Engine/Plugins/Animation/MLDeformer/MLDeformerFramework/Source/MLDeformerFrameworkEditor/Public/MLDeformerGeomCacheTrainingModel.h b/Engine/Plugins/Animation/MLDeformer/MLDeformerFramework/Source/MLDeformerFrameworkEditor/Public/MLDeformerGeomCacheTrainingModel.h
new file mode 100644
--- /dev/null
+++ b/Engine/Plugins/Animation/MLDeformer/MLDeformerFramework/Source/MLDeformerFrameworkEditor/Public/MLDeformerGeomCacheTrainingModel.h
@@ -0,0 +1,35 @@
+// Copyright Epic Games, Inc. All Rights Reserved.
+
+#pragma once
+
+#include "CoreTypes.h"
+#include "Containers/Array.h"
+#include "MLDeformerTrainingModel.h"
+#include "MLDeformerGeomCacheTrainingModel.generated.h"
+
+/**
+ * The training model base class for geometry cache based models.
+ * This class is our link to the Python training.
+ */
+UCLASS(Blueprintable)
+class MLDEFORMERFRAMEWORKEDITOR_API UMLDeformerGeomCacheTrainingModel
+	: public UMLDeformerTrainingModel
+{
+	GENERATED_BODY()
+
+public:
+	// UMLDeformerTrainingModel overrides.
+	virtual void Init(UE::MLDeformer::FMLDeformerEditorModel* InEditorModel) override;
+	virtual bool SampleNextFrame() override;
+	// ~END UMLDeformerTrainingModel overrides;
+
+protected:
+	/**
+	 * Find the next input animation to sample from.
+	 * This is an index inside the training input animations list.
+	 * This assumes the SampleAnimIndex member as starting point. The method does not modify this member directly, unless passed in as parameter.
+	 * @param OutNextAnimIndex The next animation index to sample from when we take our next sample.
+	 * @return Returns true when we found our next animation to sample. Returns false when we already sampled everything.
+	 */
+	virtual bool FindNextAnimToSample(int32& OutNextAnimIndex) const override;
+};
\ No newline at end of file
diff --git a/Engine/Plugins/Animation/MLDeformer/MLDeformerFramework/Source/MLDeformerFrameworkEditor/Public/MLDeformerModelDetails.h b/Engine/Plugins/Animation/MLDeformer/MLDeformerFramework/Source/MLDeformerFrameworkEditor/Public/MLDeformerModelDetails.h
--- a/Engine/Plugins/Animation/MLDeformer/MLDeformerFramework/Source/MLDeformerFrameworkEditor/Public/MLDeformerModelDetails.h
+++ b/Engine/Plugins/Animation/MLDeformer/MLDeformerFramework/Source/MLDeformerFrameworkEditor/Public/MLDeformerModelDetails.h
@@ -45,24 +45,34 @@ namespace UE::MLDeformer
 
 		/**
 		 * Add the target mesh property. For example the FMLDeformerGeomCacheModelDetails class will add the geometry cache property here. */
+		UE_DEPRECATED(5.4, "This method will be removed. Please add your errors in your MLDeformerTrainingInputAnim type customizations instead.")
 		virtual void AddTargetMesh() {}
 
 		/** Add additional errors related to the base mesh. */
 		virtual void AddBaseMeshErrors() {}
 
 		/** Add additional errors related to the training animation sequence. */
+		UE_DEPRECATED(5.4, "This method will be removed. Please add your errors in your MLDeformerTrainingInputAnim type customizations instead.")
 		virtual void AddAnimSequenceErrors() {}
 
+		/** Show errors about skeletons being incompatible. */
+		UE_DEPRECATED(5.4, "This method will be removed. Please add your errors in your MLDeformerTrainingInputAnim type customizations instead.")
+		virtual void AddIncompatibleSkeletonErrors() {}
+
 		/** Add additional errors related to bone inputs. */
 		virtual void AddBoneInputErrors() {}
 
 		/** Add additional errors related to curve inputs. */
 		virtual void AddCurveInputErrors() {}
 
+		/** Add the property that shows the input animations. */
+		virtual void AddTrainingInputAnims();
+
 		/** Add training input flags, which basically are things like the check boxes that specify whether bones or curves (or both) should be included. */
 		virtual void AddTrainingInputFlags() {}
 
 		/** Add additional training input filters. Filters are things like a list of bones, or curves. */
+		UE_DEPRECATED(5.4, "This method will be removed.")
 		virtual void AddTrainingInputFilters() {}
 
 		/** Add additional training input errors. */
@@ -72,20 +82,24 @@ namespace UE::MLDeformer
 		virtual void AddTrainingSettingsErrors() {}
 
 		/** Is the "include bones" checkbox visible? */
-		virtual bool IsBonesFlagVisible() const;
+		UE_DEPRECATED(5.4, "This method will be removed.")
+		virtual bool IsBonesFlagVisible() const { return true; }
 
 		/** Is the "include curves" checkbox visible? */
-		virtual bool IsCurvesFlagVisible() const;
+		UE_DEPRECATED(5.4, "This method will be removed.")
+		virtual bool IsCurvesFlagVisible() const { return true; }
 
 	protected:
 		/** The filter that only shows anim sequences that are compatible with the given skeleton. */
 		bool FilterAnimSequences(const FAssetData& AssetData, USkeleton* Skeleton);
 
 		/** Executed when we press the Animated Bones Only button. */
-		FReply OnFilterAnimatedBonesOnly() const;
+		UE_DEPRECATED(5.4, "This method will be removed.")
+		FReply OnFilterAnimatedBonesOnly() const { return FReply::Handled(); }
 
 		/** Executed when we press the Animated Curves Only button. */
-		FReply OnFilterAnimatedCurvesOnly() const;
+		UE_DEPRECATED(5.4, "This method will be removed.")
+		FReply OnFilterAnimatedCurvesOnly() const { return FReply::Handled(); }
 
 	protected:
 		/** Associated detail layout builder. */
@@ -101,6 +115,7 @@ namespace UE::MLDeformer
 		IDetailCategoryBuilder* BaseMeshCategoryBuilder = nullptr;
 
 		/** The category related to the target mesh. */
+		UE_DEPRECATED(5.4, "This category will be removed. You most likely want to use the InputOutputCategory.")
 		IDetailCategoryBuilder* TargetMeshCategoryBuilder = nullptr;
 
 		/** The category related to the inputs and outputs. */
diff --git a/Engine/Plugins/Animation/MLDeformer/MLDeformerFramework/Source/MLDeformerFrameworkEditor/Public/MLDeformerSampler.h b/Engine/Plugins/Animation/MLDeformer/MLDeformerFramework/Source/MLDeformerFrameworkEditor/Public/MLDeformerSampler.h
--- a/Engine/Plugins/Animation/MLDeformer/MLDeformerFramework/Source/MLDeformerFrameworkEditor/Public/MLDeformerSampler.h
+++ b/Engine/Plugins/Animation/MLDeformer/MLDeformerFramework/Source/MLDeformerFrameworkEditor/Public/MLDeformerSampler.h
@@ -37,8 +37,11 @@ namespace UE::MLDeformer
 		virtual ~FMLDeformerSampler();
 
 		/** Call Init whenever assets or other relevant settings change. */
+		UE_DEPRECATED(5.4, "Please switch to using Init(Model, AnimIndex) instead.")
 		virtual void Init(FMLDeformerEditorModel* Model);
 
+		virtual void Init(FMLDeformerEditorModel* Model, int32 AnimIndex);
+
 		/** Call this every time the frame changes. This will update all buffer contents. */
 		virtual void Sample(int32 AnimFrameIndex);
 
@@ -98,6 +101,7 @@ namespace UE::MLDeformer
 		virtual void CreateActors();
 
 		/** Register the target mesh components. */
+		UE_DEPRECATED(5.4, "Please create your target component inside the overloaded Init method instead.")
 		virtual void RegisterTargetComponents() {}
 
 		/** Extract the skinned vertex positions, of the skeletal mesh. */
@@ -177,6 +181,9 @@ namespace UE::MLDeformer
 		/** The animation frame we sampled the deltas for. */
 		int32 AnimFrameIndex = -1;
 
+		/** The animation index for which this is a sampler for. This is in range of 0..GetNumTrainingInputAnims(). */
+		int32 AnimIndex = -1;
+
 		/**
 		 * The number of floats per curve value. If this is larger than 1, the remaining floats (after the first one) will be set to 0. 
 		 * So basically we will pad the curve values.
diff --git a/Engine/Plugins/Animation/MLDeformer/MLDeformerFramework/Source/MLDeformerFrameworkEditor/Public/MLDeformerTrainingModel.h b/Engine/Plugins/Animation/MLDeformer/MLDeformerFramework/Source/MLDeformerFrameworkEditor/Public/MLDeformerTrainingModel.h
--- a/Engine/Plugins/Animation/MLDeformer/MLDeformerFramework/Source/MLDeformerFrameworkEditor/Public/MLDeformerTrainingModel.h
+++ b/Engine/Plugins/Animation/MLDeformer/MLDeformerFramework/Source/MLDeformerFrameworkEditor/Public/MLDeformerTrainingModel.h
@@ -75,15 +75,27 @@ public:
 	UFUNCTION(BlueprintPure, Category = "Training Data")
     int32 NumSamples() const;
 
+	/** This will make the sampling start again from the beginning. This can be used if you have to iterate multiple times over the data set. */
+	UFUNCTION(BlueprintCallable, Category = "Training Data")
+    void ResetSampling();
+
 	/** 
 	 * Set the current sample frame. This will internally call the SampleFrame method, which will update the deltas, curve values and bone rotations. 
 	 * You call this before getting the input bone/curve and vertex delta values.
 	 * @param Index The training data frame/sample number.
 	 * @return Returns true when successful, or false when the specified sample index is out of range.
 	 */
-	UFUNCTION(BlueprintCallable, Category = "Training Data")
+	UE_DEPRECATED(5.4, "Please use NextSample instead.")
+	UFUNCTION(meta = (DeprecatedFunction, DeprecationMessage = "Please use NextSample instead."))
 	bool SetCurrentSampleIndex(int32 Index);
 
+	/**
+	 * Take the next sample.
+	 * This will return false when there is something wrong or we sampled more times than NumSamples() returns.
+	 */
+	UFUNCTION(BlueprintCallable, Category = "Training Data")
+	bool NextSample();
+
 	/** 
 	 * Check whether we need to resample the inputs and outputs, or if we can use a cached version. 
 	 * This will return true if any inputs changed, that indicate that we should regenerate any cached data.
@@ -102,14 +114,27 @@ public:
 
 protected:
 	/** This updates the sample deltas, curves, and bone rotations. */
+	UE_DEPRECATED(5.4, "Please use UMLDeformerTrainingModel::SampleNextFrame() instead.")
     virtual bool SampleFrame(int32 Index);
 
+	/** Sample the next frame in the sequence of samples. */
+	virtual bool SampleNextFrame();
+
 	/** Change the pointer to the editor model. */
 	void SetEditorModel(UE::MLDeformer::FMLDeformerEditorModel* InModel);
 
 	/** Get a pointer to the editor model. */ 
 	UE::MLDeformer::FMLDeformerEditorModel* GetEditorModel() const;
 
+	/**
+	 * Find the next input animation to sample from.
+	 * This is an index inside the training input animations list.
+	 * This assumes the SampleAnimIndex member as starting point. The method does not modify this member directly, unless passed in as parameter.
+	 * @param OutNextAnimIndex The next animation index to sample from when we take our next sample.
+	 * @return Returns true when we found our next animation to sample. Returns false when we already sampled everything.
+	 */
+	virtual bool FindNextAnimToSample(int32& OutNextAnimIndex) const { return false; }
+
 public:
 	// The delta values per vertex for this sample. This is updated after SetCurrentSampleIndex is called. Contains an xyz (3 floats) for each vertex.
 	UPROPERTY(EditAnywhere, BlueprintReadWrite, Category = "Training Data")
@@ -126,4 +151,17 @@ public:
 protected:
 	/** A pointer to the editor model from which we use the sampler. */
 	UE::MLDeformer::FMLDeformerEditorModel* EditorModel = nullptr;
+
+	/** 
+	 * The number of times a given input animation has been sampled. 
+	 * The size of the array equals the number of training input animations.
+	 * Note that this also contains counts of disabled training input anims. These disabled ones should be ignored.
+	 */
+	TArray<int32> NumTimesSampled;
+
+	/** The training input animation to take the next sample from. */
+	int32 SampleAnimIndex = 0;
+
+	/** Did we finish sampling? This is set to true when every possible frame has been sampled. */
+	bool bFinishedSampling = false;
 };
diff --git a/Engine/Plugins/Animation/MLDeformer/NearestNeighborModel/Source/NearestNeighborModelEditor/Private/NearestNeighborEditorModel.cpp b/Engine/Plugins/Animation/MLDeformer/NearestNeighborModel/Source/NearestNeighborModelEditor/Private/NearestNeighborEditorModel.cpp
--- a/Engine/Plugins/Animation/MLDeformer/NearestNeighborModel/Source/NearestNeighborModelEditor/Private/NearestNeighborEditorModel.cpp
+++ b/Engine/Plugins/Animation/MLDeformer/NearestNeighborModel/Source/NearestNeighborModelEditor/Private/NearestNeighborEditorModel.cpp
@@ -39,15 +39,13 @@ namespace UE::NearestNeighborModel
 
 	void FNearestNeighborEditorModel::Init(const InitSettings& InitSettings)
 	{
-		FMLDeformerEditorModel::Init(InitSettings);
+		FMLDeformerGeomCacheEditorModel::Init(InitSettings);
 		InitInputInfo(Model->GetInputInfo());
 	}
 
-	FMLDeformerSampler* FNearestNeighborEditorModel::CreateSampler() const
+	TSharedPtr<FMLDeformerSampler> FNearestNeighborEditorModel::CreateSamplerObject() const
 	{
-		FNearestNeighborGeomCacheSampler* NewSampler = new FNearestNeighborGeomCacheSampler();
-		NewSampler->OnGetGeometryCache().BindLambda([this] { return GetGeomCacheModel()->GetGeometryCache(); });
-		return NewSampler;
+		return MakeShared<FNearestNeighborGeomCacheSampler>();
 	}
 
 	void FNearestNeighborEditorModel::OnPropertyChanged(FPropertyChangedEvent& PropertyChangedEvent)
@@ -61,10 +59,10 @@ namespace UE::NearestNeighborModel
 		FMLDeformerMorphModelEditorModel::OnPropertyChanged(PropertyChangedEvent);
 
 		if (Property->GetFName() == UMLDeformerModel::GetSkeletalMeshPropertyName() ||
-			Property->GetFName() == UMLDeformerModel::GetAnimSequencePropertyName() ||
-			Property->GetFName() == UMLDeformerGeomCacheModel::GetGeometryCachePropertyName() ||
-			Property->GetFName() == UNearestNeighborModel::GetClothPartEditorDataPropertyName()
-			|| (PropertyChangedEvent.MemberProperty != nullptr && PropertyChangedEvent.MemberProperty->GetFName() == UNearestNeighborModel::GetClothPartEditorDataPropertyName()))
+			Property->GetFName() == UNearestNeighborModel::GetClothPartEditorDataPropertyName() ||
+		    Property->GetFName() == UMLDeformerGeomCacheModel::GetTrainingInputAnimsPropertyName() ||
+		    PropertyChangedEvent.GetMemberPropertyName() == TEXT("TrainingInputAnims") ||
+			PropertyChangedEvent.GetMemberPropertyName() == UNearestNeighborModel::GetClothPartEditorDataPropertyName())
 		{
 			GetNearestNeighborModel()->InvalidateClothPartData();
 			GetEditor()->GetModelDetailsView()->ForceRefresh();
@@ -257,7 +255,8 @@ namespace UE::NearestNeighborModel
 
 	uint8 FNearestNeighborEditorModel::SetSamplerPartData(const int32 PartId)
 	{
-		FNearestNeighborGeomCacheSampler* GeomCacheSampler = static_cast<FNearestNeighborGeomCacheSampler*>(GetGeomCacheSampler());
+		FMLDeformerSampler* Sampler = GetNumTrainingInputAnims() > 0 ? GetSamplerForTrainingAnim(0) : nullptr;
+		FNearestNeighborGeomCacheSampler* GeomCacheSampler = static_cast<FNearestNeighborGeomCacheSampler*>(Sampler);
 		UNearestNeighborModel* NearestNeighborModel = GetNearestNeighborModel();
 		if (GeomCacheSampler && NearestNeighborModel && PartId < NearestNeighborModel->GetNumParts())
 		{
@@ -346,12 +345,13 @@ namespace UE::NearestNeighborModel
 
 	void FNearestNeighborEditorModel::ResetSamplerData()
 	{
-		FNearestNeighborGeomCacheSampler* GeomCacheSampler = static_cast<FNearestNeighborGeomCacheSampler*>(GetGeomCacheSampler());
+		FMLDeformerSampler* Sampler = GetNumTrainingInputAnims() > 0 ? GetSamplerForTrainingAnim(0) : nullptr;
+		FNearestNeighborGeomCacheSampler* GeomCacheSampler = static_cast<FNearestNeighborGeomCacheSampler*>(Sampler);
 		UNearestNeighborModel* NearestNeighborModel = GetNearestNeighborModel();
 		if (GeomCacheSampler && NearestNeighborModel)
 		{
 			TObjectPtr<USkeletalMeshComponent> SkeletalMeshComponent = GeomCacheSampler->GetSkeletalMeshComponent();
-			const TObjectPtr<UAnimSequence> AnimSequence = NearestNeighborModel->GetAnimSequence();
+			const TObjectPtr<UAnimSequence> AnimSequence = GetTrainingInputAnim(0)->GetAnimSequence();
 
 			if (SkeletalMeshComponent && AnimSequence)
 			{
@@ -363,7 +363,6 @@ namespace UE::NearestNeighborModel
 				SkeletalMeshComponent->RefreshBoneTransforms();
 			}
 			NumTrainingFramesOverride = -1;
-			GeomCacheSampler->RegisterTargetComponents();
 		}
 	}
 
@@ -407,7 +406,8 @@ namespace UE::NearestNeighborModel
 			return EUpdateResult::SUCCESS;
 		}
 
-		FNearestNeighborGeomCacheSampler* GeomCacheSampler = static_cast<FNearestNeighborGeomCacheSampler*>(GetGeomCacheSampler());
+		FMLDeformerSampler* Sampler = GetNumTrainingInputAnims() > 0 ? GetSamplerForTrainingAnim(0) : nullptr;
+		FNearestNeighborGeomCacheSampler* GeomCacheSampler = static_cast<FNearestNeighborGeomCacheSampler*>(Sampler);
 		if (GeomCacheSampler)
 		{
 			TObjectPtr<USkeletalMeshComponent> SkeletalMeshComponent = GeomCacheSampler->GetSkeletalMeshComponent();
diff --git a/Engine/Plugins/Animation/MLDeformer/NearestNeighborModel/Source/NearestNeighborModelEditor/Private/NearestNeighborTrainingModel.cpp b/Engine/Plugins/Animation/MLDeformer/NearestNeighborModel/Source/NearestNeighborModelEditor/Private/NearestNeighborTrainingModel.cpp
--- a/Engine/Plugins/Animation/MLDeformer/NearestNeighborModel/Source/NearestNeighborModelEditor/Private/NearestNeighborTrainingModel.cpp
+++ b/Engine/Plugins/Animation/MLDeformer/NearestNeighborModel/Source/NearestNeighborModelEditor/Private/NearestNeighborTrainingModel.cpp
@@ -46,13 +46,25 @@ int32 UNearestNeighborTrainingModel::GetPartNumNeighbors(const int32 PartId) con
 
 bool UNearestNeighborTrainingModel::SampleKmeansAnim(const int32 SkeletonId)
 {
-	FNearestNeighborGeomCacheSampler* Sampler = static_cast<FNearestNeighborGeomCacheSampler*>(EditorModel->GetSampler());
+	FMLDeformerSampler* BaseSampler = EditorModel->GetNumTrainingInputAnims() > 0 ? EditorModel->GetSamplerForTrainingAnim(0) : nullptr;
+	FNearestNeighborGeomCacheSampler* Sampler = static_cast<FNearestNeighborGeomCacheSampler*>(BaseSampler);
+	if (Sampler == nullptr)
+	{
+		return false;
+	}
+
 	return Sampler->SampleKMeansAnim(SkeletonId);
 }
 
 bool UNearestNeighborTrainingModel::SampleKmeansFrame(const int32 Frame)
 {
-	FNearestNeighborGeomCacheSampler* Sampler = static_cast<FNearestNeighborGeomCacheSampler*>(EditorModel->GetSampler());
+	FMLDeformerSampler* BaseSampler = EditorModel->GetNumTrainingInputAnims() > 0 ? EditorModel->GetSamplerForTrainingAnim(0) : nullptr;
+	FNearestNeighborGeomCacheSampler* Sampler = static_cast<FNearestNeighborGeomCacheSampler*>(BaseSampler);
+	if (Sampler == nullptr)
+	{
+		return false;
+	}
+
 	const bool bSampleExist = Sampler->SampleKMeansFrame(Frame);
 	if (bSampleExist)
 	{
@@ -84,7 +96,14 @@ int32 UNearestNeighborTrainingModel::GetKmeansNumClusters() const
 
 const TArray<float> UNearestNeighborTrainingModel::GetUnskinnedVertexPositions() const
 {
-	const TArray<FVector3f>& PositionsVec = EditorModel->GetSampler()->GetUnskinnedVertexPositions();
+	FMLDeformerSampler* BaseSampler = EditorModel->GetNumTrainingInputAnims() > 0 ? EditorModel->GetSamplerForTrainingAnim(0) : nullptr;
+	FNearestNeighborGeomCacheSampler* Sampler = static_cast<FNearestNeighborGeomCacheSampler*>(BaseSampler);
+	if (Sampler == nullptr)
+	{
+		return TArray<float>();
+	}
+
+	const TArray<FVector3f>& PositionsVec = Sampler->GetUnskinnedVertexPositions();
 	TArray<float> PositionsFloat; 
 	PositionsFloat.SetNumUninitialized(PositionsVec.Num() * 3);
 	for (int32 i = 0; i < PositionsVec.Num(); i++)
@@ -98,7 +117,13 @@ const TArray<float> UNearestNeighborTrainingModel::GetUnskinnedVertexPositions()
 
 const TArray<int32> UNearestNeighborTrainingModel::GetMeshIndexBuffer() const
 {
-	FNearestNeighborGeomCacheSampler* Sampler = static_cast<FNearestNeighborGeomCacheSampler*>(EditorModel->GetSampler());
+	FMLDeformerSampler* BaseSampler = EditorModel->GetNumTrainingInputAnims() > 0 ? EditorModel->GetSamplerForTrainingAnim(0) : nullptr;
+	FNearestNeighborGeomCacheSampler* Sampler = static_cast<FNearestNeighborGeomCacheSampler*>(BaseSampler);
+	if (Sampler == nullptr)
+	{
+		return TArray<int32>();
+	}
+
 	TArray<uint32> UIndexBuffer = Sampler->GetMeshIndexBuffer();
 	return TArray<int32>((int32*)UIndexBuffer.GetData(), UIndexBuffer.Num());
 }
diff --git a/Engine/Plugins/Animation/MLDeformer/NearestNeighborModel/Source/NearestNeighborModelEditor/Public/NearestNeighborEditorModel.h b/Engine/Plugins/Animation/MLDeformer/NearestNeighborModel/Source/NearestNeighborModelEditor/Public/NearestNeighborEditorModel.h
--- a/Engine/Plugins/Animation/MLDeformer/NearestNeighborModel/Source/NearestNeighborModelEditor/Public/NearestNeighborEditorModel.h
+++ b/Engine/Plugins/Animation/MLDeformer/NearestNeighborModel/Source/NearestNeighborModelEditor/Public/NearestNeighborEditorModel.h
@@ -37,7 +37,7 @@ namespace UE::NearestNeighborModel
 	
 		// FMLDeformerGeomCacheEditorModel overrides.
 		virtual void Init(const InitSettings& Settings) override;
-		virtual FMLDeformerSampler* CreateSampler() const override;
+		virtual TSharedPtr<FMLDeformerSampler> CreateSamplerObject() const override;
 		virtual void CreateActors(const TSharedRef<IPersonaPreviewScene>& InPersonaPreviewScene) override;
 		virtual void Tick(FEditorViewportClient* ViewportClient, float DeltaTime) override;
 		virtual FMLDeformerEditorActor* CreateEditorActor(const FMLDeformerEditorActor::FConstructSettings& Settings) const override;
diff --git a/Engine/Plugins/Animation/MLDeformer/NearestNeighborModel/Source/NearestNeighborModelEditor/Public/NearestNeighborTrainingModel.h b/Engine/Plugins/Animation/MLDeformer/NearestNeighborModel/Source/NearestNeighborModelEditor/Public/NearestNeighborTrainingModel.h
--- a/Engine/Plugins/Animation/MLDeformer/NearestNeighborModel/Source/NearestNeighborModelEditor/Public/NearestNeighborTrainingModel.h
+++ b/Engine/Plugins/Animation/MLDeformer/NearestNeighborModel/Source/NearestNeighborModelEditor/Public/NearestNeighborTrainingModel.h
@@ -3,7 +3,7 @@
 #pragma once
 
 #include "CoreTypes.h"
-#include "MLDeformerTrainingModel.h"
+#include "MLDeformerGeomCacheTrainingModel.h"
 #include "NearestNeighborTrainingModel.generated.h"
 
 namespace UE::NearestNeighborModel
@@ -16,7 +16,7 @@ class UNearestNeighborModelInstance;
 
 UCLASS(Blueprintable)
 class NEARESTNEIGHBORMODELEDITOR_API UNearestNeighborTrainingModel
-	: public UMLDeformerTrainingModel
+	: public UMLDeformerGeomCacheTrainingModel
 {
 	GENERATED_BODY()
 
diff --git a/Engine/Plugins/Animation/MLDeformer/NeuralMorphModel/Source/NeuralMorphModelEditor/Private/NeuralMorphEditorModel.cpp b/Engine/Plugins/Animation/MLDeformer/NeuralMorphModel/Source/NeuralMorphModelEditor/Private/NeuralMorphEditorModel.cpp
--- a/Engine/Plugins/Animation/MLDeformer/NeuralMorphModel/Source/NeuralMorphModelEditor/Private/NeuralMorphEditorModel.cpp
+++ b/Engine/Plugins/Animation/MLDeformer/NeuralMorphModel/Source/NeuralMorphModelEditor/Private/NeuralMorphEditorModel.cpp
@@ -672,12 +672,15 @@ namespace UE::NeuralMorphModel
 
 	void FNeuralMorphEditorModel::DebugDrawItemMask(FPrimitiveDrawInterface* PDI, int32 MaskItemIndex, const FVector& DrawOffset)
 	{
+		FMLDeformerSampler* Sampler = GetSamplerForActiveAnim();
+
 		UMLDeformerMorphModelVizSettings* VizSettings = Cast<UMLDeformerMorphModelVizSettings>(Model->GetVizSettings());
 		const int32 NumVerts = Model->GetNumBaseMeshVerts();
 		const TArray<FVector3f>& UnskinnedPositions = Sampler->GetUnskinnedVertexPositions();
 		if (MaskVizItemIndex == INDEX_NONE ||
 			NumVerts != Model->GetInputInfo()->GetNumBaseMeshVertices() ||
-			UnskinnedPositions.Num() != NumVerts)
+			UnskinnedPositions.Num() != NumVerts ||
+			Sampler == nullptr)
 		{
 			return;
 		}
diff --git a/Engine/Plugins/Animation/MLDeformer/NeuralMorphModel/Source/NeuralMorphModelEditor/Public/NeuralMorphTrainingModel.h b/Engine/Plugins/Animation/MLDeformer/NeuralMorphModel/Source/NeuralMorphModelEditor/Public/NeuralMorphTrainingModel.h
--- a/Engine/Plugins/Animation/MLDeformer/NeuralMorphModel/Source/NeuralMorphModelEditor/Public/NeuralMorphTrainingModel.h
+++ b/Engine/Plugins/Animation/MLDeformer/NeuralMorphModel/Source/NeuralMorphModelEditor/Public/NeuralMorphTrainingModel.h
@@ -4,7 +4,7 @@
 
 #include "CoreTypes.h"
 #include "Containers/Array.h"
-#include "MLDeformerTrainingModel.h"
+#include "MLDeformerGeomCacheTrainingModel.h"
 #include "NeuralMorphTrainingModel.generated.h"
 
 /**
@@ -13,7 +13,7 @@
  */
 UCLASS(Blueprintable)
 class NEURALMORPHMODELEDITOR_API UNeuralMorphTrainingModel
-	: public UMLDeformerTrainingModel
+	: public UMLDeformerGeomCacheTrainingModel
 {
 	GENERATED_BODY()
 
