[MLDeformer] Fix some compile issues.

#rnx
#trivial

diff --git a/Engine/Plugins/Animation/MLDeformer/MLDeformerFramework/Source/MLDeformerFrameworkEditor/Public/MLDeformerModelDetails.h b/Engine/Plugins/Animation/MLDeformer/MLDeformerFramework/Source/MLDeformerFrameworkEditor/Public/MLDeformerModelDetails.h
--- a/Engine/Plugins/Animation/MLDeformer/MLDeformerFramework/Source/MLDeformerFrameworkEditor/Public/MLDeformerModelDetails.h
+++ b/Engine/Plugins/Animation/MLDeformer/MLDeformerFramework/Source/MLDeformerFrameworkEditor/Public/MLDeformerModelDetails.h
@@ -115,8 +115,7 @@ namespace UE::MLDeformer
 		IDetailCategoryBuilder* BaseMeshCategoryBuilder = nullptr;
 
 		/** The category related to the target mesh. */
-		UE_DEPRECATED(5.4, "This category will be removed. You most likely want to use the InputOutputCategory.")
-		IDetailCategoryBuilder* TargetMeshCategoryBuilder = nullptr;
+		IDetailCategoryBuilder* TargetMeshCategoryBuilder_DEPRECATED = nullptr;
 
 		/** The category related to the inputs and outputs. */
 		IDetailCategoryBuilder* InputOutputCategoryBuilder = nullptr;
diff --git a/Engine/Plugins/Animation/MLDeformer/MLDeformerFramework/Source/MLDeformerFrameworkEditor/Public/MLDeformerVizSettingsDetails.h b/Engine/Plugins/Animation/MLDeformer/MLDeformerFramework/Source/MLDeformerFrameworkEditor/Public/MLDeformerVizSettingsDetails.h
--- a/Engine/Plugins/Animation/MLDeformer/MLDeformerFramework/Source/MLDeformerFrameworkEditor/Public/MLDeformerVizSettingsDetails.h
+++ b/Engine/Plugins/Animation/MLDeformer/MLDeformerFramework/Source/MLDeformerFrameworkEditor/Public/MLDeformerVizSettingsDetails.h
@@ -121,8 +121,7 @@ namespace UE::MLDeformer
 		IDetailGroup* StatsAssetSizeGroup = nullptr;
 
 		/** Main memory usage subgroup. */
-		UE_DEPRECATED(5.4, "This member will be removed.")
-		IDetailGroup* StatsMainMemUsageGroup = nullptr;
+		IDetailGroup* StatsMainMemUsageGroup_DEPRECATED = nullptr;
 
 		/** GPU memory usage subgroup. */
 		IDetailGroup* StatsGPUMemUsageGroup = nullptr;
