[MLDeformer]
- Add support for output morph weight clamping in the Neural Morph Model. This is an extra safety feature to prevent potential extreme deformations when running the deformer on unseen poses.
- Fixed a ground truth frame alignment issue when scrubbing the timeline in testing mode.
- Fixed some small errors  in the Nearest Neighbor model where it called some functions on the incorrect base class.

diff --git a/Engine/Plugins/Animation/MLDeformer/MLDeformerFramework/Source/MLDeformerFramework/Private/MLDeformerMorphModel.cpp b/Engine/Plugins/Animation/MLDeformer/MLDeformerFramework/Source/MLDeformerFramework/Private/MLDeformerMorphModel.cpp
--- a/Engine/Plugins/Animation/MLDeformer/MLDeformerFramework/Source/MLDeformerFramework/Private/MLDeformerMorphModel.cpp
+++ b/Engine/Plugins/Animation/MLDeformer/MLDeformerFramework/Source/MLDeformerFramework/Private/MLDeformerMorphModel.cpp
@@ -130,14 +130,37 @@ void UMLDeformerMorphModel::UpdateStatistics()
 	UncompressedMorphDataSizeInBytes = GetMorphTargetDeltas().Num() * GetMorphTargetDeltas().GetTypeSize();
 }
 
-void UMLDeformerMorphModel::SetMorphTargetsMaxWeights(const TArray<float>& MaxWeights)
+void UMLDeformerMorphModel::SetMorphTargetsMinMaxWeights(const TArray<FFloatInterval>& MinMaxValues)
 {
-	MaxMorphWeights = MaxWeights;
+	MorphTargetsMinMaxWeights = MinMaxValues;
 }
 
-TArrayView<const float> UMLDeformerMorphModel::GetMorphTargetMaxWeights() const
+void UMLDeformerMorphModel::SetMorphTargetsMinMaxWeights(const TArray<float>& MinValues, const TArray<float>& MaxValues)
 {
-	return MaxMorphWeights;
+	check(MinValues.Num() == MaxValues.Num());
+	const int32 NumWeights = MinValues.Num();
+
+	MorphTargetsMinMaxWeights.Reset();
+	MorphTargetsMinMaxWeights.AddUninitialized(NumWeights);
+	for (int32 Index = 0; Index < MinValues.Num(); ++Index)
+	{
+		MorphTargetsMinMaxWeights[Index].Min = MinValues[Index];
+		MorphTargetsMinMaxWeights[Index].Max = MaxValues[Index];
+	}
+}
+
+void UMLDeformerMorphModel::ClampMorphTargetWeights(TArrayView<float> WeightsArray)
+{
+	if (MorphTargetsMinMaxWeights.Num() != WeightsArray.Num())
+	{
+		return;
+	}
+
+	for (int32 MorphIndex = 0; MorphIndex < WeightsArray.Num(); ++MorphIndex)
+	{
+		const FFloatInterval& MorphMinMax = MorphTargetsMinMaxWeights[MorphIndex];
+		WeightsArray[MorphIndex] = FMath::Clamp(WeightsArray[MorphIndex], MorphMinMax.Min, MorphMinMax.Max);;
+	}
 }
 
 TArrayView<const float> UMLDeformerMorphModel::GetMorphTargetErrorValues() const
diff --git a/Engine/Plugins/Animation/MLDeformer/MLDeformerFramework/Source/MLDeformerFramework/Public/MLDeformerMorphModel.h b/Engine/Plugins/Animation/MLDeformer/MLDeformerFramework/Source/MLDeformerFramework/Public/MLDeformerMorphModel.h
--- a/Engine/Plugins/Animation/MLDeformer/MLDeformerFramework/Source/MLDeformerFramework/Public/MLDeformerMorphModel.h
+++ b/Engine/Plugins/Animation/MLDeformer/MLDeformerFramework/Source/MLDeformerFramework/Public/MLDeformerMorphModel.h
@@ -59,6 +59,7 @@ public:
 	void SetIncludeMorphTargetNormals(bool bInclude)				{ bIncludeNormals = bInclude; }
 	void SetMaskChannel(EMLDeformerMaskChannel Channel)				{ MaskChannel = Channel; }
 	void SetInvertMaskChannel(bool bInvert)							{ bInvertMaskChannel = bInvert; }
+	void SetClampMorphTargetsWeights(bool bEnabled)					{ bClampMorphWeights = bEnabled; }
 
 	UE_DEPRECATED(5.3, "Use SetMaskChannel instead.")
 	void SetWeightMask(EMLDeformerMaskChannel Channel)				{ MaskChannel = Channel; }
@@ -85,6 +86,7 @@ public:
 	static FName GetCompressedMorphDataSizeInBytesPropertyName()	{ return GET_MEMBER_NAME_CHECKED(UMLDeformerMorphModel, CompressedMorphDataSizeInBytes); }
 	static FName GetUncompressedMorphDataSizeInBytesPropertyName()	{ return GET_MEMBER_NAME_CHECKED(UMLDeformerMorphModel, UncompressedMorphDataSizeInBytes); }
 	static FName GetQualityLevelsPropertyName()						{ return GET_MEMBER_NAME_CHECKED(UMLDeformerMorphModel, QualityLevels); }
+	static FName GetClampMorphTargetWeightsPropertyName()			{ return GET_MEMBER_NAME_CHECKED(UMLDeformerMorphModel, bClampMorphWeights); }
 
 	UE_DEPRECATED(5.2, "Please use GetMorphDeltaZeroThresholdPropertyName instead.")
 	static FName GetMorphTargetDeltaThresholdPropertyName()			{ return GET_MEMBER_NAME_CHECKED(UMLDeformerMorphModel, MorphDeltaZeroThreshold); }
@@ -137,15 +139,61 @@ public:
 	 * Training python scripts mostly will call this function to set the values.
 	 * @param MaxWeights The maximum of the absolute weight values.
 	 */
+	UE_DEPRECATED(5.4, "This method will be removed.")
 	UFUNCTION(BlueprintCallable, Category = "MLDeformerMorphModel")
-	void SetMorphTargetsMaxWeights(const TArray<float>& MaxWeights);
+	void SetMorphTargetsMaxWeights(const TArray<float>& MaxWeights) {}
+
+	/** 
+	 * Set the minimum and maximum values that the morph targets weights have seen during training.
+	 * We can clamp the network output weights within these ranges to try to make sure that the output doesn't go wild
+	 * when we give inputs that are far outside of the range we have seen during training.
+	 * The size of the array must equal the number of morph targets, or empty. In case it is empty, clamping will be ignored, even when enabled.
+	 * @param MinValues The minimum weigth values, one for each morph target.
+	 * @param MaxValues The maximum weight values, one for each morph target.
+	 */
+	UFUNCTION(BlueprintCallable, Category = "MLDeformerMorphModel")
+	void SetMorphTargetsMinMaxWeights(const TArray<float>& MinValues, const TArray<float>& MaxValues);
+
+	/**
+	 * Set the minimum and maximum values that the morph targets weights have seen during training.
+	 * We can clamp the network output weights within these ranges to try to make sure that the output doesn't go wild
+	 * when we give inputs that are far outside of the range we have seen during training.
+	 * The size of the array must equal the number of morph targets, or empty. In case it is empty, clamping will be ignored, even when enabled.
+	 * @param MinMaxValues The array of float intervals that contain the min and max values, one for each morph target.
+	 */
+	void SetMorphTargetsMinMaxWeights(const TArray<FFloatInterval>& MinMaxValues);
 
 	/**
 	 * Get the estimated maximum weight values that the morph targets will ever have. It is not guaranteed that the weights will never be larger than these values though.
 	 * These values are used to estimate the importance levels of the morph targets.
 	 * @see SetMorphTargetsMaxWeights.
 	 */
-	TArrayView<const float> GetMorphTargetMaxWeights() const;
+	UE_DEPRECATED(5.4, "This method will be removed, use GetMorphTargetsMinMaxWeights instead, and take the max of the absolute value of both min and max.")
+	TArrayView<const float> GetMorphTargetMaxWeights() const { return TArrayView<const float>(); }
+
+	/**
+	 * Get the min and max morph target weight values that we seen on the output of the training data set.
+	 * These min and max values can be used to clamp the morph targets within this range.
+	 * This can be used to prevent weights from getting very different values on unseen input data, and thus can prevent visually 'exploding' correctives.
+	 */
+	const TArray<FFloatInterval>& GetMorphTargetsMinMaxWeights() const	{ return MorphTargetsMinMaxWeights; }
+
+	/**
+	 * Check whether clamping of morph target weights is enabled.
+	 * If enabled, the weights that are output will be clamped to the minimum and maximum values that are defined in the values returned by
+	 * GetMorphTargetsMinMaxWeights().
+	 * It is up to the model to actually call the ClampMorphTargetWeights method though.
+	 * You can check if a model supports weight clamping by calling FMLDeformerMorphModelEditorModel::IsMorphWeightClampingSupported().
+	 * If that would return false, even enabling clamping using SetClampMorphTargetsWeights(true) will not do anything.
+	 * @return Returns true when clamping is enabled, otherwise false is returned.
+	 */
+	bool IsMorphWeightClampingEnabled() const							{ return bClampMorphWeights; }
+
+	/**
+	 * Perform the actual weight clamping to the values as set inside the MorphTargetsMinMaxWeights member.
+	 * @param WeightsArray The array of input weights, that will be clamped. This method will modify the array values.
+	 */
+	void ClampMorphTargetWeights(TArrayView<float> WeightsArray);
 
 	/**
 	 * Get the morph target delta vectors array.
@@ -225,6 +273,16 @@ private:
 	/** The compressed morph target data, ready for the GPU. */
 	TSharedPtr<FExternalMorphSet> MorphTargetSet;
 
+	/** 
+	 * Should we enable morph target weight clamping?
+	 * The minimum and maximum values that it will be clamped against will be the min/max morph target weight values
+	 * that have been seen when running the training dataset through the network.
+	 * The advantage of clamping is that it can make deformations more stable when we have input poses that have not been seen during training.
+	 * We basically prevent the weights from 'exploding' and getting very large values, which could make the mesh look very bad.
+	 */
+	UPROPERTY(EditAnywhere, Category = "Morph Targets")
+	bool bClampMorphWeights = true;
+
 	/**
 	 * The entire set of uncompressed morph target deltas, 3 per vertex, for each morph target, as one flattened buffer.
 	 * So the size of this buffer is: (NumVertsPerMorphTarget * 3 * NumMorphTargets).
@@ -257,7 +315,17 @@ private:
 
 	/** The maximum absolute weight values of the morph targets, during training. One value for each morph target. */
 	UPROPERTY()
-	TArray<float> MaxMorphWeights;
+	TArray<float> MaxMorphWeights_DEPRECATED;
+
+	/** 
+	 * The minimum and maximum weight values that the morph targets weights have seen during training.
+	 * We can clamp the network output weights within these ranges to try to make sure that the output doesn't go wild
+	 * when we give inputs that are far outside of the range we have seen during training.
+	 * It is possible that this array is empty, in which case this clamping isn't supported.
+	 * The size of the array equals the number of morph targets.
+	 */
+	UPROPERTY()
+	TArray<FFloatInterval> MorphTargetsMinMaxWeights;
 
 	/** 
 	 * The list of quality levels, where the first item represents the highest quality and the last element the lowest quality level.
diff --git a/Engine/Plugins/Animation/MLDeformer/MLDeformerFramework/Source/MLDeformerFrameworkEditor/Private/MLDeformerEditorModel.cpp b/Engine/Plugins/Animation/MLDeformer/MLDeformerFramework/Source/MLDeformerFrameworkEditor/Private/MLDeformerEditorModel.cpp
--- a/Engine/Plugins/Animation/MLDeformer/MLDeformerFramework/Source/MLDeformerFrameworkEditor/Private/MLDeformerEditorModel.cpp
+++ b/Engine/Plugins/Animation/MLDeformer/MLDeformerFramework/Source/MLDeformerFrameworkEditor/Private/MLDeformerEditorModel.cpp
@@ -229,6 +229,10 @@ namespace UE::MLDeformer
 		SkelMeshComponent->SetVisibility(false);
 		SkelMeshComponent->MarkRenderStateDirty();
 		Actor->SetRootComponent(SkelMeshComponent);
+		if (SkelMeshComponent->GetAnimInstance())
+		{
+			SkelMeshComponent->GetAnimInstance()->GetRequiredBones().SetUseRAWData(true);
+		}
 
 		// Register the editor actor.
 		const FLinearColor LabelColor = FMLDeformerEditorStyle::Get().GetColor("MLDeformer.BaseMesh.LabelColor");
@@ -261,6 +265,10 @@ namespace UE::MLDeformer
 		SkelMeshComponent->SetWireframeMeshOverlayColor(MLDeformedWireColor);
 		SkelMeshComponent->SetVisibility(false);
 		SkelMeshComponent->MarkRenderStateDirty();
+		if (SkelMeshComponent->GetAnimInstance())
+		{
+			SkelMeshComponent->GetAnimInstance()->GetRequiredBones().SetUseRAWData(true);
+		}
 
 		// Create the ML Deformer component.
 		UMLDeformerAsset* DeformerAsset = Model->GetDeformerAsset();
@@ -633,11 +641,11 @@ namespace UE::MLDeformer
 		if (VizSettings->GetVisualizationMode() == EMLDeformerVizMode::TrainingData)
 		{
 			const int32 TargetFrame = GetTrainingFrameAtTime(NewScrubTime);
+			PlayOffset = GetTrainingTimeAtFrame(TargetFrame);
 			for (FMLDeformerEditorActor* EditorActor : EditorActors)
 			{
 				if (EditorActor && EditorActor->IsTrainingActor())
 				{
-					PlayOffset = GetTrainingTimeAtFrame(TargetFrame);
 					EditorActor->SetPlayPosition(PlayOffset);
 				}
 			}
@@ -646,14 +654,11 @@ namespace UE::MLDeformer
 		else if (VizSettings->GetVisualizationMode() == EMLDeformerVizMode::TestData)
 		{
 			const int32 TargetFrame = GetTestFrameAtTime(NewScrubTime);
+			PlayOffset = GetTestTimeAtFrame(TargetFrame);
 			for (FMLDeformerEditorActor* EditorActor : EditorActors)
 			{
 				if (EditorActor && EditorActor->IsTestActor())
 				{
-					if (Model->GetVizSettings()->HasTestGroundTruth())
-					{
-						PlayOffset = GetTestTimeAtFrame(TargetFrame);
-					}
 					EditorActor->SetPlayPosition(PlayOffset);
 				}
 			}
@@ -697,7 +702,15 @@ namespace UE::MLDeformer
 
 	double FMLDeformerEditorModel::GetTestTimeAtFrame(int32 FrameNumber) const
 	{
-		return Model->GetVizSettings()->GetTestAnimSequence() ? Model->GetVizSettings()->GetTestAnimSequence()->GetTimeAtFrame(FrameNumber) : 0.0;
+		UAnimSequence* AnimSequence = Model->GetVizSettings()->GetTestAnimSequence();
+		if (AnimSequence)
+		{
+			const FFrameRate FrameRate = AnimSequence->GetSamplingFrameRate();
+			const float UncorrectedTime = AnimSequence->GetTimeAtFrame(FrameNumber);
+			// due to floating point errors, return a double that when converted to a float cleanly maps to the value
+			return CorrectedFrameTime(FrameNumber, UncorrectedTime, FrameRate);
+		}
+		return 0.0; 
 	}
 
 	int32 FMLDeformerEditorModel::GetTestFrameAtTime(double TimeInSeconds) const
diff --git a/Engine/Plugins/Animation/MLDeformer/MLDeformerFramework/Source/MLDeformerFrameworkEditor/Private/MLDeformerGeomCacheEditorModel.cpp b/Engine/Plugins/Animation/MLDeformer/MLDeformerFramework/Source/MLDeformerFrameworkEditor/Private/MLDeformerGeomCacheEditorModel.cpp
--- a/Engine/Plugins/Animation/MLDeformer/MLDeformerFramework/Source/MLDeformerFrameworkEditor/Private/MLDeformerGeomCacheEditorModel.cpp
+++ b/Engine/Plugins/Animation/MLDeformer/MLDeformerFramework/Source/MLDeformerFrameworkEditor/Private/MLDeformerGeomCacheEditorModel.cpp
@@ -60,29 +60,6 @@ namespace UE::MLDeformer
 		return static_cast<FMLDeformerGeomCacheActor*>(FindEditorActor(ID));
 	}
 
-	double FMLDeformerGeomCacheEditorModel::GetTestTimeAtFrame(int32 FrameNumber) const
-	{
-		// Try to get the frame from the geometry cache.
-		const FMLDeformerGeomCacheActor* EditorActor = static_cast<FMLDeformerGeomCacheActor*>(FindEditorActor(ActorID_Test_GroundTruth));
-		if (EditorActor && EditorActor->GetGeometryCacheComponent() && EditorActor->GetGeometryCacheComponent()->GeometryCache.Get())
-		{
-			return EditorActor->GetGeometryCacheComponent()->GetTimeAtFrame(FrameNumber);
-		}
-
-		return FMLDeformerEditorModel::GetTestTimeAtFrame(FrameNumber);
-	}
-
-	int32 FMLDeformerGeomCacheEditorModel::GetTestFrameAtTime(double TimeInSeconds) const
-	{
-		const FMLDeformerGeomCacheActor* EditorActor = static_cast<FMLDeformerGeomCacheActor*>(FindEditorActor(ActorID_Test_GroundTruth));
-		if (EditorActor && EditorActor->GetGeometryCacheComponent() && EditorActor->GetGeometryCacheComponent()->GeometryCache.Get())
-		{
-			return EditorActor->GetGeometryCacheComponent()->GetFrameAtTime(TimeInSeconds);
-		}
-
-		return FMLDeformerEditorModel::GetTestFrameAtTime(TimeInSeconds);
-	}
-
 	UGeometryCache* FMLDeformerGeomCacheEditorModel::GetActiveGeometryCache() const
 	{
 		const int32 ActiveAnimIndex = GetActiveTrainingInputAnimIndex();
diff --git a/Engine/Plugins/Animation/MLDeformer/MLDeformerFramework/Source/MLDeformerFrameworkEditor/Private/MLDeformerMorphModelDetails.cpp b/Engine/Plugins/Animation/MLDeformer/MLDeformerFramework/Source/MLDeformerFrameworkEditor/Private/MLDeformerMorphModelDetails.cpp
--- a/Engine/Plugins/Animation/MLDeformer/MLDeformerFramework/Source/MLDeformerFrameworkEditor/Private/MLDeformerMorphModelDetails.cpp
+++ b/Engine/Plugins/Animation/MLDeformer/MLDeformerFramework/Source/MLDeformerFrameworkEditor/Private/MLDeformerMorphModelDetails.cpp
@@ -39,6 +39,9 @@ namespace UE::MLDeformer
 		// Create all the detail categories and add the properties of the base class.
 		FMLDeformerGeomCacheModelDetails::CustomizeDetails(DetailBuilder);
 
+		MorphTargetCategoryBuilder->AddProperty(DetailLayoutBuilder->GetProperty(UMLDeformerMorphModel::GetClampMorphTargetWeightsPropertyName(), UMLDeformerMorphModel::StaticClass()))
+			.Visibility(MorphModelEditorModel->IsMorphWeightClampingSupported() ? EVisibility::Visible : EVisibility::Collapsed);
+
 		MorphTargetCategoryBuilder->AddProperty(DetailLayoutBuilder->GetProperty(UMLDeformerMorphModel::GetIncludeMorphTargetNormalsPropertyName(), UMLDeformerMorphModel::StaticClass()));
 
 		IDetailGroup& CompressionGroup = MorphTargetCategoryBuilder->AddGroup("Compression", LOCTEXT("MorphCompressionGroupLabel", "Compression"), false, true);
diff --git a/Engine/Plugins/Animation/MLDeformer/MLDeformerFramework/Source/MLDeformerFrameworkEditor/Private/MLDeformerMorphModelEditorModel.cpp b/Engine/Plugins/Animation/MLDeformer/MLDeformerFramework/Source/MLDeformerFrameworkEditor/Private/MLDeformerMorphModelEditorModel.cpp
--- a/Engine/Plugins/Animation/MLDeformer/MLDeformerFramework/Source/MLDeformerFrameworkEditor/Private/MLDeformerMorphModelEditorModel.cpp
+++ b/Engine/Plugins/Animation/MLDeformer/MLDeformerFramework/Source/MLDeformerFrameworkEditor/Private/MLDeformerMorphModelEditorModel.cpp
@@ -113,6 +113,7 @@ namespace UE::MLDeformer
 	{
 		// Backup the morph target deltas in case we abort training.
 		MorphTargetDeltasBackup = GetMorphModel()->GetMorphTargetDeltas();
+		MorphTargetsMinMaxWeightsBackup = GetMorphModel()->GetMorphTargetsMinMaxWeights();
 	}
 
 	void FMLDeformerMorphModelEditorModel::OnPostTraining(ETrainingResult TrainingResult, bool bUsePartiallyTrainedWhenAborted)
@@ -122,6 +123,7 @@ namespace UE::MLDeformer
 		{
 			// Restore the morph target vertex deltas backup.
 			GetMorphModel()->SetMorphTargetDeltas(MorphTargetDeltasBackup);
+			GetMorphModel()->SetMorphTargetsMinMaxWeights(MorphTargetsMinMaxWeightsBackup);
 		}
 		else if (TrainingResult == ETrainingResult::Success || (TrainingResult == ETrainingResult::Aborted && bUsePartiallyTrainedWhenAborted))
 		{
@@ -171,7 +173,7 @@ namespace UE::MLDeformer
 		// Check if we have max morph weight information.
 		// If we do not have this yet, we have to initialize the weights to 1.
 		UMLDeformerMorphModel* MorphModel = GetMorphModel();
-		TArrayView<const float> MaxMorphWeights = MorphModel->GetMorphTargetMaxWeights();
+		const TArray<FFloatInterval>& MinMaxMorphWeights = MorphModel->GetMorphTargetsMinMaxWeights();
 
 		// Preallocate space for the standard deviation of each morph target.
 		TArray<float> ErrorValues;
@@ -182,7 +184,10 @@ namespace UE::MLDeformer
 		for (int32 MorphIndex = 0; MorphIndex < MorphTargets.Num() - 1; ++MorphIndex)	// We have one extra morph for the means, skip that one.
 		{
 			const UMorphTarget* MorphTarget = MorphTargets[MorphIndex + 1];
-			const float MaxWeight = !MaxMorphWeights.IsEmpty() ? MaxMorphWeights[MorphIndex] : 1.0f;
+
+			// Calculate the maximum of the absolute values of the min and max weight we saw during training.
+			// We will multiply this with the length of the deltas later on to get an estimate of the maximum deformation for all deltas.
+			const float MaxWeight = !MinMaxMorphWeights.IsEmpty() ? FMath::Max(FMath::Abs(MinMaxMorphWeights[MorphIndex].Min), FMath::Abs(MinMaxMorphWeights[MorphIndex].Min)) : 1.0f;
 
 			// Get the array of deltas.
 			int32 NumDeltas = 0;
diff --git a/Engine/Plugins/Animation/MLDeformer/MLDeformerFramework/Source/MLDeformerFrameworkEditor/Public/MLDeformerGeomCacheEditorModel.h b/Engine/Plugins/Animation/MLDeformer/MLDeformerFramework/Source/MLDeformerFrameworkEditor/Public/MLDeformerGeomCacheEditorModel.h
--- a/Engine/Plugins/Animation/MLDeformer/MLDeformerFramework/Source/MLDeformerFrameworkEditor/Public/MLDeformerGeomCacheEditorModel.h
+++ b/Engine/Plugins/Animation/MLDeformer/MLDeformerFramework/Source/MLDeformerFrameworkEditor/Public/MLDeformerGeomCacheEditorModel.h
@@ -35,8 +35,6 @@ namespace UE::MLDeformer
 		virtual int32 GetNumTrainingInputAnims() const override;
 		virtual FMLDeformerTrainingInputAnim* GetTrainingInputAnim(int32 Index) const override;
 		virtual void UpdateNumTrainingFrames() override;
-		virtual double GetTestTimeAtFrame(int32 FrameNumber) const override;
-		virtual int32 GetTestFrameAtTime(double TimeInSeconds) const override;
 		virtual void UpdateIsReadyForTrainingState() override;
 		virtual void OnPropertyChanged(FPropertyChangedEvent& PropertyChangedEvent) override;
 		virtual void OnInputAssetsChanged() override;
diff --git a/Engine/Plugins/Animation/MLDeformer/MLDeformerFramework/Source/MLDeformerFrameworkEditor/Public/MLDeformerMorphModelEditorModel.h b/Engine/Plugins/Animation/MLDeformer/MLDeformerFramework/Source/MLDeformerFrameworkEditor/Public/MLDeformerMorphModelEditorModel.h
--- a/Engine/Plugins/Animation/MLDeformer/MLDeformerFramework/Source/MLDeformerFrameworkEditor/Public/MLDeformerMorphModelEditorModel.h
+++ b/Engine/Plugins/Animation/MLDeformer/MLDeformerFramework/Source/MLDeformerFrameworkEditor/Public/MLDeformerMorphModelEditorModel.h
@@ -35,6 +35,15 @@ namespace UE::MLDeformer
 		virtual void OnPostTraining(ETrainingResult TrainingResult, bool bUsePartiallyTrainedWhenAborted) override;
 		// ~END FMLDeformerEditorModel overrides.
 
+		/**
+		 * Returns whether this model supports morph output weight clamping or not.
+		 * This is used to make sure the morph target weights are within valid ranges.
+		 * Some models might choose not to do this, because they might already clamp input values.
+		 * When this returns false, the "Clamp morph weights" option in the UI will also be hidden.
+		 * @return Returns true when the model supports morph weight clamping, otherwise false is returned.
+		 */
+		virtual bool IsMorphWeightClampingSupported() const		{ return true; }
+
 		/**
 		 * Get the mask buffer for a given morph target.
 		 * @param MorphTargetIndex The morph target index, excluding the first 'means' morph target that we always have, so 0 would mean the actual first real morph target.
@@ -280,6 +289,9 @@ namespace UE::MLDeformer
 		 */
 		TArray<FVector3f> MorphTargetDeltasBackup;
 
+		/** The backup of the minimum and maximum weights for each morph target. */
+		TArray<FFloatInterval> MorphTargetsMinMaxWeightsBackup;
+
 		/** The backup of the input item mask buffer. This contains all the bone and bone group masks. */
 		TArray<float> InputItemMaskBufferBackup;
 	};
diff --git a/Engine/Plugins/Animation/MLDeformer/NearestNeighborModel/Source/NearestNeighborModelEditor/Private/NearestNeighborModelDetails.cpp b/Engine/Plugins/Animation/MLDeformer/NearestNeighborModel/Source/NearestNeighborModelEditor/Private/NearestNeighborModelDetails.cpp
--- a/Engine/Plugins/Animation/MLDeformer/NearestNeighborModel/Source/NearestNeighborModelEditor/Private/NearestNeighborModelDetails.cpp
+++ b/Engine/Plugins/Animation/MLDeformer/NearestNeighborModel/Source/NearestNeighborModelEditor/Private/NearestNeighborModelDetails.cpp
@@ -3,6 +3,7 @@
 #include "NearestNeighborModelDetails.h"
 #include "NearestNeighborEditorModel.h"
 #include "NearestNeighborModel.h"
+#include "MLDeformerMorphModelDetails.h"
 #include "MLDeformerEditorToolkit.h"
 #include "DetailLayoutBuilder.h"
 #include "DetailCategoryBuilder.h"
@@ -34,7 +35,7 @@ namespace UE::NearestNeighborModel
 
 	bool FNearestNeighborModelDetails::UpdateMemberPointers(const TArray<TWeakObjectPtr<UObject>>& Objects)
 	{
-		if (!FMLDeformerGeomCacheModelDetails::UpdateMemberPointers(Objects))
+		if (!FMLDeformerMorphModelDetails::UpdateMemberPointers(Objects))
 		{
 			return false;
 		}
@@ -48,12 +49,11 @@ namespace UE::NearestNeighborModel
 
 	void FNearestNeighborModelDetails::CreateCategories()
 	{
-		FMLDeformerGeomCacheModelDetails::CreateCategories();
+		FMLDeformerMorphModelDetails::CreateCategories();
 
 		FileCacheCategoryBuilder = &DetailLayoutBuilder->EditCategory("File Cache", FText::GetEmpty(), ECategoryPriority::Important);
 		ClothPartCategoryBuilder = &DetailLayoutBuilder->EditCategory("Cloth Parts", FText::GetEmpty(), ECategoryPriority::Important);
 		NearestNeighborCategoryBuilder = &DetailLayoutBuilder->EditCategory("Nearest Neighbors", FText::GetEmpty(), ECategoryPriority::Important);
-		MorphTargetCategoryBuilder = &DetailLayoutBuilder->EditCategory("Morph Targets", FText::GetEmpty(), ECategoryPriority::Important);
 		KMeansCategoryBuilder = &DetailLayoutBuilder->EditCategory("KMeans Pose Generator", FText::GetEmpty(), ECategoryPriority::Important);
 
 		// Add warning in CreateCategories so that the warning appears at the top of the details panel.
@@ -174,8 +174,13 @@ namespace UE::NearestNeighborModel
 	void FNearestNeighborModelDetails::CustomizeDetails(IDetailLayoutBuilder& DetailBuilder)
 	{
 		// Create all the detail categories and add the properties of the base class.
-		FMLDeformerGeomCacheModelDetails::CustomizeDetails(DetailBuilder);
+		FMLDeformerMorphModelDetails::CustomizeDetails(DetailBuilder);
 
+		if (NearestNeighborModel == nullptr)
+		{
+			return;
+		}
+		
 		// Training settings.
 		TrainingSettingsCategoryBuilder->AddProperty(UNearestNeighborModel::GetInputDimPropertyName());
 		TrainingSettingsCategoryBuilder->AddProperty(UNearestNeighborModel::GetHiddenLayerDimsPropertyName());
@@ -192,11 +197,6 @@ namespace UE::NearestNeighborModel
 		Group->AddPropertyRow(DetailBuilder.GetProperty(UNearestNeighborModel::GetRecomputeDeltasPropertyName()));
 		Group->AddPropertyRow(DetailBuilder.GetProperty(UNearestNeighborModel::GetRecomputePCAPropertyName()));
 
-		if (NearestNeighborModel == nullptr)
-		{
-			return;
-		}
-
 		BuildSubMeshNames();
 		const int32 MaxPartMeshIndex = NearestNeighborModel->GetMaxPartMeshIndex();
 		if (MaxPartMeshIndex >= SubMeshNames.Num())
diff --git a/Engine/Plugins/Animation/MLDeformer/NearestNeighborModel/Source/NearestNeighborModelEditor/Public/NearestNeighborEditorModel.h b/Engine/Plugins/Animation/MLDeformer/NearestNeighborModel/Source/NearestNeighborModelEditor/Public/NearestNeighborEditorModel.h
--- a/Engine/Plugins/Animation/MLDeformer/NearestNeighborModel/Source/NearestNeighborModelEditor/Public/NearestNeighborEditorModel.h
+++ b/Engine/Plugins/Animation/MLDeformer/NearestNeighborModel/Source/NearestNeighborModelEditor/Public/NearestNeighborEditorModel.h
@@ -55,6 +55,11 @@ namespace UE::NearestNeighborModel
 		virtual int32 GetNumTrainingFrames() const override;
 		// ~END FMLDeformerGeomCacheEditorModel overrides.
 
+		// UMLDeformerMorphModelEditorModel overrides.
+		virtual bool IsMorphWeightClampingSupported() const override	{ return false; }	// We already do input clamping, so output clamping really isn't needed.
+		// ~END UMLDeformerMorphModelEditorModel overrides.
+
+
 		friend class FNearestNeighborModelSampler;
 		friend class FNearestNeighborModelDetails;
 		friend class ::UNearestNeighborTrainingModel;
diff --git a/Engine/Plugins/Animation/MLDeformer/NeuralMorphModel/Source/NeuralMorphModel/Private/NeuralMorphModelInstance.cpp b/Engine/Plugins/Animation/MLDeformer/NeuralMorphModel/Source/NeuralMorphModel/Private/NeuralMorphModelInstance.cpp
--- a/Engine/Plugins/Animation/MLDeformer/NeuralMorphModel/Source/NeuralMorphModel/Private/NeuralMorphModelInstance.cpp
+++ b/Engine/Plugins/Animation/MLDeformer/NeuralMorphModel/Source/NeuralMorphModel/Private/NeuralMorphModelInstance.cpp
@@ -291,4 +291,15 @@ void UNeuralMorphModelInstance::Execute(float ModelWeight)
 			WeightData->Weights[MorphIndex + 1] = NetworkOutputs[MorphIndex] * ModelWeight;
 		}
 	}
+
+	// Clamp morph target weights to be within the bounds we seen on the training data set.
+	// This can prevent 'exploding' weight values when we have network inputs that are far away from inputs we seen during training.
+	if (MorphModel->IsMorphWeightClampingEnabled())
+	{
+		// Start at weights[1], because the first item contains the morph weight that represents the mean.
+		// And we don't want to clamp that one as it cannot go out of bounds, as it is basically always 1.0 if the deformer is fully active.
+		// The ClampMorphTargetWeights method will modify the values in the array.
+		TArrayView<float> WeightsArray(&WeightData->Weights[1], NumNetworkWeights);
+		MorphModel->ClampMorphTargetWeights(WeightsArray);
+	}
 }
