[MLDeformer]
- Added an optional comparison model in the viewport. This can be used to compare your model to another pre-trained model. The comparison MLD asset requires to be trained on the same skeletal mesh for it to work.
- Added support for detailed searches for MLD assets inside the content browser.

diff --git a/Engine/Plugins/Animation/MLDeformer/MLDeformerFramework/Source/MLDeformerFramework/Private/MLDeformerAsset.cpp b/Engine/Plugins/Animation/MLDeformer/MLDeformerFramework/Source/MLDeformerFramework/Private/MLDeformerAsset.cpp
--- a/Engine/Plugins/Animation/MLDeformer/MLDeformerFramework/Source/MLDeformerFramework/Private/MLDeformerAsset.cpp
+++ b/Engine/Plugins/Animation/MLDeformer/MLDeformerFramework/Source/MLDeformerFramework/Private/MLDeformerAsset.cpp
@@ -2,6 +2,9 @@
 
 #include "MLDeformerAsset.h"
 #include "MLDeformerObjectVersion.h"
+#include "MLDeformerModel.h"
+#include "MLDeformerInputInfo.h"
+#include "Engine/SkeletalMesh.h"
 
 void UMLDeformerAsset::Serialize(FArchive& Archive)
 {
@@ -10,3 +13,12 @@ void UMLDeformerAsset::Serialize(FArchive& Archive)
 	Archive.UsingCustomVersion(UE::MLDeformer::FMLDeformerObjectVersion::GUID);
 	Super::Serialize(Archive);
 }
+
+void UMLDeformerAsset::GetAssetRegistryTags(TArray<FAssetRegistryTag>& OutTags) const
+{
+	Super::GetAssetRegistryTags(OutTags);
+	if (Model)
+	{
+		Model->GetAssetRegistryTags(OutTags);
+	}
+}
diff --git a/Engine/Plugins/Animation/MLDeformer/MLDeformerFramework/Source/MLDeformerFramework/Private/MLDeformerGeomCacheModel.cpp b/Engine/Plugins/Animation/MLDeformer/MLDeformerFramework/Source/MLDeformerFramework/Private/MLDeformerGeomCacheModel.cpp
--- a/Engine/Plugins/Animation/MLDeformer/MLDeformerFramework/Source/MLDeformerFramework/Private/MLDeformerGeomCacheModel.cpp
+++ b/Engine/Plugins/Animation/MLDeformer/MLDeformerFramework/Source/MLDeformerFramework/Private/MLDeformerGeomCacheModel.cpp
@@ -53,6 +53,28 @@ void UMLDeformerGeomCacheModel::PostLoad()
 	#endif
 }
 
+void UMLDeformerGeomCacheModel::GetAssetRegistryTags(TArray<FAssetRegistryTag>& OutTags) const
+{
+	Super::GetAssetRegistryTags(OutTags);
+
+	#if WITH_EDITORONLY_DATA
+		FString AnimInputString;
+		for (const FMLDeformerGeomCacheTrainingInputAnim& Anim : TrainingInputAnims)
+		{
+			AnimInputString += Anim.GetAnimSequenceSoftObjectPtr().ToSoftObjectPath().ToString();
+			AnimInputString += TEXT("\n");
+			AnimInputString += Anim.GetGeometryCacheSoftObjectPtr().ToSoftObjectPath().ToString();
+			AnimInputString += TEXT("\n");
+		}
+		if (AnimInputString.IsEmpty())
+		{
+			AnimInputString = TEXT("None");
+		}
+		OutTags.Add(FAssetRegistryTag("MLDeformer.TrainingAnims", AnimInputString, FAssetRegistryTag::TT_Alphabetical));
+		OutTags.Add(FAssetRegistryTag("MLDeformer.NumTrainingAnims", FString::FromInt(TrainingInputAnims.Num()), FAssetRegistryTag::TT_Numerical));
+	#endif
+}
+
 #if WITH_EDITOR
 void UMLDeformerGeomCacheModel::UpdateNumTargetMeshVertices()
 {
diff --git a/Engine/Plugins/Animation/MLDeformer/MLDeformerFramework/Source/MLDeformerFramework/Private/MLDeformerInputInfo.cpp b/Engine/Plugins/Animation/MLDeformer/MLDeformerFramework/Source/MLDeformerFramework/Private/MLDeformerInputInfo.cpp
--- a/Engine/Plugins/Animation/MLDeformer/MLDeformerFramework/Source/MLDeformerFramework/Private/MLDeformerInputInfo.cpp
+++ b/Engine/Plugins/Animation/MLDeformer/MLDeformerFramework/Source/MLDeformerFramework/Private/MLDeformerInputInfo.cpp
@@ -21,6 +21,16 @@ void UMLDeformerInputInfo::Reset()
 	SkeletalMesh = nullptr;
 }
 
+void UMLDeformerInputInfo::GetAssetRegistryTags(TArray<FAssetRegistryTag>& OutTags) const
+{
+	Super::GetAssetRegistryTags(OutTags);
+	OutTags.Add(FAssetRegistryTag("MLDeformer.Trained.SkeletalMesh", SkeletalMesh.IsValid() ? SkeletalMesh.ToString() : TEXT("None"), FAssetRegistryTag::TT_Alphabetical));
+	OutTags.Add(FAssetRegistryTag("MLDeformer.Trained.NumBaseMeshVertices", FString::FromInt(NumBaseMeshVertices), FAssetRegistryTag::TT_Numerical));
+	OutTags.Add(FAssetRegistryTag("MLDeformer.Trained.NumTargetMeshVertices", FString::FromInt(NumTargetMeshVertices), FAssetRegistryTag::TT_Numerical));
+	OutTags.Add(FAssetRegistryTag("MLDeformer.Trained.NumBones", FString::FromInt(GetNumBones()), FAssetRegistryTag::TT_Numerical));
+	OutTags.Add(FAssetRegistryTag("MLDeformer.Trained.NumCurves", FString::FromInt(GetNumCurves()), FAssetRegistryTag::TT_Numerical));
+}
+
 void UMLDeformerInputInfo::OnPostLoad()
 {
 	PRAGMA_DISABLE_DEPRECATION_WARNINGS
diff --git a/Engine/Plugins/Animation/MLDeformer/MLDeformerFramework/Source/MLDeformerFramework/Private/MLDeformerModel.cpp b/Engine/Plugins/Animation/MLDeformer/MLDeformerFramework/Source/MLDeformerFramework/Private/MLDeformerModel.cpp
--- a/Engine/Plugins/Animation/MLDeformer/MLDeformerFramework/Source/MLDeformerFramework/Private/MLDeformerModel.cpp
+++ b/Engine/Plugins/Animation/MLDeformer/MLDeformerFramework/Source/MLDeformerFramework/Private/MLDeformerModel.cpp
@@ -14,6 +14,7 @@
 #include "RHICommandList.h"
 #include "UObject/UObjectGlobals.h"
 #include "RHICommandList.h"
+#include "AssetRegistry/AssetData.h"
 
 #include UE_INLINE_GENERATED_CPP_BY_NAME(MLDeformerModel)
 
@@ -152,6 +153,29 @@ bool UMLDeformerModel::IsReadyForFinishDestroy()
 	return Super::IsReadyForFinishDestroy() && RenderResourceDestroyFence.IsFenceComplete();
 }
 
+void UMLDeformerModel::GetAssetRegistryTags(TArray<FAssetRegistryTag>& OutTags) const
+{
+	Super::GetAssetRegistryTags(OutTags);
+
+	OutTags.Add(FAssetRegistryTag("MLDeformer.ModelType", GetClass()->GetName(), FAssetRegistryTag::TT_Alphabetical));
+	OutTags.Add(FAssetRegistryTag("MLDeformer.IsTrained", IsTrained() ? TEXT("True") : TEXT("False"), FAssetRegistryTag::TT_Alphabetical));
+	OutTags.Add(FAssetRegistryTag("MLDeformer.NumBaseMeshVerts", FString::FromInt(NumBaseMeshVerts), FAssetRegistryTag::TT_Numerical));
+	OutTags.Add(FAssetRegistryTag("MLDeformer.NumTargetMeshVerts", FString::FromInt(NumTargetMeshVerts), FAssetRegistryTag::TT_Numerical));
+	OutTags.Add(FAssetRegistryTag("MLDeformer.SkeletalMesh", SkeletalMesh ? FAssetData(SkeletalMesh).ToSoftObjectPath().ToString() : TEXT("None"), FAssetRegistryTag::TT_Alphabetical));
+
+	#if WITH_EDITORONLY_DATA
+		OutTags.Add(FAssetRegistryTag("MLDeformer.NumBones", FString::FromInt(BoneIncludeList.Num()), FAssetRegistryTag::TT_Numerical));
+		OutTags.Add(FAssetRegistryTag("MLDeformer.NumCurves", FString::FromInt(CurveIncludeList.Num()), FAssetRegistryTag::TT_Numerical));
+		OutTags.Add(FAssetRegistryTag("MLDeformer.DeltaCutoffLength", FString::Printf(TEXT("%f"), DeltaCutoffLength), FAssetRegistryTag::TT_Numerical));
+		OutTags.Add(FAssetRegistryTag("MLDeformer.MaxTrainingFrames", FString::FromInt(MaxTrainingFrames), FAssetRegistryTag::TT_Numerical));
+	#endif
+
+	if (InputInfo)
+	{
+		InputInfo->GetAssetRegistryTags(OutTags);
+	}
+}
+
 void UMLDeformerModel::InitGPUData()
 {
 	BeginReleaseResource(&VertexMapBuffer);
diff --git a/Engine/Plugins/Animation/MLDeformer/MLDeformerFramework/Source/MLDeformerFramework/Private/MLDeformerMorphModel.cpp b/Engine/Plugins/Animation/MLDeformer/MLDeformerFramework/Source/MLDeformerFramework/Private/MLDeformerMorphModel.cpp
--- a/Engine/Plugins/Animation/MLDeformer/MLDeformerFramework/Source/MLDeformerFramework/Private/MLDeformerMorphModel.cpp
+++ b/Engine/Plugins/Animation/MLDeformer/MLDeformerFramework/Source/MLDeformerFramework/Private/MLDeformerMorphModel.cpp
@@ -71,6 +71,21 @@ void UMLDeformerMorphModel::PostLoad()
 #endif
 }
 
+void UMLDeformerMorphModel::GetAssetRegistryTags(TArray<FAssetRegistryTag>& OutTags) const
+{
+	Super::GetAssetRegistryTags(OutTags);
+
+	OutTags.Add(FAssetRegistryTag("MLDeformer.MorphModel.ClampMorphWeights", bClampMorphWeights ? TEXT("True") : TEXT("False"), FAssetRegistryTag::TT_Alphabetical));
+	OutTags.Add(FAssetRegistryTag("MLDeformer.MorphModel.InvertMaskChannel", bInvertMaskChannel? TEXT("True") : TEXT("False"), FAssetRegistryTag::TT_Alphabetical));
+	OutTags.Add(FAssetRegistryTag("MLDeformer.MorphModel.IncludeNormals", bIncludeNormals ? TEXT("True") : TEXT("False"), FAssetRegistryTag::TT_Alphabetical));
+	OutTags.Add(FAssetRegistryTag("MLDeformer.MorphModel.NumMorphTargets", FString::FromInt(NumMorphTargets), FAssetRegistryTag::TT_Numerical));
+	OutTags.Add(FAssetRegistryTag("MLDeformer.MorphModel.DeltaZeroThreshold", FString::Printf(TEXT("%f"), MorphDeltaZeroThreshold), FAssetRegistryTag::TT_Numerical));
+	OutTags.Add(FAssetRegistryTag("MLDeformer.MorphModel.CompressionLevel", FString::Printf(TEXT("%f"), MorphCompressionLevel), FAssetRegistryTag::TT_Numerical));
+	OutTags.Add(FAssetRegistryTag("MLDeformer.MorphModel.NumQualityLevels", FString::FromInt(QualityLevels.Num()), FAssetRegistryTag::TT_Numerical));
+	OutTags.Add(FAssetRegistryTag("MLDeformer.MorphModel.CompressedSize", FString::FromInt(CompressedMorphDataSizeInBytes), FAssetRegistryTag::TT_Numerical));
+	OutTags.Add(FAssetRegistryTag("MLDeformer.MorphModel.UncompressedSize", FString::FromInt(UncompressedMorphDataSizeInBytes), FAssetRegistryTag::TT_Numerical));
+}
+
 bool UMLDeformerMorphModel::CanDynamicallyUpdateMorphTargets() const
 {
 	return GetMorphTargetDeltas().Num() == (GetNumBaseMeshVerts() * GetNumMorphTargets());
diff --git a/Engine/Plugins/Animation/MLDeformer/MLDeformerFramework/Source/MLDeformerFramework/Public/MLDeformerAsset.h b/Engine/Plugins/Animation/MLDeformer/MLDeformerFramework/Source/MLDeformerFramework/Public/MLDeformerAsset.h
--- a/Engine/Plugins/Animation/MLDeformer/MLDeformerFramework/Source/MLDeformerFramework/Public/MLDeformerAsset.h
+++ b/Engine/Plugins/Animation/MLDeformer/MLDeformerFramework/Source/MLDeformerFramework/Public/MLDeformerAsset.h
@@ -22,6 +22,7 @@ class MLDEFORMERFRAMEWORK_API UMLDeformerAsset
 public:
 	// UObject overrides.
 	virtual void Serialize(FArchive& Archive) override;
+	virtual void GetAssetRegistryTags(TArray<FAssetRegistryTag>& OutTags) const override;
 	// ~END UObject overrides.
 
 	/**
diff --git a/Engine/Plugins/Animation/MLDeformer/MLDeformerFramework/Source/MLDeformerFramework/Public/MLDeformerGeomCacheModel.h b/Engine/Plugins/Animation/MLDeformer/MLDeformerFramework/Source/MLDeformerFramework/Public/MLDeformerGeomCacheModel.h
--- a/Engine/Plugins/Animation/MLDeformer/MLDeformerFramework/Source/MLDeformerFramework/Public/MLDeformerGeomCacheModel.h
+++ b/Engine/Plugins/Animation/MLDeformer/MLDeformerFramework/Source/MLDeformerFramework/Public/MLDeformerGeomCacheModel.h
@@ -30,6 +30,7 @@ public:
 	// UObject overrides.
 	virtual void Serialize(FArchive& Archive) override;
 	virtual void PostLoad() override;
+	virtual void GetAssetRegistryTags(TArray<FAssetRegistryTag>& OutTags) const override;
 	// ~END UObject overrides.
 
 	// UMLDeformerModel overrides.
diff --git a/Engine/Plugins/Animation/MLDeformer/MLDeformerFramework/Source/MLDeformerFramework/Public/MLDeformerGeomCacheTrainingInputAnim.h b/Engine/Plugins/Animation/MLDeformer/MLDeformerFramework/Source/MLDeformerFramework/Public/MLDeformerGeomCacheTrainingInputAnim.h
--- a/Engine/Plugins/Animation/MLDeformer/MLDeformerFramework/Source/MLDeformerFramework/Public/MLDeformerGeomCacheTrainingInputAnim.h
+++ b/Engine/Plugins/Animation/MLDeformer/MLDeformerFramework/Source/MLDeformerFramework/Public/MLDeformerGeomCacheTrainingInputAnim.h
@@ -25,16 +25,17 @@ public:
 	virtual int32 ExtractNumAnimFrames() const override;
 	// ~END FMLDeformerTrainingInputAnim overrides
 
-	void SetGeometryCache(TSoftObjectPtr<UGeometryCache> GeomCache)	{ GeometryCache = GeomCache; }
+	void SetGeometryCache(TSoftObjectPtr<UGeometryCache> GeomCache)				{ GeometryCache = GeomCache; }
 
-	const UGeometryCache* GetGeometryCache() const					{ return GeometryCache.LoadSynchronous(); }
-	UGeometryCache* GetGeometryCache()								{ return GeometryCache.LoadSynchronous(); }
-
-	static FName GetGeomCachePropertyName()							{ return GET_MEMBER_NAME_CHECKED(FMLDeformerGeomCacheTrainingInputAnim, GeometryCache); }
+	const UGeometryCache* GetGeometryCache() const								{ return GeometryCache.LoadSynchronous(); }
+	UGeometryCache* GetGeometryCache()											{ return GeometryCache.LoadSynchronous(); }
+	TSoftObjectPtr<UGeometryCache>& GetGeometryCacheSoftObjectPtr()				{ return GeometryCache; }
+	const TSoftObjectPtr<UGeometryCache>& GetGeometryCacheSoftObjectPtr() const { return GeometryCache; }
+	static FName GetGeomCachePropertyName()										{ return GET_MEMBER_NAME_CHECKED(FMLDeformerGeomCacheTrainingInputAnim, GeometryCache); }
 
 private:
 	/** The geometry cache which contains the target deformations. */
 	UPROPERTY(EditAnywhere, Category = "Settings", meta = (DisplayPriority = 10))
-	TSoftObjectPtr<UGeometryCache> GeometryCache = nullptr;
+	TSoftObjectPtr<UGeometryCache> GeometryCache;
 #endif	// #if WITH_EDITORONLY_DATA
 };
diff --git a/Engine/Plugins/Animation/MLDeformer/MLDeformerFramework/Source/MLDeformerFramework/Public/MLDeformerInputInfo.h b/Engine/Plugins/Animation/MLDeformer/MLDeformerFramework/Source/MLDeformerFramework/Public/MLDeformerInputInfo.h
--- a/Engine/Plugins/Animation/MLDeformer/MLDeformerFramework/Source/MLDeformerFramework/Public/MLDeformerInputInfo.h
+++ b/Engine/Plugins/Animation/MLDeformer/MLDeformerFramework/Source/MLDeformerFramework/Public/MLDeformerInputInfo.h
@@ -24,6 +24,10 @@ class MLDEFORMERFRAMEWORK_API UMLDeformerInputInfo
 public:
 	virtual ~UMLDeformerInputInfo() = default;
 
+	// UObject overrides.
+	virtual void GetAssetRegistryTags(TArray<FAssetRegistryTag>& OutTags) const override;
+	// ~END UObject overrides.
+
 	/**
 	 * This method is executed post loading of the ML Deformer asset.
 	 */
@@ -222,7 +226,7 @@ protected:
 	/**
 	 * The path to the skeletal mesh that this model was trained on.
 	 */
-	UPROPERTY()	
+	UPROPERTY()
 	FSoftObjectPath SkeletalMesh;
 
 #if WITH_EDITORONLY_DATA
@@ -252,7 +256,7 @@ protected:
 	TArray<FName> CurveNames;
 
 	/** Number of imported base mesh vertices, so not render vertices. */
-	UPROPERTY()
+	UPROPERTY(AssetRegistrySearchable)
 	int32 NumBaseMeshVertices = 0;
 
 	/** Number of imported target mesh vertices, so not render vertices. */
diff --git a/Engine/Plugins/Animation/MLDeformer/MLDeformerFramework/Source/MLDeformerFramework/Public/MLDeformerModel.h b/Engine/Plugins/Animation/MLDeformer/MLDeformerFramework/Source/MLDeformerFramework/Public/MLDeformerModel.h
--- a/Engine/Plugins/Animation/MLDeformer/MLDeformerFramework/Source/MLDeformerFramework/Public/MLDeformerModel.h
+++ b/Engine/Plugins/Animation/MLDeformer/MLDeformerFramework/Source/MLDeformerFramework/Public/MLDeformerModel.h
@@ -188,6 +188,9 @@ public:
 	 */
 	virtual int32 GetNumFloatsPerCurve() const				{ return 1; }
 
+	/** Check whether this model has been trained or not. */
+	virtual bool IsTrained() const							{ ensureMsgf(false, TEXT("Please override the UMLDeformerModel::IsTrained() inside your model.")); return false; }
+
 	/**
 	 * Get the skeletal mesh that is used during training.
 	 * You typically want to apply the ML Deformer on this specific skeletal mesh in your game as well.
@@ -297,6 +300,7 @@ public:
 	virtual void PostLoad() override;
 	virtual void BeginDestroy() override;
 	virtual bool IsReadyForFinishDestroy() override;
+	virtual void GetAssetRegistryTags(TArray<FAssetRegistryTag>& OutTags) const override;
 	// ~END UObject overrides.
 
 	// IBoneReferenceSkeletonProvider overrides.
diff --git a/Engine/Plugins/Animation/MLDeformer/MLDeformerFramework/Source/MLDeformerFramework/Public/MLDeformerMorphModel.h b/Engine/Plugins/Animation/MLDeformer/MLDeformerFramework/Source/MLDeformerFramework/Public/MLDeformerMorphModel.h
--- a/Engine/Plugins/Animation/MLDeformer/MLDeformerFramework/Source/MLDeformerFramework/Public/MLDeformerMorphModel.h
+++ b/Engine/Plugins/Animation/MLDeformer/MLDeformerFramework/Source/MLDeformerFramework/Public/MLDeformerMorphModel.h
@@ -39,6 +39,7 @@ public:
 	virtual void BeginDestroy() override;
 	virtual bool IsReadyForFinishDestroy() override;
 	virtual void PostLoad() override;
+	virtual void GetAssetRegistryTags(TArray<FAssetRegistryTag>& OutTags) const override;
 	// ~END UObject overrides.
 
 	int32 GetNumMorphTargets() const								{ return NumMorphTargets; }
diff --git a/Engine/Plugins/Animation/MLDeformer/MLDeformerFramework/Source/MLDeformerFramework/Public/MLDeformerTrainingInputAnim.h b/Engine/Plugins/Animation/MLDeformer/MLDeformerFramework/Source/MLDeformerFramework/Public/MLDeformerTrainingInputAnim.h
--- a/Engine/Plugins/Animation/MLDeformer/MLDeformerFramework/Source/MLDeformerFramework/Public/MLDeformerTrainingInputAnim.h
+++ b/Engine/Plugins/Animation/MLDeformer/MLDeformerFramework/Source/MLDeformerFramework/Public/MLDeformerTrainingInputAnim.h
@@ -46,6 +46,8 @@ public:
 	int32 GetEndFrame() const									{ return EndFrame; }
 	const UAnimSequence* GetAnimSequence() const				{ return AnimSequence.LoadSynchronous(); }
 	UAnimSequence* GetAnimSequence()							{ return AnimSequence.LoadSynchronous(); }
+	const TSoftObjectPtr<UAnimSequence>& GetAnimSequenceSoftObjectPtr() const	{ return AnimSequence; } 
+	TSoftObjectPtr<UAnimSequence>& GetAnimSequenceSoftObjectPtr()				{ return AnimSequence; } 
 
 	static FName GetAnimSequencePropertyName()					{ return GET_MEMBER_NAME_CHECKED(FMLDeformerTrainingInputAnim, AnimSequence); }
 	static FName GetEnabledPropertyName()						{ return GET_MEMBER_NAME_CHECKED(FMLDeformerTrainingInputAnim, bEnabled); }
@@ -56,7 +58,7 @@ public:
 private:
 	/** The animation sequence. */
 	UPROPERTY(EditAnywhere, Category = "Settings", meta = (DisplayPriority = 0))	// Show as first property.
-	TSoftObjectPtr<UAnimSequence> AnimSequence = nullptr;
+	TSoftObjectPtr<UAnimSequence> AnimSequence;
 
 	/** On default all frames are included, unless we specify a custom frame include range. */
 	UPROPERTY(EditAnywhere, Category = "Frame Settings")
diff --git a/Engine/Plugins/Animation/MLDeformer/MLDeformerFramework/Source/MLDeformerFramework/Public/MLDeformerVizSettings.h b/Engine/Plugins/Animation/MLDeformer/MLDeformerFramework/Source/MLDeformerFramework/Public/MLDeformerVizSettings.h
--- a/Engine/Plugins/Animation/MLDeformer/MLDeformerFramework/Source/MLDeformerFramework/Public/MLDeformerVizSettings.h
+++ b/Engine/Plugins/Animation/MLDeformer/MLDeformerFramework/Source/MLDeformerFramework/Public/MLDeformerVizSettings.h
@@ -8,9 +8,9 @@
 #include "UObject/SoftObjectPtr.h"
 #include "Animation/MeshDeformer.h"
 #include "Animation/AnimSequence.h"
+#include "MLDeformerAsset.h"
 #include "MLDeformerVizSettings.generated.h"
 
-
 /** The visualization mode, which selects whether you want to view the training data, or test your already trained model. */
 UENUM()
 enum class EMLDeformerVizMode : uint8
@@ -33,6 +33,21 @@ enum class EMLDeformerHeatMapMode : uint8
 	GroundTruth
 };
 
+/** A comparison actor. */
+USTRUCT()
+struct MLDEFORMERFRAMEWORK_API FMLDeformerCompareActor
+{
+	GENERATED_BODY()
+
+	/** The name of the comparison actor, which is the label shown above it. */
+	UPROPERTY(EditAnywhere, Category = "Actor Settings")
+	FName Name;
+
+	/** The ML Deformer asset that this actor should use. */
+	UPROPERTY(EditAnywhere, Category = "Actor Settings")
+	TSoftObjectPtr<UMLDeformerAsset> DeformerAsset;
+};
+
 /**
  * The vizualization settings.
  */
@@ -74,6 +89,7 @@ public:
 	bool GetDrawLinearSkinnedActor() const					{ return bDrawLinearSkinnedActor; }
 	bool GetDrawMLDeformedActor() const						{ return bDrawMLDeformedActor; }
 	bool GetDrawGroundTruthActor() const					{ return bDrawGroundTruthActor; }
+	bool GetDrawMLCompareActors() const						{ return bDrawMLCompareActors; }
 	bool GetShowHeatMap() const								{ return bShowHeatMap; }
 	EMLDeformerHeatMapMode GetHeatMapMode() const			{ return HeatMapMode; }
 	float GetHeatMapMax() const								{ return HeatMapMax; }
@@ -83,6 +99,8 @@ public:
 	bool GetXRayDeltas() const								{ return bXRayDeltas; }
 	bool GetDrawVertexDeltas() const						{ return bDrawDeltas; }
 	int32 GetQualityLevel() const							{ return QualityLevel; }
+	const TArray<FMLDeformerCompareActor>& GetCompareActors() const { return CompareActors; }
+	TArray<FMLDeformerCompareActor>& GetCompareActors()		{ return CompareActors; }
 
 	// Get property names.
 	static FName GetVisualizationModePropertyName()			{ return GET_MEMBER_NAME_CHECKED(UMLDeformerVizSettings, VisualizationMode); }
@@ -97,6 +115,7 @@ public:
 	static FName GetDrawLinearSkinnedActorPropertyName()	{ return GET_MEMBER_NAME_CHECKED(UMLDeformerVizSettings, bDrawLinearSkinnedActor); }
 	static FName GetDrawMLDeformedActorPropertyName()		{ return GET_MEMBER_NAME_CHECKED(UMLDeformerVizSettings, bDrawMLDeformedActor); }
 	static FName GetDrawGroundTruthActorPropertyName()		{ return GET_MEMBER_NAME_CHECKED(UMLDeformerVizSettings, bDrawGroundTruthActor); }
+	static FName GetDrawMLCompareActorsPropertyName()		{ return GET_MEMBER_NAME_CHECKED(UMLDeformerVizSettings, bDrawMLCompareActors); }
 	static FName GetShowHeatMapPropertyName()				{ return GET_MEMBER_NAME_CHECKED(UMLDeformerVizSettings, bShowHeatMap); }
 	static FName GetHeatMapModePropertyName()				{ return GET_MEMBER_NAME_CHECKED(UMLDeformerVizSettings, HeatMapMode); }
 	static FName GetHeatMapMaxPropertyName()				{ return GET_MEMBER_NAME_CHECKED(UMLDeformerVizSettings, HeatMapMax); }
@@ -106,6 +125,7 @@ public:
 	static FName GetXRayDeltasPropertyName()				{ return GET_MEMBER_NAME_CHECKED(UMLDeformerVizSettings, bXRayDeltas); }
 	static FName GetDrawVertexDeltasPropertyName()			{ return GET_MEMBER_NAME_CHECKED(UMLDeformerVizSettings, bDrawDeltas); }
 	static FName GetQualityLevelPropertyName()				{ return GET_MEMBER_NAME_CHECKED(UMLDeformerVizSettings, QualityLevel); }
+	static FName GetCompareActorsPropertyName()				{ return GET_MEMBER_NAME_CHECKED(UMLDeformerVizSettings, CompareActors); }
 #endif
 
 protected:
@@ -116,11 +136,18 @@ protected:
 
 	/** The animation sequence to play on the skeletal mesh. */
 	UPROPERTY(EditAnywhere, Category = "Test Assets")
-	TSoftObjectPtr<UAnimSequence> TestAnimSequence = nullptr;
+	TSoftObjectPtr<UAnimSequence> TestAnimSequence;
 
 	/** The deformer graph to use on the asset editor's deformed test actor. */
 	UPROPERTY(EditAnywhere, Category = "Test Assets")
-	TSoftObjectPtr<UMeshDeformer> DeformerGraph = nullptr;
+	TSoftObjectPtr<UMeshDeformer> DeformerGraph;
+
+	/**
+	 * The model comparison actors. Each will create a skeletal mesh with the selected ML Deformer applied to it.
+	 * This is useful to compare the output of different models side by side. 
+	 */
+	UPROPERTY(EditAnywhere, Category = "Test Assets")
+	TArray<FMLDeformerCompareActor> CompareActors;
 
 	/** The play speed factor of the test anim sequence. */
 	UPROPERTY(EditAnywhere, Category = "Live Settings", meta = (ClampMin = "0.0", ClampMax = "2.0", ForceUnits="Multiplier"))
@@ -196,6 +223,10 @@ protected:
 	UPROPERTY(EditAnywhere, Category = "Live Settings")
 	bool bDrawGroundTruthActor = true;
 
+	/** Specifies whether we draw the comparison actors or not. */
+	UPROPERTY(EditAnywhere, Category = "Live Settings", DisplayName = "Draw ML Compare Actors")
+	bool bDrawMLCompareActors = true;
+
 	/** The scale factor of the ML deformer deltas being applied on top of the linear skinned results. */
 	UPROPERTY(EditAnywhere, Transient, Category = "Live Settings", meta = (ClampMin = "0.0", ClampMax = "1.0"))
 	float Weight = 1.0f;
diff --git a/Engine/Plugins/Animation/MLDeformer/MLDeformerFramework/Source/MLDeformerFrameworkEditor/Private/MLDeformerEditorActor.cpp b/Engine/Plugins/Animation/MLDeformer/MLDeformerFramework/Source/MLDeformerFrameworkEditor/Private/MLDeformerEditorActor.cpp
--- a/Engine/Plugins/Animation/MLDeformer/MLDeformerFramework/Source/MLDeformerFrameworkEditor/Private/MLDeformerEditorActor.cpp
+++ b/Engine/Plugins/Animation/MLDeformer/MLDeformerFramework/Source/MLDeformerFrameworkEditor/Private/MLDeformerEditorActor.cpp
@@ -17,6 +17,7 @@ namespace UE::MLDeformer
 		Actor = Settings.Actor;
 		LabelComponent = CreateLabelComponent(Settings.Actor, Settings.LabelColor, Settings.LabelText);
 		bIsTrainingActor = Settings.bIsTrainingActor;
+		ActorTypeInstanceIndex = Settings.ActorTypeInstanceIndex;
 	}
 
 	FMLDeformerEditorActor::~FMLDeformerEditorActor()
diff --git a/Engine/Plugins/Animation/MLDeformer/MLDeformerFramework/Source/MLDeformerFrameworkEditor/Private/MLDeformerEditorModel.cpp b/Engine/Plugins/Animation/MLDeformer/MLDeformerFramework/Source/MLDeformerFrameworkEditor/Private/MLDeformerEditorModel.cpp
--- a/Engine/Plugins/Animation/MLDeformer/MLDeformerFramework/Source/MLDeformerFrameworkEditor/Private/MLDeformerEditorModel.cpp
+++ b/Engine/Plugins/Animation/MLDeformer/MLDeformerFramework/Source/MLDeformerFrameworkEditor/Private/MLDeformerEditorModel.cpp
@@ -286,13 +286,70 @@ namespace UE::MLDeformer
 		Settings.LabelColor = LabelColor;
 		Settings.LabelText = LOCTEXT("TestMLDeformedActorLabelText", "ML Deformed");
 		Settings.bIsTrainingActor = false;
-		FMLDeformerEditorActor* EditorActor = static_cast<FMLDeformerEditorActor*>(CreateEditorActor(Settings));
+		FMLDeformerEditorActor* EditorActor = CreateEditorActor(Settings);
 		EditorActor->SetSkeletalMeshComponent(SkelMeshComponent);
 		EditorActor->SetMLDeformerComponent(MLDeformerComponent);
 		EditorActor->SetMeshOffsetFactor(1.0f);
 		EditorActors.Add(EditorActor);
 	}
 
+	void FMLDeformerEditorModel::CreateTestCompareActors(UWorld* World)
+	{
+		const FLinearColor MLDeformedWireColor = FMLDeformerEditorStyle::Get().GetColor("MLDeformer.MLCompareMesh.WireframeColor");
+		USkeletalMesh* SkelMesh = Model->GetSkeletalMesh();
+
+		const TArray<FMLDeformerCompareActor>& CompareActors = Model->GetVizSettings()->GetCompareActors();
+		for (int32 ActorIndex = 0; ActorIndex < CompareActors.Num(); ++ActorIndex)
+		{
+			const FMLDeformerCompareActor& CompareActor = CompareActors[ActorIndex];
+
+			// Create the ML deformed actor.
+			FActorSpawnParameters SpawnParams;
+			SpawnParams.Name = MakeUniqueObjectName(World, AActor::StaticClass(), "Test Compare Actor");
+			AActor* Actor = World->SpawnActor<AActor>(SpawnParams);
+			Actor->SetFlags(RF_Transient);
+
+			// Create the skeletal mesh component.
+			UDebugSkelMeshComponent* SkelMeshComponent = NewObject<UDebugSkelMeshComponent>(Actor);
+			UMLDeformerAsset* DeformerAsset = CompareActor.DeformerAsset.LoadSynchronous();
+			if (!IsCompatibleDeformer(DeformerAsset))
+			{
+				SkelMesh = nullptr;
+			}
+			SkelMeshComponent->SetSkeletalMesh(SkelMesh);
+			Actor->SetRootComponent(SkelMeshComponent);
+			SkelMeshComponent->RegisterComponent();
+			SkelMeshComponent->SetWireframeMeshOverlayColor(MLDeformedWireColor);
+			SkelMeshComponent->SetVisibility(false);
+			SkelMeshComponent->MarkRenderStateDirty();
+
+			// Create the ML Deformer component.
+			UMLDeformerAsset* FinalDeformerAsset = IsCompatibleDeformer(DeformerAsset) ? DeformerAsset : nullptr;
+			UMLDeformerComponent* MLDeformerComponent = NewObject<UMLDeformerComponent>(Actor);
+			MLDeformerComponent->SetSuppressMeshDeformerLogWarnings(true);
+			MLDeformerComponent->SetDeformerAsset(FinalDeformerAsset);
+			MLDeformerComponent->RegisterComponent();
+			MLDeformerComponent->SetupComponent(FinalDeformerAsset, SkelMeshComponent);
+
+			// Create the editor actor.
+			const FLinearColor LabelColor = FMLDeformerEditorStyle::Get().GetColor("MLDeformer.MLCompareMesh.LabelColor");
+			FMLDeformerEditorActor::FConstructSettings Settings;
+			Settings.Actor = Actor;
+			Settings.TypeID = ActorID_Test_Compare;
+			Settings.LabelColor = LabelColor;
+			Settings.LabelText = LOCTEXT("TestCompareActorLabelText", "ML Compare");
+			Settings.bIsTrainingActor = false;
+			Settings.ActorTypeInstanceIndex = ActorIndex;
+			FMLDeformerEditorActor* EditorActor = CreateEditorActor(Settings);
+			EditorActor->SetSkeletalMeshComponent(SkelMeshComponent);
+			EditorActor->SetMLDeformerComponent(MLDeformerComponent);
+			EditorActor->SetMeshOffsetFactor(2.0f + static_cast<float>(ActorIndex));
+			EditorActors.Add(EditorActor);
+		}
+
+		UpdateCompareActorLabels();
+	}
+
 	void FMLDeformerEditorModel::CreateActors(const TSharedRef<IPersonaPreviewScene>& InPersonaPreviewScene)
 	{
 		UWorld* World = InPersonaPreviewScene->GetWorld();
@@ -301,11 +358,43 @@ namespace UE::MLDeformer
 		CreateTestMLDeformedActor(World);
 		CreateTrainingGroundTruthActor(World);
 		CreateTestGroundTruthActor(World);
+		CreateTestCompareActors(World);
+
+		UpdateMeshOffsetFactors();
+	}
+
+	int32 FMLDeformerEditorModel::CalcNumValidCompareActorsPriorTo(int32 CompareActorIndex) const
+	{
+		int Result = 0;
+
+		const TArray<FMLDeformerCompareActor>& CompareActors = Model->GetVizSettings()->GetCompareActors();
+		for (int32 CompareInstanceIndex = 0; CompareInstanceIndex < CompareActors.Num(); ++CompareInstanceIndex)
+		{
+			if (CompareInstanceIndex < CompareActorIndex)
+			{
+				const FMLDeformerCompareActor& CompareActor = CompareActors[CompareInstanceIndex];
+				UMLDeformerAsset* CompareDeformerAsset = CompareActor.DeformerAsset.LoadSynchronous();
+				if (IsCompatibleDeformer(CompareDeformerAsset))
+				{
+					Result++;
+				}
+			}
+		}
 
+		return Result;
+	}
+
+	void FMLDeformerEditorModel::UpdateMeshOffsetFactors()
+	{
 		// Set the default mesh translation offsets for our ground truth actors.
 		for (FMLDeformerEditorActor* EditorActor : EditorActors)
 		{
-			if (EditorActor && EditorActor->IsGroundTruthActor())
+			if (EditorActor == nullptr)
+			{
+				continue;
+			}
+
+			if (EditorActor->IsGroundTruthActor())
 			{			
 				// The mesh offset factor basically just offsets the actor position by a given factor.
 				// The amount the actor is moved from the origin is: (MeshSpacing * MeshOffsetFactor).
@@ -313,7 +402,25 @@ namespace UE::MLDeformer
 				// It is 2.0 because the ground truth actor in testing mode is all the way on the right, next to the ML Deformed model.
 				// In training mode we have only the Linear Skinned actor and the ground truth, so there the spacing factor is 1.0.
 				// TLDR: Basically 1.0 means its the first actor next to the linear skinned actor while 2.0 means its the second character, etc.
-				EditorActor->SetMeshOffsetFactor(EditorActor->IsTestActor() ? 2.0f : 1.0f);
+				float MeshOffsetFactor = 1.0f;
+				if (EditorActor->IsTestActor())
+				{
+					MeshOffsetFactor = 2.0f;
+
+					if (Model->GetVizSettings()->GetDrawMLCompareActors())
+					{
+						const TArray<FMLDeformerCompareActor>& CompareActors = Model->GetVizSettings()->GetCompareActors();
+						const int32 NumValidCompareActors = static_cast<float>(CalcNumValidCompareActorsPriorTo(CompareActors.Num()));
+						MeshOffsetFactor += static_cast<float>(NumValidCompareActors);
+					}
+				}
+
+				EditorActor->SetMeshOffsetFactor(MeshOffsetFactor);
+			}
+			else if (EditorActor->GetTypeID() == ActorID_Test_Compare)
+			{
+				const int32 NumValidCompareActorsBefore = static_cast<float>(CalcNumValidCompareActorsPriorTo(EditorActor->GetActorTypeInstanceIndex()));
+				EditorActor->SetMeshOffsetFactor(2.0f + static_cast<float>(NumValidCompareActorsBefore));
 			}
 		}
 	}
@@ -343,6 +450,22 @@ namespace UE::MLDeformer
 		}
 	}
 
+	bool FMLDeformerEditorModel::IsCompatibleDeformer(UMLDeformerAsset* Deformer) const
+	{
+		if (Deformer == nullptr)
+		{
+			return false;
+		}
+
+		UMLDeformerModel* CompareModel = Deformer->GetModel();
+		if (CompareModel->GetInputInfo()->GetSkeletalMesh() != Model->GetSkeletalMesh())
+		{
+			return false;
+		}
+
+		return true;
+	}
+
 	void FMLDeformerEditorModel::ClearPersonaPreviewScene()
 	{
 		TSharedRef<IPersonaPreviewScene> PreviewScene = Editor->GetPersonaToolkit()->GetPreviewScene();
@@ -426,18 +549,18 @@ namespace UE::MLDeformer
 			}
 		}
 
+		UpdateMeshOffsetFactors();
 		UpdateActorTransforms();
 		UpdateLabels();
 		CheckTrainingDataFrameChanged();
 
-		// Update the ML Deformer component of the ML Deformed test actor.
-		FMLDeformerEditorActor* EditorActor = FindEditorActor(ActorID_Test_MLDeformed);
-		if (EditorActor)
+		// Update the ML Deformer components.
+		for (FMLDeformerEditorActor* EditorActor : EditorActors)
 		{
-			UMLDeformerComponent* DeformerComponent = EditorActor->GetMLDeformerComponent();
-			if (DeformerComponent)
+			if (EditorActor && EditorActor->GetMLDeformerComponent())
 			{		
 				const UMLDeformerVizSettings* VizSettings = Model->GetVizSettings();
+				UMLDeformerComponent* DeformerComponent = EditorActor->GetMLDeformerComponent();
 				DeformerComponent->SetWeight(VizSettings->GetWeight());
 				DeformerComponent->SetQualityLevel(VizSettings->GetQualityLevel());
 			}
@@ -474,7 +597,7 @@ namespace UE::MLDeformer
 				LabelComponent->SetRelativeScale3D(FVector(VizSettings->GetLabelScale() * 0.3f));
 
 				// Update visibility.
-				const bool bLabelIsVisible = (bDrawTrainingActors && EditorActor->IsTrainingActor()) || (bDrawTestActors && EditorActor->IsTestActor());
+				const bool bLabelIsVisible = ((bDrawTrainingActors && EditorActor->IsTrainingActor()) || (bDrawTestActors && EditorActor->IsTestActor())) && EditorActor->IsVisible();
 				LabelComponent->SetVisibility(bLabelIsVisible);
 
 				// Handle test ground truth, disable its label when no ground truth asset was selected.
@@ -526,6 +649,13 @@ namespace UE::MLDeformer
 				{
 					bIsVisible &= VizSettings->GetDrawGroundTruthActor();
 				}
+				else if (EditorActor->GetTypeID() == ActorID_Test_Compare)
+				{
+					bIsVisible &= VizSettings->GetDrawMLCompareActors();
+				}
+
+				bIsVisible &= EditorActor->HasVisualMesh();
+
 				EditorActor->SetVisibility(bIsVisible);
 			}
 		}
@@ -609,6 +739,35 @@ namespace UE::MLDeformer
 			SkeletalMeshComponent->SetPlayRate(TestAnimSpeed);
 			SkeletalMeshComponent->Play(true);
 		}
+
+		// Update all compare actors.
+		for (FMLDeformerEditorActor* EditorActor : EditorActors)
+		{
+			if (EditorActor->GetTypeID() != ActorID_Test_Compare)
+			{
+				continue;
+			}
+
+			SkeletalMeshComponent = EditorActor->GetSkeletalMeshComponent();
+			if (SkeletalMeshComponent)
+			{
+				USkeletalMesh* CompareSkelMesh = Model->GetSkeletalMesh();
+				const int32 InstanceIndex = EditorActor->GetActorTypeInstanceIndex();
+				const FMLDeformerCompareActor& CompareActor = VizSettings->GetCompareActors()[InstanceIndex];
+				UMLDeformerAsset* DeformerAsset = CompareActor.DeformerAsset.LoadSynchronous();
+				if (!IsCompatibleDeformer(DeformerAsset))
+				{
+					CompareSkelMesh = nullptr;
+				}
+				SkeletalMeshComponent->SetSkeletalMesh(CompareSkelMesh);
+				SkeletalMeshComponent->SetAnimationMode(EAnimationMode::AnimationSingleNode);
+				const float CurrentPlayTime = SkeletalMeshComponent->GetPosition();
+				SkeletalMeshComponent->SetAnimation(TestAnimSequence);
+				SkeletalMeshComponent->SetPosition(CurrentPlayTime);
+				SkeletalMeshComponent->SetPlayRate(TestAnimSpeed);
+				SkeletalMeshComponent->Play(true);
+			}
+		}
 	}
 
 	void FMLDeformerEditorModel::OnPostInputAssetChanged()
@@ -617,6 +776,7 @@ namespace UE::MLDeformer
 		UpdateTimelineTrainingAnimList();
 		Editor->UpdateTimeSliderRange();
 		Model->UpdateCachedNumVertices();
+		UpdateActorVisibility();
 		UpdateNumTrainingFrames();
 		UpdateDeformerGraph();
 		RefreshMLDeformerComponents();
@@ -765,6 +925,60 @@ namespace UE::MLDeformer
 		{
 			TriggerInputAssetChanged(true);
 		}
+		else if (Property->GetFName() == UMLDeformerVizSettings::GetCompareActorsPropertyName())
+		{
+			const int32 ArrayIndex = PropertyChangedEvent.GetArrayIndex(PropertyChangedEvent.Property->GetFName().ToString());
+
+			// Make sure we're not moving elements.
+			switch (PropertyChangedEvent.ChangeType)
+			{
+				case EPropertyChangeType::ArrayAdd:
+				{
+					check(ArrayIndex != INDEX_NONE);
+					AddCompareActor(ArrayIndex);
+					break;
+				}
+				case EPropertyChangeType::ArrayRemove:
+				{
+					check(ArrayIndex != INDEX_NONE);
+					RemoveCompareActor(ArrayIndex);
+					break;
+				}
+				case EPropertyChangeType::ArrayClear:
+				{
+					RemoveAllCompareActors();
+					break;
+				}
+				case EPropertyChangeType::Duplicate:
+				{
+					check(ArrayIndex != INDEX_NONE);
+					AddCompareActor(ArrayIndex);
+					break;
+				}
+				default:;
+			};
+
+			TriggerInputAssetChanged(true);
+			UpdateActorVisibility();
+			UpdateCompareActorLabels();
+		}
+		else if (Property->GetFName() == GET_MEMBER_NAME_CHECKED(FMLDeformerCompareActor, DeformerAsset))
+		{
+			if (PropertyChangedEvent.ChangeType == EPropertyChangeType::ValueSet)
+			{
+				TriggerInputAssetChanged(true);
+				UpdateActorVisibility();
+				UpdateCompareActorLabels();
+			}
+		}
+		else if (Property->GetFName() == GET_MEMBER_NAME_CHECKED(FMLDeformerCompareActor, Name))
+		{
+			if (PropertyChangedEvent.ChangeType == EPropertyChangeType::ValueSet)
+			{
+				UpdateCompareActorLabels();
+			}
+		}
+		else
 		if (Property->GetFName() == UMLDeformerModel::GetAlignmentTransformPropertyName() ||
 		    Property->GetFName() == UMLDeformerModel::GetDeltaCutoffLengthPropertyName())
 		{
@@ -811,7 +1025,8 @@ namespace UE::MLDeformer
 		else
 		if (Property->GetFName() == UMLDeformerVizSettings::GetDrawLinearSkinnedActorPropertyName() ||
 			Property->GetFName() == UMLDeformerVizSettings::GetDrawMLDeformedActorPropertyName() ||
-			Property->GetFName() == UMLDeformerVizSettings::GetDrawGroundTruthActorPropertyName())
+			Property->GetFName() == UMLDeformerVizSettings::GetDrawGroundTruthActorPropertyName() ||
+			Property->GetFName() == UMLDeformerVizSettings::GetDrawMLCompareActorsPropertyName())
 		{
 			UpdateActorVisibility();
 		}
@@ -1467,18 +1682,40 @@ namespace UE::MLDeformer
 
 	void FMLDeformerEditorModel::RefreshMLDeformerComponents()
 	{
-		for (FMLDeformerEditorActor* EditorActor : EditorActors)
+		// Update the ML Deformed actor.
+		FMLDeformerEditorActor* EditorActor = FindEditorActor(ActorID_Test_MLDeformed);
+		if (EditorActor && EditorActor->GetMLDeformerComponent())
 		{
-			if (EditorActor && EditorActor->GetMLDeformerComponent())
+			USkeletalMeshComponent* SkelMeshComponent = EditorActor ? EditorActor->GetSkeletalMeshComponent() : nullptr;
+			UMLDeformerAsset* DeformerAsset = GetModel()->GetDeformerAsset();
+			EditorActor->GetMLDeformerComponent()->SetupComponent(DeformerAsset, SkelMeshComponent);
+			UMLDeformerModelInstance* ModelInstance = EditorActor->GetMLDeformerComponent()->GetModelInstance();
+			if (ModelInstance)
+			{
+				ModelInstance->UpdateCompatibilityStatus();
+			}
+		}
+
+		// Update all compare actors.
+		for (FMLDeformerEditorActor* CompareEditorActor : EditorActors)
+		{
+			if (CompareEditorActor->GetTypeID() != ActorID_Test_Compare)
 			{
-				USkeletalMeshComponent* SkelMeshComponent = EditorActor ? EditorActor->GetSkeletalMeshComponent() : nullptr;
-				UMLDeformerAsset* DeformerAsset = GetModel()->GetDeformerAsset();
-				EditorActor->GetMLDeformerComponent()->SetupComponent(DeformerAsset, SkelMeshComponent);
-				UMLDeformerModelInstance* ModelInstance = EditorActor->GetMLDeformerComponent()->GetModelInstance();
-				if (ModelInstance)
+				continue;
+			}
+
+			const int32 InstanceIndex = CompareEditorActor->GetActorTypeInstanceIndex();
+			const FMLDeformerCompareActor& CompareActor = Model->GetVizSettings()->GetCompareActors()[InstanceIndex];
+
+			if (CompareEditorActor->GetMLDeformerComponent())
+			{
+				USkeletalMeshComponent* SkelMeshComponent = CompareEditorActor->GetSkeletalMeshComponent();
+				UMLDeformerAsset* DeformerAsset = CompareActor.DeformerAsset.LoadSynchronous();
+				if (!IsCompatibleDeformer(DeformerAsset))
 				{
-					ModelInstance->UpdateCompatibilityStatus();
+					DeformerAsset = nullptr;
 				}
+				CompareEditorActor->GetMLDeformerComponent()->SetupComponent(DeformerAsset, SkelMeshComponent);
 			}
 		}
 	}
@@ -1582,26 +1819,45 @@ namespace UE::MLDeformer
 	}
 
 	void FMLDeformerEditorModel::UpdateDeformerGraph()
-	{	
-		for (FMLDeformerEditorActor* EditorActor : EditorActors)
+	{
+		UMLDeformerVizSettings* VizSettings = Model->GetVizSettings();
+
+		FMLDeformerEditorActor* EditorActor = FindEditorActor(ActorID_Test_MLDeformed);
+		if (EditorActor && EditorActor->GetMLDeformerComponent())
 		{
-			if (EditorActor == nullptr || EditorActor->GetMLDeformerComponent() == nullptr)
+			USkeletalMeshComponent* SkelMeshComponent = EditorActor->GetSkeletalMeshComponent();
+			if (SkelMeshComponent)
 			{
-				continue;
+				UMeshDeformer* MeshDeformer = IsTrained() ? VizSettings->GetDeformerGraph() : nullptr;	
+				const bool bUseHeatMapDeformer = VizSettings->GetShowHeatMap();
+				SkelMeshComponent->SetMeshDeformer(bUseHeatMapDeformer ? HeatMapDeformerGraph.Get() : MeshDeformer);
 			}
+		}
 
-			USkeletalMeshComponent* SkelMeshComponent = EditorActor->GetSkeletalMeshComponent();
-			if (SkelMeshComponent == nullptr)
+		// Overwrite the comparison model's deformer graph, as that might require a different one.
+		for (FMLDeformerEditorActor* CompareEditorActor : EditorActors)
+		{
+			if (CompareEditorActor->GetTypeID() != ActorID_Test_Compare)
 			{
 				continue;
 			}
 
-			if (SkelMeshComponent)
+			if (CompareEditorActor->GetMLDeformerComponent())
 			{
-				const UMLDeformerVizSettings* VizSettings = Model->GetVizSettings();
-				UMeshDeformer* MeshDeformer = IsTrained() ? VizSettings->GetDeformerGraph() : nullptr;	
-				const bool bUseHeatMapDeformer = VizSettings->GetShowHeatMap();
-				SkelMeshComponent->SetMeshDeformer(bUseHeatMapDeformer ? HeatMapDeformerGraph.Get() : MeshDeformer);
+				const int32 InstanceIndex = CompareEditorActor->GetActorTypeInstanceIndex();
+				const FMLDeformerCompareActor& CompareActor = VizSettings->GetCompareActors()[InstanceIndex];
+
+				USkeletalMeshComponent* SkelMeshComponent = CompareEditorActor->GetSkeletalMeshComponent();
+				if (SkelMeshComponent)
+				{
+					const UMLDeformerAsset* Asset = CompareActor.DeformerAsset.LoadSynchronous();
+					if (Asset && Asset->GetModel())
+					{
+						UMeshDeformer* MeshDeformer = Asset->GetModel()->GetVizSettings()->GetDeformerGraph();	
+						const bool bUseHeatMapDeformer = VizSettings->GetShowHeatMap();
+						SkelMeshComponent->SetMeshDeformer(bUseHeatMapDeformer ? HeatMapDeformerGraph.Get() : MeshDeformer);
+					}
+				}
 			}
 		}
 	}
@@ -2338,6 +2594,170 @@ namespace UE::MLDeformer
 		return nullptr;
 	}
 
+	void FMLDeformerEditorModel::AddCompareActor(int32 ArrayIndex)
+	{
+		// Renumber the instance index values of other editor compare actors.
+		for (FMLDeformerEditorActor* EditorActor : EditorActors)
+		{
+			if (EditorActor->GetTypeID() == ActorID_Test_Compare)
+			{
+				if (EditorActor->GetActorTypeInstanceIndex() >= ArrayIndex)
+				{
+					EditorActor->SetActorTypeInstanceIndex(EditorActor->GetActorTypeInstanceIndex() + 1);
+				}
+			}
+		}
+
+		const FLinearColor MLDeformedWireColor = FMLDeformerEditorStyle::Get().GetColor("MLDeformer.MLCompareMesh.WireframeColor");
+		USkeletalMesh* SkelMesh = Model->GetSkeletalMesh();
+
+		TArray<FMLDeformerCompareActor>& CompareActors = Model->GetVizSettings()->GetCompareActors();
+		FMLDeformerCompareActor& CompareActor = CompareActors[ArrayIndex];
+
+		// Create the ML deformed actor.
+		FActorSpawnParameters SpawnParams;
+		UWorld* World = GetWorld();
+		SpawnParams.Name = MakeUniqueObjectName(World, AActor::StaticClass(), "Test Compare Actor");
+		AActor* Actor = World->SpawnActor<AActor>(SpawnParams);
+		Actor->SetFlags(RF_Transient);
+
+		// Create the skeletal mesh component.
+		UDebugSkelMeshComponent* SkelMeshComponent = NewObject<UDebugSkelMeshComponent>(Actor);
+		UMLDeformerAsset* DeformerAsset = CompareActor.DeformerAsset.LoadSynchronous();
+		if (!IsCompatibleDeformer(DeformerAsset))
+		{
+			SkelMesh = nullptr;
+		}
+		SkelMeshComponent->SetSkeletalMesh(SkelMesh);
+		Actor->SetRootComponent(SkelMeshComponent);
+		SkelMeshComponent->RegisterComponent();
+		SkelMeshComponent->SetWireframeMeshOverlayColor(MLDeformedWireColor);
+		SkelMeshComponent->SetVisibility(false);
+		SkelMeshComponent->MarkRenderStateDirty();
+		if (SkelMeshComponent->GetAnimInstance())
+		{
+			SkelMeshComponent->GetAnimInstance()->GetRequiredBones().SetUseRAWData(true);
+		}
+
+		// Create the ML Deformer component.
+		UMLDeformerAsset* FinalDeformerAsset = IsCompatibleDeformer(DeformerAsset) ? DeformerAsset : nullptr;
+		UMLDeformerComponent* MLDeformerComponent = NewObject<UMLDeformerComponent>(Actor);
+		MLDeformerComponent->SetSuppressMeshDeformerLogWarnings(true);
+		MLDeformerComponent->SetDeformerAsset(FinalDeformerAsset);
+		MLDeformerComponent->RegisterComponent();
+		MLDeformerComponent->SetupComponent(FinalDeformerAsset, SkelMeshComponent);
+
+		// Create the editor actor.
+		const FLinearColor LabelColor = FMLDeformerEditorStyle::Get().GetColor("MLDeformer.MLCompareMesh.LabelColor");
+		FMLDeformerEditorActor::FConstructSettings Settings;
+		Settings.Actor = Actor;
+		Settings.TypeID = ActorID_Test_Compare;
+		Settings.LabelColor = LabelColor;
+		Settings.LabelText = LOCTEXT("TestCompareActorLabelText", "ML Compare");
+		Settings.bIsTrainingActor = false;
+		Settings.ActorTypeInstanceIndex = ArrayIndex;
+		FMLDeformerEditorActor* EditorActor = CreateEditorActor(Settings);
+		EditorActor->SetSkeletalMeshComponent(SkelMeshComponent);
+		EditorActor->SetMLDeformerComponent(MLDeformerComponent);
+		EditorActor->SetMeshOffsetFactor(2.0f + static_cast<float>(ArrayIndex));
+		EditorActors.Add(EditorActor);
+	}
+
+	void FMLDeformerEditorModel::RemoveCompareActor(int32 ArrayIndex)
+	{
+		UWorld* World = GetWorld();
+
+		FMLDeformerEditorActor* EditorActorToDestroy = nullptr;
+		AActor* ActorToDestroy = nullptr;
+		for (FMLDeformerEditorActor* EditorActor : EditorActors)
+		{
+			if (EditorActor && EditorActor->GetTypeID() == ActorID_Test_Compare && EditorActor->GetActorTypeInstanceIndex() == ArrayIndex)
+			{
+				EditorActorToDestroy = EditorActor;
+				ActorToDestroy = EditorActor->GetActor();
+				World->DestroyActor(EditorActor->GetActor(), true);
+				EditorActors.Remove(EditorActor);
+				delete EditorActor;
+				break;
+			}
+		}
+
+		// Renumber other editor actors that come after.
+		for (FMLDeformerEditorActor* EditorActor : EditorActors)
+		{
+			if (EditorActor && EditorActor->GetTypeID() == ActorID_Test_Compare && EditorActor->GetActorTypeInstanceIndex() > ArrayIndex)
+			{
+				EditorActor->SetActorTypeInstanceIndex(EditorActor->GetActorTypeInstanceIndex() - 1);
+			}
+		}
+
+		if (ActorToDestroy)
+		{
+			ActorToDestroy->Destroy();
+		}
+	}
+
+	void FMLDeformerEditorModel::RemoveAllCompareActors()
+	{
+		UWorld* World = GetWorld();
+
+		TArray<AActor*> ActorsToDestroy;
+		TArray<FMLDeformerEditorActor*> EditorActorsToDestroy;
+		for (FMLDeformerEditorActor* EditorActor : EditorActors)
+		{
+			if (EditorActor && EditorActor->GetTypeID() == ActorID_Test_Compare)
+			{				
+				ActorsToDestroy.Add(EditorActor->GetActor());
+				World->DestroyActor(EditorActor->GetActor(), true);
+				EditorActorsToDestroy.Add(EditorActor);
+			}
+		}
+
+		// Now delete the editor actors, which removes components.
+		for (FMLDeformerEditorActor* EditorActor : EditorActorsToDestroy)
+		{
+			EditorActors.Remove(EditorActor);
+			delete EditorActor;
+		}
+
+		for (AActor* Actor : ActorsToDestroy)
+		{
+			if (Actor)
+			{
+				Actor->Destroy();
+			}
+		}
+	}
+
+	void FMLDeformerEditorModel::UpdateCompareActorLabels()
+	{
+		for (FMLDeformerEditorActor* EditorActor : EditorActors)
+		{
+			if (EditorActor && EditorActor->GetTypeID() == ActorID_Test_Compare)
+			{
+				UTextRenderComponent* LabelComponent = EditorActor->GetLabelComponent();
+				if (LabelComponent)
+				{
+					const int32 InstanceIndex = EditorActor->GetActorTypeInstanceIndex();
+					const FMLDeformerCompareActor& CompareActor = Model->GetVizSettings()->GetCompareActors()[InstanceIndex];
+					FText LabelText;				
+					if (!CompareActor.Name.IsNone() && CompareActor.Name.IsValid())
+					{
+						LabelText = FText::FromName(CompareActor.Name);
+					}
+					else
+					{
+						UMLDeformerAsset* DeformerAsset = CompareActor.DeformerAsset.LoadSynchronous();
+						LabelText = DeformerAsset ? FText::FromString(DeformerAsset->GetName()) : FText::Format(LOCTEXT("DefaultCompareActorDeformerName", "ML Compare #{0}"), InstanceIndex);
+					}
+					const FLinearColor LabelColor = FMLDeformerEditorStyle::Get().GetColor("MLDeformer.MLCompareMesh.LabelColor");
+					LabelComponent->SetText(LabelText);
+					LabelComponent->SetTextRenderColor(LabelColor.ToFColor(true));
+				}
+			}
+		}
+	}
+
 	int32 FMLDeformerEditorModel::GetActiveTrainingInputAnimIndex() const
 	{
 		return ActiveTrainingInputAnimIndex;
diff --git a/Engine/Plugins/Animation/MLDeformer/MLDeformerFramework/Source/MLDeformerFrameworkEditor/Private/MLDeformerEditorStyle.cpp b/Engine/Plugins/Animation/MLDeformer/MLDeformerFramework/Source/MLDeformerFrameworkEditor/Private/MLDeformerEditorStyle.cpp
--- a/Engine/Plugins/Animation/MLDeformer/MLDeformerFramework/Source/MLDeformerFrameworkEditor/Private/MLDeformerEditorStyle.cpp
+++ b/Engine/Plugins/Animation/MLDeformer/MLDeformerFramework/Source/MLDeformerFrameworkEditor/Private/MLDeformerEditorStyle.cpp
@@ -38,6 +38,9 @@ namespace UE::MLDeformer
 		Set("MLDeformer.MLDeformedMesh.WireframeColor", FLinearColor(1.0f, 0.0f, 0.0f));
 		Set("MLDeformer.MLDeformedMesh.LabelColor", FLinearColor(1.0f, 0.0f, 0.0f));
 
+		Set("MLDeformer.MLCompareMesh.WireframeColor", FLinearColor(1.0f, 0.4f, 0.0f));
+		Set("MLDeformer.MLCompareMesh.LabelColor", FLinearColor(1.0f, 0.4f, 0.0f));
+
 		Set("MLDeformer.GroundTruth.WireframeColor", FLinearColor(0.0f, 0.0f, 1.0f));
 		Set("MLDeformer.GroundTruth.LabelColor", FLinearColor(0.0f, 0.0f, 1.0f));
 
diff --git a/Engine/Plugins/Animation/MLDeformer/MLDeformerFramework/Source/MLDeformerFrameworkEditor/Private/MLDeformerVizSettingsDetails.cpp b/Engine/Plugins/Animation/MLDeformer/MLDeformerFramework/Source/MLDeformerFrameworkEditor/Private/MLDeformerVizSettingsDetails.cpp
--- a/Engine/Plugins/Animation/MLDeformer/MLDeformerFramework/Source/MLDeformerFrameworkEditor/Private/MLDeformerVizSettingsDetails.cpp
+++ b/Engine/Plugins/Animation/MLDeformer/MLDeformerFramework/Source/MLDeformerFrameworkEditor/Private/MLDeformerVizSettingsDetails.cpp
@@ -46,7 +46,6 @@ namespace UE::MLDeformer
 	PRAGMA_ENABLE_DEPRECATION_WARNINGS
 	// END UE_DEPRECATED(5.4)
 
-
 	bool FMLDeformerVizSettingsDetails::UpdateMemberPointers(const TArray<TWeakObjectPtr<UObject>>& Objects)
 	{
 		Model = nullptr;
@@ -169,6 +168,9 @@ namespace UE::MLDeformer
 
 		AddGroundTruth();
 
+		IDetailPropertyRow& CompareAssetRow = TestAssetsCategory->AddProperty(UMLDeformerVizSettings::GetCompareActorsPropertyName(), UMLDeformerVizSettings::StaticClass())
+			.ShouldAutoExpand(true);
+
 		LiveSettingsCategory->SetCategoryVisibility(bShowTestData);
 		LiveSettingsCategory->AddProperty(UMLDeformerVizSettings::GetWeightPropertyName(), UMLDeformerVizSettings::StaticClass());
 		LiveSettingsCategory->AddProperty(UMLDeformerVizSettings::GetAnimPlaySpeedPropertyName(), UMLDeformerVizSettings::StaticClass());
@@ -189,6 +191,8 @@ namespace UE::MLDeformer
 		VisGroup.AddPropertyRow(DetailBuilder.GetProperty(UMLDeformerVizSettings::GetDrawMLDeformedActorPropertyName(), UMLDeformerVizSettings::StaticClass()));
 		VisGroup.AddPropertyRow(DetailBuilder.GetProperty(UMLDeformerVizSettings::GetDrawGroundTruthActorPropertyName(), UMLDeformerVizSettings::StaticClass()))
 			.Visibility(VizSettings->HasTestGroundTruth() ? EVisibility::Visible : EVisibility::Collapsed);
+		VisGroup.AddPropertyRow(DetailBuilder.GetProperty(UMLDeformerVizSettings::GetDrawMLCompareActorsPropertyName(), UMLDeformerVizSettings::StaticClass()))
+			.Visibility(!VizSettings->GetCompareActors().IsEmpty() ? EVisibility::Visible : EVisibility::Collapsed);
 
 		TrainingMeshesCategoryBuilder->SetCategoryVisibility(bShowTrainingData);
 		TrainingMeshesCategoryBuilder->AddProperty(UMLDeformerVizSettings::GetTrainingFrameNumberPropertyName(), UMLDeformerVizSettings::StaticClass());
diff --git a/Engine/Plugins/Animation/MLDeformer/MLDeformerFramework/Source/MLDeformerFrameworkEditor/Public/MLDeformerEditorActor.h b/Engine/Plugins/Animation/MLDeformer/MLDeformerFramework/Source/MLDeformerFrameworkEditor/Public/MLDeformerEditorActor.h
--- a/Engine/Plugins/Animation/MLDeformer/MLDeformerFramework/Source/MLDeformerFrameworkEditor/Public/MLDeformerEditorActor.h
+++ b/Engine/Plugins/Animation/MLDeformer/MLDeformerFramework/Source/MLDeformerFrameworkEditor/Public/MLDeformerEditorActor.h
@@ -32,7 +32,10 @@ namespace UE::MLDeformer
 		ActorID_Test_MLDeformed,
 
 		/** The ground truth test actor, in test mode. This will represent your test anim sequence ground truth. */
-		ActorID_Test_GroundTruth
+		ActorID_Test_GroundTruth,
+
+		/** The comparison actor, in test mode. This can be used to compare the results to another pretrained ML Deformer model. */
+		ActorID_Test_Compare
 	};
 
 	/**
@@ -55,6 +58,12 @@ namespace UE::MLDeformer
 			/** The ID of the actor. Look at the UE::MLDeformer::ActorID_Train_Base etc for example. */
 			int32 TypeID = -1;
 
+			/** 
+			 * The instance index. There can for example be multiple editor actors with a TypeID of ActorID_Test_Compare. 
+			 * This index separates those. If you have three test compare editor actors, on will have index 0, another 1, and the third one will have a value of 2 for this member.
+			 */
+			int32 ActorTypeInstanceIndex = 0;
+
 			/** The color of the label text. */
 			FLinearColor LabelColor = FLinearColor(1.0f, 0.0f, 0.0f);
 
@@ -82,18 +91,20 @@ namespace UE::MLDeformer
 		virtual bool IsGroundTruthActor() const;
 		virtual bool HasVisualMesh() const;
 
-		AActor* GetActor() const { return Actor; }
-		int32 GetTypeID() const { return TypeID; }
-
-		UDebugSkelMeshComponent* GetSkeletalMeshComponent() const { return SkeletalMeshComponent.Get(); }
-		UTextRenderComponent* GetLabelComponent() const { return LabelComponent.Get(); }
-		UMLDeformerComponent* GetMLDeformerComponent() const { return MLDeformerComponent.Get(); }
-		void SetSkeletalMeshComponent(UDebugSkelMeshComponent* SkelMeshComponent) { SkeletalMeshComponent = SkelMeshComponent; }
-		void SetMLDeformerComponent(UMLDeformerComponent* Component) { MLDeformerComponent = Component; }
-		void SetMeshOffsetFactor(float OffsetFactor) { MeshOffsetFactor = OffsetFactor; }
-		bool IsTrainingActor() const { return bIsTrainingActor; }
-		bool IsTestActor() const { return !bIsTrainingActor; }
-		float GetMeshOffsetFactor() const { return MeshOffsetFactor; }
+		AActor* GetActor() const													{ return Actor; }
+		int32 GetTypeID() const														{ return TypeID; }
+		int32 GetActorTypeInstanceIndex() const										{ return ActorTypeInstanceIndex; }
+		UDebugSkelMeshComponent* GetSkeletalMeshComponent() const					{ return SkeletalMeshComponent.Get(); }
+		UTextRenderComponent* GetLabelComponent() const								{ return LabelComponent.Get(); }
+		UMLDeformerComponent* GetMLDeformerComponent() const						{ return MLDeformerComponent.Get(); }
+		bool IsTrainingActor() const												{ return bIsTrainingActor; }
+		bool IsTestActor() const													{ return !bIsTrainingActor; }
+		float GetMeshOffsetFactor() const											{ return MeshOffsetFactor; }
+
+		void SetSkeletalMeshComponent(UDebugSkelMeshComponent* SkelMeshComponent)	{ SkeletalMeshComponent = SkelMeshComponent; }
+		void SetMLDeformerComponent(UMLDeformerComponent* Component)				{ MLDeformerComponent = Component; }
+		void SetMeshOffsetFactor(float OffsetFactor)								{ MeshOffsetFactor = OffsetFactor; }
+		void SetActorTypeInstanceIndex(int32 Index)									{ ActorTypeInstanceIndex = Index; }
 
 	protected:
 		UTextRenderComponent* CreateLabelComponent(AActor* InActor, FLinearColor Color, const FText& Text) const;
@@ -105,6 +116,12 @@ namespace UE::MLDeformer
 		 */
 		int32 TypeID = -1;
 
+		/** 
+ 		 * The instance index. There can for example be multiple editor actors with a TypeID of ActorID_Test_Compare. 
+		 * This index separates those. If you have three test compare editor actors, on will have index 0, another 1, and the third one will have a value of 2 for this member.
+		 */
+		int32 ActorTypeInstanceIndex = 0;
+
 		/** The label component, which shows above the actor. */
 		TObjectPtr<UTextRenderComponent> LabelComponent = nullptr;
 
diff --git a/Engine/Plugins/Animation/MLDeformer/MLDeformerFramework/Source/MLDeformerFrameworkEditor/Public/MLDeformerEditorModel.h b/Engine/Plugins/Animation/MLDeformer/MLDeformerFramework/Source/MLDeformerFrameworkEditor/Public/MLDeformerEditorModel.h
--- a/Engine/Plugins/Animation/MLDeformer/MLDeformerFramework/Source/MLDeformerFrameworkEditor/Public/MLDeformerEditorModel.h
+++ b/Engine/Plugins/Animation/MLDeformer/MLDeformerFramework/Source/MLDeformerFrameworkEditor/Public/MLDeformerEditorModel.h
@@ -282,6 +282,10 @@ namespace UE::MLDeformer
 		 */
 		virtual void CreateTestMLDeformedActor(UWorld* World);
 
+		virtual void CreateTestCompareActors(UWorld* World);
+		virtual bool IsCompatibleDeformer(UMLDeformerAsset* Deformer) const;
+		virtual void UpdateMeshOffsetFactors();
+
 		/**
 		 * Create the training ground truth actor.
 		 * This is the training target, so the actor that has the complex deformations, for example using a geometry cache.
@@ -572,7 +576,7 @@ namespace UE::MLDeformer
 		 * This is determined by looking whether the neural network pointer is nullptr or not.
 		 * @return Returns true when the model is trained already, or false if it hasn't been trained yet.
 		 */
-		virtual bool IsTrained() const { return false; }
+		virtual bool IsTrained() const { return Model->IsTrained(); }
 
 		/**
 		 * Get the editor actor that defines the timeline play position.
@@ -975,6 +979,13 @@ namespace UE::MLDeformer
 		 */
 		void UpdateRanges();
 
+		void AddCompareActor(int32 ArrayIndex);
+		void RemoveCompareActor(int32 ArrayIndex);
+		void RemoveAllCompareActors();
+		void UpdateCompareActorLabels();		
+
+		int32 CalcNumValidCompareActorsPriorTo(int32 CompareActorIndex) const;
+
 	protected:
 		/** The runtime model associated with this editor model. */
 		TObjectPtr<UMLDeformerModel> Model = nullptr;
diff --git a/Engine/Plugins/Animation/MLDeformer/NearestNeighborModel/Source/NearestNeighborModel/Private/NearestNeighborModel.cpp b/Engine/Plugins/Animation/MLDeformer/NearestNeighborModel/Source/NearestNeighborModel/Private/NearestNeighborModel.cpp
--- a/Engine/Plugins/Animation/MLDeformer/NearestNeighborModel/Source/NearestNeighborModel/Private/NearestNeighborModel.cpp
+++ b/Engine/Plugins/Animation/MLDeformer/NearestNeighborModel/Source/NearestNeighborModel/Private/NearestNeighborModel.cpp
@@ -55,6 +55,38 @@ UNearestNeighborModel::UNearestNeighborModel(const FObjectInitializer& ObjectIni
 #endif
 }
 
+void UNearestNeighborModel::GetAssetRegistryTags(TArray<FAssetRegistryTag>& OutTags) const
+{
+	Super::GetAssetRegistryTags(OutTags);
+
+	#if WITH_EDITORONLY_DATA
+		OutTags.Add(FAssetRegistryTag("MLDeformer.NearestNeighborModel.NumEpochs", FString::FromInt(NumEpochs), FAssetRegistryTag::TT_Numerical));
+		OutTags.Add(FAssetRegistryTag("MLDeformer.NearestNeighborModel.BatchSize", FString::FromInt(BatchSize), FAssetRegistryTag::TT_Numerical));
+		OutTags.Add(FAssetRegistryTag("MLDeformer.NearestNeighborModel.BasisSmoothIter", FString::FromInt(BasisSmoothIter), FAssetRegistryTag::TT_Numerical));
+		OutTags.Add(FAssetRegistryTag("MLDeformer.NearestNeighborModel.NumHiddenLayers", FString::FromInt(HiddenLayerDims.Num()), FAssetRegistryTag::TT_Numerical));
+		OutTags.Add(FAssetRegistryTag("MLDeformer.NearestNeighborModel.NumClusters", FString::FromInt(NumClusters), FAssetRegistryTag::TT_Numerical));
+		OutTags.Add(FAssetRegistryTag("MLDeformer.NearestNeighborModel.NumClusterAnims", FString::FromInt(SourceAnims.Num()), FAssetRegistryTag::TT_Numerical));
+		OutTags.Add(FAssetRegistryTag("MLDeformer.NearestNeighborModel.NumNearestNeighborAnims", FString::FromInt(NearestNeighborData.Num()), FAssetRegistryTag::TT_Numerical));
+		OutTags.Add(FAssetRegistryTag("MLDeformer.NearestNeighborModel.KMeansPartID", FString::FromInt(KMeansPartId), FAssetRegistryTag::TT_Numerical));
+		OutTags.Add(FAssetRegistryTag("MLDeformer.NearestNeighborModel.LearningRate", FString::Printf(TEXT("%f"), LearningRate), FAssetRegistryTag::TT_Numerical));
+		OutTags.Add(FAssetRegistryTag("MLDeformer.NearestNeighborModel.IsPartDataValid", bClothPartDataValid ? TEXT("True") : TEXT("False"), FAssetRegistryTag::TT_Alphabetical));
+		OutTags.Add(FAssetRegistryTag("MLDeformer.NearestNeighborModel.UseDualQuaternions", bUseDualQuaternionDeltas ? TEXT("True") : TEXT("False"), FAssetRegistryTag::TT_Alphabetical));
+		OutTags.Add(FAssetRegistryTag("MLDeformer.NearestNeighborModel.IsNearestNeighborDataValid", bNearestNeighborDataValid ? TEXT("True") : TEXT("False"), FAssetRegistryTag::TT_Alphabetical));
+		OutTags.Add(FAssetRegistryTag("MLDeformer.NearestNeighborModel.IsMorphTargetDataValid", bMorphTargetDataValid ? TEXT("True") : TEXT("False"), FAssetRegistryTag::TT_Alphabetical));
+		OutTags.Add(FAssetRegistryTag("MLDeformer.NearestNeighborModel.UseFileCache", bUseFileCache ? TEXT("True") : TEXT("False"), FAssetRegistryTag::TT_Alphabetical));
+		OutTags.Add(FAssetRegistryTag("MLDeformer.NearestNeighborModel.RecomputDeltas", bRecomputeDeltas ? TEXT("True") : TEXT("False"), FAssetRegistryTag::TT_Alphabetical));
+		OutTags.Add(FAssetRegistryTag("MLDeformer.NearestNeighborModel.RecomputPCA", bRecomputePCA ? TEXT("True") : TEXT("False"), FAssetRegistryTag::TT_Alphabetical));
+	#endif
+
+	OutTags.Add(FAssetRegistryTag("MLDeformer.NearestNeighborModel.UseInputMultipliers", bUseInputMultipliers ? TEXT("True") : TEXT("False"), FAssetRegistryTag::TT_Alphabetical));
+	OutTags.Add(FAssetRegistryTag("MLDeformer.NearestNeighborModel.UseRBF", bUseRBF ? TEXT("True") : TEXT("False"), FAssetRegistryTag::TT_Alphabetical));
+	OutTags.Add(FAssetRegistryTag("MLDeformer.NearestNeighborModel.NumParts", FString::FromInt(ClothPartData.Num()), FAssetRegistryTag::TT_Numerical));
+	OutTags.Add(FAssetRegistryTag("MLDeformer.NearestNeighborModel.DecayFactor", FString::Printf(TEXT("%f"), DecayFactor), FAssetRegistryTag::TT_Numerical));
+	OutTags.Add(FAssetRegistryTag("MLDeformer.NearestNeighborModel.NearestNeighborOffsetWeight", FString::Printf(TEXT("%f"), NearestNeighborOffsetWeight), FAssetRegistryTag::TT_Numerical));
+	OutTags.Add(FAssetRegistryTag("MLDeformer.NearestNeighborModel.RBFSigma", FString::Printf(TEXT("%f"), RBFSigma), FAssetRegistryTag::TT_Numerical));
+	OutTags.Add(FAssetRegistryTag("MLDeformer.NearestNeighborModel.NumInputMultipliers", FString::FromInt(InputMultipliers.Num()), FAssetRegistryTag::TT_Numerical));
+}
+
 UMLDeformerInputInfo* UNearestNeighborModel::CreateInputInfo()
 {
 	UNearestNeighborModelInputInfo* NearestNeighborModelInputInfo = NewObject<UNearestNeighborModelInputInfo>(this);
@@ -64,25 +96,10 @@ UMLDeformerInputInfo* UNearestNeighborModel::CreateInputInfo()
 	return NearestNeighborModelInputInfo;
 }
 
-#if WITH_EDITOR
-void UNearestNeighborModel::UpdateMemoryUsage()
+bool UNearestNeighborModel::IsTrained() const
 {
-	Super::UpdateMemoryUsage();
-
-	uint64 ClothDataSize = 0;
-	for (const FClothPartData& Part : ClothPartData)
-	{
-		ClothDataSize += Part.PCABasis.GetTypeSize() * Part.PCABasis.Num();
-		ClothDataSize += Part.AssetNeighborOffsets.GetTypeSize() * Part.AssetNeighborOffsets.Num();
-		ClothDataSize += Part.AssetNeighborCoeffs.GetTypeSize() * Part.AssetNeighborCoeffs.Num();
-		ClothDataSize += Part.VertexMean.GetTypeSize() * Part.VertexMean.Num();
-		ClothDataSize += Part.VertexMap.GetTypeSize() * Part.VertexMap.Num();
-	}
-
-	CookedAssetSizeInBytes -= ClothDataSize;
-	MemUsageInBytes -= ClothDataSize;
+	return DoesUseOptimizedNetwork() && GetOptimizedNetwork() != nullptr;
 }
-#endif
 
 UMLDeformerModelInstance* UNearestNeighborModel::CreateModelInstance(UMLDeformerComponent* Component)
 {
diff --git a/Engine/Plugins/Animation/MLDeformer/NearestNeighborModel/Source/NearestNeighborModel/Public/NearestNeighborModel.h b/Engine/Plugins/Animation/MLDeformer/NearestNeighborModel/Source/NearestNeighborModel/Public/NearestNeighborModel.h
--- a/Engine/Plugins/Animation/MLDeformer/NearestNeighborModel/Source/NearestNeighborModel/Public/NearestNeighborModel.h
+++ b/Engine/Plugins/Animation/MLDeformer/NearestNeighborModel/Source/NearestNeighborModel/Public/NearestNeighborModel.h
@@ -176,6 +176,10 @@ class NEARESTNEIGHBORMODEL_API UNearestNeighborModel
 public:
 	UNearestNeighborModel(const FObjectInitializer& ObjectInitializer);
 
+	// UObject overrides.
+	virtual void GetAssetRegistryTags(TArray<FAssetRegistryTag>& OutTags) const override;
+	// ~END UObject overrides.
+
 	// UMLDeformerModel overrides.
 	virtual void PostLoad() override;
 	virtual void Serialize(FArchive& Archive) override;
@@ -183,9 +187,7 @@ public:
 	virtual UMLDeformerModelInstance* CreateModelInstance(UMLDeformerComponent* Component) override;
 	virtual UMLDeformerInputInfo* CreateInputInfo() override;
 	virtual FString GetDisplayName() const override { return "Nearest Neighbor Model"; }
-#if WITH_EDITOR
-	virtual void UpdateMemoryUsage() override;
-#endif
+	virtual bool IsTrained() const override;
 	// ~END UMLDeformerModel overrides.
 
 	friend class UNearestNeighborModelInstance;
diff --git a/Engine/Plugins/Animation/MLDeformer/NearestNeighborModel/Source/NearestNeighborModelEditor/Private/NearestNeighborEditorModel.cpp b/Engine/Plugins/Animation/MLDeformer/NearestNeighborModel/Source/NearestNeighborModelEditor/Private/NearestNeighborEditorModel.cpp
--- a/Engine/Plugins/Animation/MLDeformer/NearestNeighborModel/Source/NearestNeighborModelEditor/Private/NearestNeighborEditorModel.cpp
+++ b/Engine/Plugins/Animation/MLDeformer/NearestNeighborModel/Source/NearestNeighborModelEditor/Private/NearestNeighborEditorModel.cpp
@@ -116,16 +116,6 @@ namespace UE::NearestNeighborModel
 		return false;
 	}
 
-	bool FNearestNeighborEditorModel::IsTrained() const
-	{
-		const UNearestNeighborModel* NearestNeighborModel = GetNearestNeighborModel();
-		if (NearestNeighborModel)
-		{
-			return NearestNeighborModel->DoesUseOptimizedNetwork() && NearestNeighborModel->GetOptimizedNetwork() != nullptr;
-		}
-		return false;
-	}
-
 	FString FNearestNeighborEditorModel::GetTrainedNetworkOnnxFile() const
 	{
 		const UNearestNeighborModel* NearestNeighborModel = GetNearestNeighborModel();
@@ -230,7 +220,7 @@ namespace UE::NearestNeighborModel
 			FNearestNeighborEditorModelActor* NearestNeighborActor = static_cast<FNearestNeighborEditorModelActor*>(EditorActors.Last());
 			UMLDeformerComponent* MLDeformerComponent = GetTestMLDeformerComponent();
 			NearestNeighborActor->InitNearestNeighborActor(PartId, MLDeformerComponent);
-			NearestNeighborActor->SetMeshOffsetFactor(Offset);
+			NearestNeighborActor->SetMeshOffsetFactor(Offset * 2.0f);
 			NearestNeighborActors[PartId] = NearestNeighborActor;
 		}
 	}
@@ -390,7 +380,7 @@ namespace UE::NearestNeighborModel
 					FNearestNeighborEditorModelActor* EditorActor = NearestNeighborActors[i];
 					NearestNeighborActors.RemoveAt(i);
 					EditorActors.Remove(EditorActor);
-					EditorWorld->RemoveActor(EditorActor->GetActor(), true/*ShouldModifyLevel*/);
+					EditorWorld->DestroyActor(EditorActor->GetActor(), true/*ShouldModifyLevel*/);
 					delete EditorActor;
 				}
 			}
diff --git a/Engine/Plugins/Animation/MLDeformer/NearestNeighborModel/Source/NearestNeighborModelEditor/Public/NearestNeighborEditorModel.h b/Engine/Plugins/Animation/MLDeformer/NearestNeighborModel/Source/NearestNeighborModelEditor/Public/NearestNeighborEditorModel.h
--- a/Engine/Plugins/Animation/MLDeformer/NearestNeighborModel/Source/NearestNeighborModelEditor/Public/NearestNeighborEditorModel.h
+++ b/Engine/Plugins/Animation/MLDeformer/NearestNeighborModel/Source/NearestNeighborModelEditor/Public/NearestNeighborEditorModel.h
@@ -45,7 +45,6 @@ namespace UE::NearestNeighborModel
 
 		virtual void InitInputInfo(UMLDeformerInputInfo* InputInfo) override;
 		virtual ETrainingResult Train() override;
-		virtual bool IsTrained() const override;
 		virtual bool LoadTrainedNetwork() const override;
 		virtual void OnPropertyChanged(FPropertyChangedEvent& PropertyChangedEvent) override;
 		virtual void OnPostTraining(ETrainingResult TrainingResult, bool bUsePartiallyTrainedWhenAborted) override;
diff --git a/Engine/Plugins/Animation/MLDeformer/NeuralMorphModel/Source/NeuralMorphModel/Private/NeuralMorphModel.cpp b/Engine/Plugins/Animation/MLDeformer/NeuralMorphModel/Source/NeuralMorphModel/Private/NeuralMorphModel.cpp
--- a/Engine/Plugins/Animation/MLDeformer/NeuralMorphModel/Source/NeuralMorphModel/Private/NeuralMorphModel.cpp
+++ b/Engine/Plugins/Animation/MLDeformer/NeuralMorphModel/Source/NeuralMorphModel/Private/NeuralMorphModel.cpp
@@ -127,4 +127,35 @@ int32 UNeuralMorphModel::GetNumFloatsPerCurve() const
 	return (Mode == ENeuralMorphMode::Local) ? 6 : 1;
 }
 
+bool UNeuralMorphModel::IsTrained() const
+{
+	return (NeuralMorphNetwork && NeuralMorphNetwork->GetMainModel() != nullptr);
+}
+
+void UNeuralMorphModel::GetAssetRegistryTags(TArray<FAssetRegistryTag>& OutTags) const
+{
+	Super::GetAssetRegistryTags(OutTags);
+
+	const UNeuralMorphInputInfo* NeuralInputInfo = Cast<UNeuralMorphInputInfo>(GetInputInfo());
+	if (NeuralInputInfo)
+	{
+		OutTags.Add(FAssetRegistryTag("MLDeformer.NeuralMorphModel.Trained.NumBoneGroups", FString::FromInt(NeuralInputInfo->GetBoneGroups().Num()), FAssetRegistryTag::TT_Numerical));
+		OutTags.Add(FAssetRegistryTag("MLDeformer.NeuralMorphModel.Trained.NumCurveGroups", FString::FromInt(NeuralInputInfo->GetCurveGroups().Num()), FAssetRegistryTag::TT_Numerical));
+	}
+
+	OutTags.Add(FAssetRegistryTag("MLDeformer.NeuralMorphModel.Mode", Mode == ENeuralMorphMode::Local ? TEXT("Local") : TEXT("Global"), FAssetRegistryTag::TT_Alphabetical));
+	OutTags.Add(FAssetRegistryTag("MLDeformer.NeuralMorphModel.NumMorphsPerBone", FString::FromInt(LocalNumMorphTargetsPerBone), FAssetRegistryTag::TT_Numerical));
+	OutTags.Add(FAssetRegistryTag("MLDeformer.NeuralMorphModel.NumMorphTargets", FString::FromInt(GlobalNumMorphTargets), FAssetRegistryTag::TT_Numerical));
+	OutTags.Add(FAssetRegistryTag("MLDeformer.NeuralMorphModel.NumIterations", FString::FromInt(NumIterations), FAssetRegistryTag::TT_Numerical));
+	OutTags.Add(FAssetRegistryTag("MLDeformer.NeuralMorphModel.LocalNumHiddenLayers", FString::FromInt(LocalNumHiddenLayers), FAssetRegistryTag::TT_Numerical));
+	OutTags.Add(FAssetRegistryTag("MLDeformer.NeuralMorphModel.LocalNumNeuronsPerLayer", FString::FromInt(LocalNumNeuronsPerLayer), FAssetRegistryTag::TT_Numerical));
+	OutTags.Add(FAssetRegistryTag("MLDeformer.NeuralMorphModel.GlobalNumHiddenLayers", FString::FromInt(GlobalNumHiddenLayers), FAssetRegistryTag::TT_Numerical));
+	OutTags.Add(FAssetRegistryTag("MLDeformer.NeuralMorphModel.GlobalNumNeuronsPerLayer", FString::FromInt(GlobalNumNeuronsPerLayer), FAssetRegistryTag::TT_Numerical));
+	OutTags.Add(FAssetRegistryTag("MLDeformer.NeuralMorphModel.BatchSize", FString::FromInt(BatchSize), FAssetRegistryTag::TT_Numerical));
+	OutTags.Add(FAssetRegistryTag("MLDeformer.NeuralMorphModel.LearningRate", FString::Printf(TEXT("%f"), LearningRate), FAssetRegistryTag::TT_Numerical));
+	OutTags.Add(FAssetRegistryTag("MLDeformer.NeuralMorphModel.RegularizationFactor", FString::Printf(TEXT("%f"), RegularizationFactor), FAssetRegistryTag::TT_Numerical));
+	OutTags.Add(FAssetRegistryTag("MLDeformer.NeuralMorphModel.SmoothLossBeta", FString::Printf(TEXT("%f"), SmoothLossBeta), FAssetRegistryTag::TT_Numerical));
+	OutTags.Add(FAssetRegistryTag("MLDeformer.NeuralMorphModel.EnableBoneMasks", bEnableBoneMasks ? TEXT("True") : TEXT("False"), FAssetRegistryTag::TT_Alphabetical));
+}
+
 #undef LOCTEXT_NAMESPACE
diff --git a/Engine/Plugins/Animation/MLDeformer/NeuralMorphModel/Source/NeuralMorphModel/Public/NeuralMorphModel.h b/Engine/Plugins/Animation/MLDeformer/NeuralMorphModel/Source/NeuralMorphModel/Public/NeuralMorphModel.h
--- a/Engine/Plugins/Animation/MLDeformer/NeuralMorphModel/Source/NeuralMorphModel/Public/NeuralMorphModel.h
+++ b/Engine/Plugins/Animation/MLDeformer/NeuralMorphModel/Source/NeuralMorphModel/Public/NeuralMorphModel.h
@@ -42,6 +42,7 @@ public:
 	// UObject overrides.
 	virtual void Serialize(FArchive& Archive) override;
 	virtual void PostLoad() override;
+	virtual void GetAssetRegistryTags(TArray<FAssetRegistryTag>& OutTags) const override;
 	// ~END UObject overrides.
 
 	// UMLDeformerModel overrides.
@@ -49,6 +50,7 @@ public:
 	virtual UMLDeformerModelInstance* CreateModelInstance(UMLDeformerComponent* Component) override;
 	virtual UMLDeformerInputInfo* CreateInputInfo() override;
 	virtual int32 GetNumFloatsPerCurve() const override;
+	virtual bool IsTrained() const override;
 	// ~END UMLDeformerModel overrides.
 
 	const TArray<FNeuralMorphBoneGroup>& GetBoneGroups() const		{ return BoneGroups; }
@@ -71,7 +73,6 @@ public:
 
 	void UpdateMissingGroupNames();
 	void SetNeuralMorphNetwork(UNeuralMorphNetwork* Net);
-
 	void SetNumIterations(int32 InNumIterations)					{ check(InNumIterations > 0); NumIterations = InNumIterations; }
 
 public:
diff --git a/Engine/Plugins/Animation/MLDeformer/NeuralMorphModel/Source/NeuralMorphModelEditor/Private/NeuralMorphEditorModel.cpp b/Engine/Plugins/Animation/MLDeformer/NeuralMorphModel/Source/NeuralMorphModelEditor/Private/NeuralMorphEditorModel.cpp
--- a/Engine/Plugins/Animation/MLDeformer/NeuralMorphModel/Source/NeuralMorphModelEditor/Private/NeuralMorphEditorModel.cpp
+++ b/Engine/Plugins/Animation/MLDeformer/NeuralMorphModel/Source/NeuralMorphModelEditor/Private/NeuralMorphEditorModel.cpp
@@ -90,17 +90,6 @@ namespace UE::NeuralMorphModel
 			.EditorModel(this);
 	}
 
-	bool FNeuralMorphEditorModel::IsTrained() const
-	{
-		UNeuralMorphNetwork* MorphNetwork = GetNeuralMorphModel()->GetNeuralMorphNetwork();
-		if (MorphNetwork != nullptr && MorphNetwork->GetMainModel() != nullptr)
-		{
-			return true;
-		}
-
-		return false;
-	}
-
 	void FNeuralMorphEditorModel::Init(const InitSettings& Settings)
 	{
 		FMLDeformerMorphModelEditorModel::Init(Settings);
diff --git a/Engine/Plugins/Animation/MLDeformer/NeuralMorphModel/Source/NeuralMorphModelEditor/Public/NeuralMorphEditorModel.h b/Engine/Plugins/Animation/MLDeformer/NeuralMorphModel/Source/NeuralMorphModelEditor/Public/NeuralMorphEditorModel.h
--- a/Engine/Plugins/Animation/MLDeformer/NeuralMorphModel/Source/NeuralMorphModelEditor/Public/NeuralMorphEditorModel.h
+++ b/Engine/Plugins/Animation/MLDeformer/NeuralMorphModel/Source/NeuralMorphModelEditor/Public/NeuralMorphEditorModel.h
@@ -36,7 +36,6 @@ namespace UE::NeuralMorphModel
 		// FMLDeformerEditorModel overrides.
 		virtual void Init(const InitSettings& Settings) override;
 		virtual ETrainingResult Train() override;
-		virtual bool IsTrained() const override;
 		virtual void InitInputInfo(UMLDeformerInputInfo* InputInfo) override;
 		virtual FText GetOverlayText() const override;
 		virtual void OnPropertyChanged(FPropertyChangedEvent& PropertyChangedEvent) override;
diff --git a/Engine/Plugins/Animation/MLDeformer/VertexDeltaModel/Source/VertexDeltaModel/Private/VertexDeltaModel.cpp b/Engine/Plugins/Animation/MLDeformer/VertexDeltaModel/Source/VertexDeltaModel/Private/VertexDeltaModel.cpp
--- a/Engine/Plugins/Animation/MLDeformer/VertexDeltaModel/Source/VertexDeltaModel/Private/VertexDeltaModel.cpp
+++ b/Engine/Plugins/Animation/MLDeformer/VertexDeltaModel/Source/VertexDeltaModel/Private/VertexDeltaModel.cpp
@@ -57,9 +57,27 @@ FString UVertexDeltaModel::GetDefaultDeformerGraphAssetPath() const
 	return FString(TEXT("/VertexDeltaModel/Deformers/DG_VertexDeltaModel.DG_VertexDeltaModel"));
 }
 
+bool UVertexDeltaModel::IsTrained() const
+{
+	return NNEModel.Get() != nullptr;
+}
+
 void UVertexDeltaModel::PostLoad()
 {
 	Super::PostLoad();
 }
 
+void UVertexDeltaModel::GetAssetRegistryTags(TArray<FAssetRegistryTag>& OutTags) const
+{
+	Super::GetAssetRegistryTags(OutTags);
+
+	#if WITH_EDITORONLY_DATA
+		OutTags.Add(FAssetRegistryTag("MLDeformer.VertexDeltaModel.NumHiddenLayers", FString::FromInt(NumHiddenLayers), FAssetRegistryTag::TT_Numerical));
+		OutTags.Add(FAssetRegistryTag("MLDeformer.VertexDeltaModel.NumNeuronsPerLayer", FString::FromInt(NumNeuronsPerLayer), FAssetRegistryTag::TT_Numerical));
+		OutTags.Add(FAssetRegistryTag("MLDeformer.VertexDeltaModel.NumIterations", FString::FromInt(NumIterations), FAssetRegistryTag::TT_Numerical));
+		OutTags.Add(FAssetRegistryTag("MLDeformer.VertexDeltaModel.BatchSize", FString::FromInt(BatchSize), FAssetRegistryTag::TT_Numerical));
+		OutTags.Add(FAssetRegistryTag("MLDeformer.VertexDeltaModel.LearningRate", FString::Printf(TEXT("%f"), LearningRate), FAssetRegistryTag::TT_Numerical));
+	#endif
+}
+
 #undef LOCTEXT_NAMESPACE
diff --git a/Engine/Plugins/Animation/MLDeformer/VertexDeltaModel/Source/VertexDeltaModel/Public/VertexDeltaModel.h b/Engine/Plugins/Animation/MLDeformer/VertexDeltaModel/Source/VertexDeltaModel/Public/VertexDeltaModel.h
--- a/Engine/Plugins/Animation/MLDeformer/VertexDeltaModel/Source/VertexDeltaModel/Public/VertexDeltaModel.h
+++ b/Engine/Plugins/Animation/MLDeformer/VertexDeltaModel/Source/VertexDeltaModel/Public/VertexDeltaModel.h
@@ -29,7 +29,8 @@ public:
 	UVertexDeltaModel(const FObjectInitializer& ObjectInitializer);
 
 	// UObject overrides.
-	void PostLoad() override;
+	virtual void PostLoad() override;
+	virtual void GetAssetRegistryTags(TArray<FAssetRegistryTag>& OutTags) const override;
 	// ~END UObject overrides.
 
 	// UMLDeformerModel overrides.
@@ -37,6 +38,7 @@ public:
 	virtual UMLDeformerModelInstance* CreateModelInstance(UMLDeformerComponent* Component) override;
 	virtual bool IsNeuralNetworkOnGPU() const override		{ return true; }	// GPU neural network.
 	virtual FString GetDefaultDeformerGraphAssetPath() const override;
+	virtual bool IsTrained() const override;
 	// ~END UMLDeformerModel overrides.
 
 #if WITH_EDITORONLY_DATA
diff --git a/Engine/Plugins/Animation/MLDeformer/VertexDeltaModel/Source/VertexDeltaModelEditor/Private/VertexDeltaEditorModel.cpp b/Engine/Plugins/Animation/MLDeformer/VertexDeltaModel/Source/VertexDeltaModelEditor/Private/VertexDeltaEditorModel.cpp
--- a/Engine/Plugins/Animation/MLDeformer/VertexDeltaModel/Source/VertexDeltaModelEditor/Private/VertexDeltaEditorModel.cpp
+++ b/Engine/Plugins/Animation/MLDeformer/VertexDeltaModel/Source/VertexDeltaModelEditor/Private/VertexDeltaEditorModel.cpp
@@ -75,11 +75,6 @@ namespace UE::VertexDeltaModel
 		return false;
 	}
 
-	bool FVertexDeltaEditorModel::IsTrained() const
-	{
-		return GetVertexDeltaModel()->NNEModel != nullptr;
-	}
-
 	FString FVertexDeltaEditorModel::GetHeatMapDeformerGraphPath() const
 	{
 		return FString(TEXT("/VertexDeltaModel/Deformers/DG_VertexDeltaModel_HeatMap.DG_VertexDeltaModel_HeatMap"));
diff --git a/Engine/Plugins/Animation/MLDeformer/VertexDeltaModel/Source/VertexDeltaModelEditor/Public/VertexDeltaEditorModel.h b/Engine/Plugins/Animation/MLDeformer/VertexDeltaModel/Source/VertexDeltaModelEditor/Public/VertexDeltaEditorModel.h
--- a/Engine/Plugins/Animation/MLDeformer/VertexDeltaModel/Source/VertexDeltaModelEditor/Public/VertexDeltaEditorModel.h
+++ b/Engine/Plugins/Animation/MLDeformer/VertexDeltaModel/Source/VertexDeltaModelEditor/Public/VertexDeltaEditorModel.h
@@ -30,7 +30,6 @@ namespace UE::VertexDeltaModel
 		// FMLDeformerEditorModel overrides.
 		virtual ETrainingResult Train() override;
 		virtual bool LoadTrainedNetwork() const override;
-		virtual bool IsTrained() const override;
 		virtual FString GetHeatMapDeformerGraphPath() const override;
 		// ~END FMLDeformerEditorModel overrides.
 
