[MLDeformer] Fixed a crash when enabling ML Deformer Framework plugin, but no ML Deformer model plugins, and creating a new ML Deformer asset.

diff --git a/Engine/Plugins/Animation/MLDeformer/MLDeformerFramework/Source/MLDeformerFrameworkEditor/Private/MLDeformerAssetFactory.cpp b/Engine/Plugins/Animation/MLDeformer/MLDeformerFramework/Source/MLDeformerFrameworkEditor/Private/MLDeformerAssetFactory.cpp
--- a/Engine/Plugins/Animation/MLDeformer/MLDeformerFramework/Source/MLDeformerFrameworkEditor/Private/MLDeformerAssetFactory.cpp
+++ b/Engine/Plugins/Animation/MLDeformer/MLDeformerFramework/Source/MLDeformerFrameworkEditor/Private/MLDeformerAssetFactory.cpp
@@ -27,12 +27,19 @@ UObject* UMLDeformerFactory::FactoryCreateNew(UClass* Class, UObject* InParent,
 		FMLDeformerEditorModule& EditorModule = FModuleManager::GetModuleChecked<FMLDeformerEditorModule>("MLDeformerFrameworkEditor");
 		FMLDeformerEditorModelRegistry& ModelRegistry = EditorModule.GetModelRegistry();
 
-		// Create the highest priority ML model on default and use that in our asset.
-		const int32 HighestPriorityIndex = ModelRegistry.GetHighestPriorityModelIndex();
+		// If we have no models registered we can early out as we cannot create a model.
 		TArray<UClass*> ModelTypes;
 		ModelRegistry.GetRegisteredModels().GenerateKeyArray(ModelTypes);
+		if (ModelTypes.IsEmpty())
+		{
+			return DeformerAsset;
+		}
+
+		// Create the highest priority ML model on default and use that in our asset.
+		const int32 HighestPriorityIndex = ModelRegistry.GetHighestPriorityModelIndex();
 		TObjectPtr<UMLDeformerModel> Model = NewObject<UMLDeformerModel>(DeformerAsset, ModelTypes[HighestPriorityIndex]);
 		check(Model);
+
 		Model->Init(DeformerAsset);
 		DeformerAsset->SetModel(Model);
 	}
