[MLDeformer] Added PIE debugging support, and Rewind Debugger integration of the external morph targets.

diff --git a/Engine/Plugins/Animation/GameplayInsights/Source/GameplayInsights/Private/AnimationAnalyzer.cpp b/Engine/Plugins/Animation/GameplayInsights/Source/GameplayInsights/Private/AnimationAnalyzer.cpp
--- a/Engine/Plugins/Animation/GameplayInsights/Source/GameplayInsights/Private/AnimationAnalyzer.cpp
+++ b/Engine/Plugins/Animation/GameplayInsights/Source/GameplayInsights/Private/AnimationAnalyzer.cpp
@@ -7,6 +7,7 @@
 #include "HAL/LowLevelMemTracker.h"
 #include "TraceServices/Model/AnalysisSession.h"
 #include "TraceServices/Utils.h"
+#include "Animation/AnimTrace.h"
 
 FAnimationAnalyzer::FAnimationAnalyzer(TraceServices::IAnalysisSession& InSession, FAnimationProvider& InAnimationProvider)
 	: Session(InSession)
@@ -25,6 +26,7 @@ void FAnimationAnalyzer::OnAnalysisBegin(const FOnAnalysisContext& Context)
 	Builder.RouteEvent(RouteId_SkeletalMeshComponent, "Animation", "SkeletalMeshComponent");
 	Builder.RouteEvent(RouteId_SkeletalMeshComponent2, "Animation", "SkeletalMeshComponent2");
 	Builder.RouteEvent(RouteId_SkeletalMeshComponent3, "Animation", "SkeletalMeshComponent3");
+	Builder.RouteEvent(RouteId_SkeletalMeshComponent4, "Animation", "SkeletalMeshComponent4");
 	Builder.RouteEvent(RouteId_SkeletalMeshFrame, "Animation", "SkeletalMeshFrame");
 	Builder.RouteEvent(RouteId_AnimGraph, "Animation", "AnimGraph");
 	Builder.RouteEvent(RouteId_AnimNodeStart, "Animation", "AnimNodeStart");
@@ -176,6 +178,40 @@ bool FAnimationAnalyzer::OnEvent(uint16 RouteId, EStyle Style, const FOnEventCon
 		AnimationProvider.AppendSkeletalMeshComponent(ComponentId, MeshId, Context.EventTime.AsSeconds(Cycle), RecordingTime, LodIndex, FrameCounter, ComponentToWorldFloatArray, PoseFloatArray, CurveIds, CurveValues);
 		break;
 	}
+	case RouteId_SkeletalMeshComponent4:
+	{
+		uint64 Cycle = EventData.GetValue<uint64>("Cycle");
+		double RecordingTime = EventData.GetValue<double>("RecordingTime");
+		
+		uint64 ComponentId = EventData.GetValue<uint64>("ComponentId");
+		uint64 MeshId = EventData.GetValue<uint64>("MeshId");
+		uint32 BoneCount = EventData.GetValue<uint32>("BoneCount");
+		uint32 CurveCount = EventData.GetValue<uint32>("CurveCount");
+		uint16 FrameCounter = EventData.GetValue<uint16>("FrameCounter");
+		uint16 LodIndex = EventData.GetValue<uint16>("LodIndex");
+
+		TArrayView<const float> ComponentToWorldFloatArray = EventData.GetArrayView<float>("ComponentToWorld");
+		TArrayView<const float> PoseFloatArray = EventData.GetArrayView<float>("Pose");
+		TArrayView<const uint32> CurveIds = EventData.GetArrayView<uint32>("CurveIds");
+		TArrayView<const float> CurveValues = EventData.GetArrayView<float>("CurveValues");
+		check(CurveIds.Num() == CurveValues.Num());
+
+		const bool bRecordExternalMorphs = CVarRecordExternalMorphTargets->GetBool();
+		AnimationProvider.AppendSkeletalMeshComponent(
+			ComponentId, 
+			MeshId,
+			Context.EventTime.AsSeconds(Cycle),
+			RecordingTime,
+			LodIndex,
+			FrameCounter,
+			ComponentToWorldFloatArray,
+			PoseFloatArray,
+			CurveIds,
+			CurveValues,
+			bRecordExternalMorphs ? EventData.GetArrayView<float>("ExternalMorphSetWeights") : TArrayView<const float>(),
+			bRecordExternalMorphs ? EventData.GetArrayView<int32>("ExternalMorphSetWeightCounts") : TArrayView<const int32>());
+		break;
+	}
 	case RouteId_Name:
 	{
 		uint32 Id = EventData.GetValue<uint32>("Id");
diff --git a/Engine/Plugins/Animation/GameplayInsights/Source/GameplayInsights/Private/AnimationAnalyzer.h b/Engine/Plugins/Animation/GameplayInsights/Source/GameplayInsights/Private/AnimationAnalyzer.h
--- a/Engine/Plugins/Animation/GameplayInsights/Source/GameplayInsights/Private/AnimationAnalyzer.h
+++ b/Engine/Plugins/Animation/GameplayInsights/Source/GameplayInsights/Private/AnimationAnalyzer.h
@@ -3,6 +3,7 @@
 #pragma once
 
 #include "Trace/Analyzer.h"
+#include "HAL/IConsoleManager.h"
 
 class FAnimationProvider;
 namespace TraceServices { class IAnalysisSession; }
@@ -26,6 +27,7 @@ private:
 		RouteId_SkeletalMeshComponent,
 		RouteId_SkeletalMeshComponent2,
 		RouteId_SkeletalMeshComponent3,
+		RouteId_SkeletalMeshComponent4,
 		RouteId_SkeletalMeshFrame,
 		RouteId_AnimGraph,
 		RouteId_AnimNodeStart,
diff --git a/Engine/Plugins/Animation/GameplayInsights/Source/GameplayInsights/Private/AnimationProvider.cpp b/Engine/Plugins/Animation/GameplayInsights/Source/GameplayInsights/Private/AnimationProvider.cpp
--- a/Engine/Plugins/Animation/GameplayInsights/Source/GameplayInsights/Private/AnimationProvider.cpp
+++ b/Engine/Plugins/Animation/GameplayInsights/Source/GameplayInsights/Private/AnimationProvider.cpp
@@ -3,6 +3,7 @@
 #include "AnimationProvider.h"
 #include "GameplayProvider.h"
 #include "Insights/ViewModels/TimingEventsTrack.h"
+#include "AnimationAnalyzer.h"
 
 #define LOCTEXT_NAMESPACE "AnimationProvider"
 
@@ -13,6 +14,7 @@ FAnimationProvider::FAnimationProvider(TraceServices::IAnalysisSession& InSessio
 	, GameplayProvider(InGameplayProvider)
 	, SkeletalMeshPoseTransforms(InSession.GetLinearAllocator(), 256)
 	, SkeletalMeshCurves(InSession.GetLinearAllocator(), 256)
+	, ExternalMorphWeights(InSession.GetLinearAllocator(), 256)
 	, SkeletalMeshParentIndices(InSession.GetLinearAllocator(), 256)
 	, PoseWatchRequiredBones(InSession.GetLinearAllocator(), 256)
 	, bHasAnyData(false)
@@ -155,6 +157,18 @@ void FAnimationProvider::EnumeratePoseWatchCurves(const FPoseWatchMessage& InMes
 	}
 }
 
+void FAnimationProvider::EnumerateExternalMorphSets(const FSkeletalMeshPoseMessage& InMessage, TFunctionRef<void(const FExternalMorphWeightMessage&)> Callback) const
+{
+	Session.ReadAccessCheck();
+
+	const uint64 StartIndex = InMessage.ExternalMorphStartIndex;
+	const uint64 EndIndex = StartIndex + InMessage.NumExternalMorphSets;
+	for (uint64 Index = StartIndex; Index < EndIndex; ++Index)
+	{
+		Callback(ExternalMorphWeights[Index]);
+	}
+}
+
 void FAnimationProvider::EnumerateTickRecordIds(uint64 InObjectId, TFunctionRef<void(uint64, int32)> Callback) const
 {
 	Session.ReadAccessCheck();
@@ -517,6 +531,24 @@ FText FAnimationProvider::FormatNodeValue(const FAnimNodeValueMessage& InMessage
 	return Text;
 }
 
+bool FAnimationProvider::HasExternalMorphSets(uint64 InObjectId) const
+{
+	//ObjectIdToSkeletalMeshPoseTimelines
+	Session.ReadAccessCheck();
+
+	const uint32* IndexPtr = ObjectIdToSkeletalMeshPoseTimelines.Find(InObjectId);
+	if(IndexPtr != nullptr)
+	{
+		if (*IndexPtr < uint32(SkeletalMeshPoseTimelineStorage.Num()))
+		{
+			const TSharedRef<FSkeletalMeshTimelineStorage>& SkeletalMeshTimelineStorage = SkeletalMeshPoseTimelineStorage[*IndexPtr];
+			return (SkeletalMeshTimelineStorage->NumExternalMorphSets > 0);
+		}
+	}
+
+	return false;
+}
+
 bool FAnimationProvider::HasAnyData() const
 {
 	Session.ReadAccessCheck();
@@ -695,7 +727,18 @@ static FTransform ConvertTransform(int TransformSize, const float* TransformFloa
 	return FTransform(Rotation, Translation, Scale3D);
 }
 
-void FAnimationProvider::AppendSkeletalMeshComponent(uint64 InObjectId, uint64 InMeshId, double InProfileTime, double InRecordingTime, uint16 InLodIndex, uint16 InFrameCounter, const TArrayView<const float>& InComponentToWorldRaw, const TArrayView<const float>& InPoseRaw, const TArrayView<const uint32>& InCurveIds, const TArrayView<const float>& InCurveValues)
+void FAnimationProvider::AppendSkeletalMeshComponent(
+	uint64 InObjectId,
+	uint64 InMeshId,
+	double InProfileTime,
+	double InRecordingTime,
+	uint16 InLodIndex,
+	uint16 InFrameCounter,
+	const TArrayView<const float>& InComponentToWorldRaw,
+	const TArrayView<const float>& InPoseRaw,
+	const TArrayView<const uint32>& InCurveIds,
+	const TArrayView<const float>& InCurveValues)
+
 {
 	Session.WriteAccessCheck();
 
@@ -783,6 +826,127 @@ void FAnimationProvider::AppendSkeletalMeshComponent(uint64 InObjectId, uint64 I
 	Session.UpdateDurationSeconds(InProfileTime);
 }
 
+void FAnimationProvider::AppendSkeletalMeshComponent(
+	uint64 InObjectId,
+	uint64 InMeshId,
+	double InProfileTime,
+	double InRecordingTime,
+	uint16 InLodIndex,
+	uint16 InFrameCounter,
+	const TArrayView<const float>& InComponentToWorldRaw,
+	const TArrayView<const float>& InPoseRaw,
+	const TArrayView<const uint32>& InCurveIds,
+	const TArrayView<const float>& InCurveValues,
+	const TArrayView<const float>& InExternalMorphWeights,
+	const TArrayView<const int32>& InExternalMorphCounts)
+{
+	Session.WriteAccessCheck();
+
+	bHasAnyData = true;
+
+	TSharedPtr<FSkeletalMeshTimelineStorage> TimelineStorage;
+	uint32* IndexPtr = ObjectIdToSkeletalMeshPoseTimelines.Find(InObjectId);
+	if(IndexPtr != nullptr)
+	{
+		TimelineStorage = SkeletalMeshPoseTimelineStorage[*IndexPtr];
+	}
+	else
+	{
+		TimelineStorage = MakeShared<FSkeletalMeshTimelineStorage>();
+		TimelineStorage->Timeline = MakeShared<TraceServices::TIntervalTimeline<FSkeletalMeshPoseMessage>>(Session.GetLinearAllocator());
+		ObjectIdToSkeletalMeshPoseTimelines.Add(InObjectId, SkeletalMeshPoseTimelineStorage.Num());
+		SkeletalMeshPoseTimelineStorage.Add(TimelineStorage.ToSharedRef());
+	}
+
+	// terminate existing scopes
+	uint64 NumEvents = TimelineStorage->Timeline->GetEventCount();
+	if (NumEvents > 0)
+	{
+		// Add end event at current time
+		if (TimelineStorage->Timeline->GetEventEndTime(NumEvents - 1) > InProfileTime)
+		{
+			TimelineStorage->Timeline->EndEvent(NumEvents - 1, InProfileTime);
+		}
+	}
+
+	const int32 NumCurves = InCurveIds.Num();
+
+	const int CaptureTransformSize = InComponentToWorldRaw.Num();
+	const int LocalTransformSize = sizeof (FTransform) / sizeof(float);
+
+	FTransform ComponentToWorld;
+
+	if (CaptureTransformSize == LocalTransformSize)
+	{
+	 	FMemory::Memcpy(&ComponentToWorld, &InComponentToWorldRaw[0], sizeof(FTransform));
+	}
+	else
+	{
+	 	ComponentToWorld = ConvertTransform(CaptureTransformSize, &InComponentToWorldRaw[0]);
+	}
+
+	const int PoseTransformCount = InPoseRaw.Num()/CaptureTransformSize;
+
+	FSkeletalMeshPoseMessage Message;
+	Message.RecordingTime = InRecordingTime;
+	Message.ComponentToWorld = ComponentToWorld;
+	Message.TransformStartIndex = SkeletalMeshPoseTransforms.Num();
+	Message.CurveStartIndex = SkeletalMeshCurves.Num();
+	Message.ComponentId = InObjectId;
+	Message.MeshId = InMeshId;
+	Message.MeshName = GameplayProvider.GetObjectInfo(Message.MeshId).Name;
+	Message.NumTransforms = (uint16)PoseTransformCount;
+	Message.NumCurves = (uint16)NumCurves;
+	Message.LodIndex = InLodIndex;
+	Message.FrameCounter = InFrameCounter;
+	Message.ExternalMorphStartIndex = static_cast<uint64>(ExternalMorphWeights.Num());
+	Message.NumExternalMorphSets = static_cast<uint16>(InExternalMorphCounts.Num());
+
+	TimelineStorage->Timeline->AppendBeginEvent(InProfileTime, Message);
+
+	if (CaptureTransformSize == LocalTransformSize)
+	{
+		for(int i=0; i<PoseTransformCount; i++)
+		{
+			FMemory::Memcpy(&SkeletalMeshPoseTransforms.PushBack(), &InPoseRaw[CaptureTransformSize * i], sizeof(FTransform));
+		}
+	}
+	else
+	{
+		for(int i=0; i<PoseTransformCount; i++)
+		{
+			SkeletalMeshPoseTransforms.PushBack() = ConvertTransform(CaptureTransformSize, &InPoseRaw[CaptureTransformSize * i]);
+		}
+	}
+
+	for(int32 CurveIndex = 0; CurveIndex < NumCurves; ++CurveIndex)
+	{
+		SkeletalMeshCurves.PushBack() = { InCurveIds[CurveIndex], InCurveValues[CurveIndex] };
+		TimelineStorage->AllCurveIds.Add(InCurveIds[CurveIndex]);
+	}
+
+	// Save the weights of all morph targets inside each external morph set.
+	int32 WeightOffset = 0;
+	for (int32 MorphSetIndex = 0; MorphSetIndex < InExternalMorphCounts.Num(); ++MorphSetIndex)
+	{
+		// Get the weight values for the current morph set.
+		const int32 NumWeightsInMorphSet = InExternalMorphCounts[MorphSetIndex];
+		const TConstArrayView<float> SetWeights(&InExternalMorphWeights[WeightOffset], NumWeightsInMorphSet);
+
+		// Record them in our buffer by copying them.
+		FExternalMorphWeightMessage& NewWeights = ExternalMorphWeights.PushBack();
+		NewWeights.Weights = SetWeights;
+		NewWeights.NumMorphs = NumWeightsInMorphSet;
+		NewWeights.Index = MorphSetIndex;
+
+		WeightOffset += NumWeightsInMorphSet;
+	}
+
+	TimelineStorage->NumExternalMorphSets = InExternalMorphCounts.Num();
+
+	Session.UpdateDurationSeconds(InProfileTime);
+}
+	
 void FAnimationProvider::AppendName(uint32 InId, const TCHAR* InName)
 {
 	Session.WriteAccessCheck();
diff --git a/Engine/Plugins/Animation/GameplayInsights/Source/GameplayInsights/Private/AnimationProvider.h b/Engine/Plugins/Animation/GameplayInsights/Source/GameplayInsights/Private/AnimationProvider.h
--- a/Engine/Plugins/Animation/GameplayInsights/Source/GameplayInsights/Private/AnimationProvider.h
+++ b/Engine/Plugins/Animation/GameplayInsights/Source/GameplayInsights/Private/AnimationProvider.h
@@ -30,6 +30,7 @@ public:
 	virtual void EnumerateSkeletalMeshCurveIds(uint64 InObjectId, TFunctionRef<void(uint32)> Callback) const override;
 	virtual void EnumerateSkeletalMeshCurves(const FSkeletalMeshPoseMessage& InMessage, TFunctionRef<void(const FSkeletalMeshNamedCurve&)> Callback) const override;
 	virtual void EnumeratePoseWatchCurves(const FPoseWatchMessage& InMessage, TFunctionRef<void(const FSkeletalMeshNamedCurve&)> Callback) const override;
+	virtual void EnumerateExternalMorphSets(const FSkeletalMeshPoseMessage& InMessage, TFunctionRef<void(const FExternalMorphWeightMessage&)> Callback) const override;
 	virtual bool ReadTickRecordTimeline(uint64 InObjectId, TFunctionRef<void(const TickRecordTimeline&)> Callback) const override;
 	virtual void EnumerateTickRecordIds(uint64 InObjectId, TFunctionRef<void(uint64, int32)> Callback) const override;
 	virtual void EnumerateAnimGraphTimelines(TFunctionRef<void(uint64 ObjectId, const AnimGraphTimeline&)> Callback) const override;
@@ -52,6 +53,9 @@ public:
 	virtual FText FormatNodeKeyValue(const FAnimNodeValueMessage& InMessage) const override;
 	virtual FText FormatNodeValue(const FAnimNodeValueMessage& InMessage) const override;
 
+	/** Does the given object have external morph sets? */
+	bool HasExternalMorphSets(uint64 InObjectId) const;
+
 	// Check to see if any data is present for this provider. Used as an early-out when processing.
 	bool HasAnyData() const;
 
@@ -64,6 +68,7 @@ public:
 	/** Add a skeletal mesh pose/curves etc. */
 	void AppendSkeletalMeshComponent(uint64 InObjectId, uint64 InMeshId, double InProfileTime, double InRecordingTime, uint16 InLodIndex, uint16 InFrameCounter, const TArrayView<const FTransform>& InPose, const TArrayView<const FSkeletalMeshNamedCurve>& InCurves);
 	void AppendSkeletalMeshComponent(uint64 InObjectId, uint64 InMeshId, double InProfileTime, double InRecordingTime, uint16 InLodIndex, uint16 InFrameCounter, const TArrayView<const float>& InComponentToWorldRaw, const TArrayView<const float>& InPoseRaw, const TArrayView<const uint32>& InCurveIds, const TArrayView<const float>& InCurveValues);
+	void AppendSkeletalMeshComponent(uint64 InObjectId, uint64 InMeshId, double InProfileTime, double InRecordingTime, uint16 InLodIndex, uint16 InFrameCounter, const TArrayView<const float>& InComponentToWorldRaw, const TArrayView<const float>& InPoseRaw, const TArrayView<const uint32>& InCurveIds, const TArrayView<const float>& InCurveValues, const TArrayView<const float>& InExternalMorphWeights, const TArrayView<const int32>& InExternalMorphCounts);
 
 	/** Get a skeletal mesh for a specified path. If the mesh no longer exists, make a fake one */
 	USkeletalMesh* GetSkeletalMesh(const TCHAR* InPath);
@@ -169,6 +174,7 @@ private:
 	{
 		TSharedPtr<TraceServices::TIntervalTimeline<FSkeletalMeshPoseMessage>> Timeline;
 		TSet<uint32> AllCurveIds;
+		int32 NumExternalMorphSets = 0;
 	};
 
 	struct FAnimNotifyStateTimelineStorage
@@ -198,6 +204,7 @@ private:
 	TArray<TSharedRef<FMontageTimelineStorage>> AnimMontageTimelineStorage;
 	TraceServices::TPagedArray<FTransform> SkeletalMeshPoseTransforms;
 	TraceServices::TPagedArray<FSkeletalMeshNamedCurve> SkeletalMeshCurves;
+	TraceServices::TPagedArray<FExternalMorphWeightMessage> ExternalMorphWeights;
 	TraceServices::TPagedArray<int32> SkeletalMeshParentIndices;
 	TArray<TSharedRef<TraceServices::TPointTimeline<FAnimAttributeMessage>>> AnimAttributeTimelines;
 	TArray<TSharedRef<TraceServices::TPointTimeline<FAnimSyncMessage>>> AnimSyncTimelines;
diff --git a/Engine/Plugins/Animation/GameplayInsights/Source/GameplayInsights/Private/ExternalMorphTrack.cpp b/Engine/Plugins/Animation/GameplayInsights/Source/GameplayInsights/Private/ExternalMorphTrack.cpp
new file mode 100644
--- /dev/null
+++ b/Engine/Plugins/Animation/GameplayInsights/Source/GameplayInsights/Private/ExternalMorphTrack.cpp
@@ -0,0 +1,301 @@
+// Copyright Epic Games, Inc. All Rights Reserved.
+
+#include "ExternalMorphTrack.h"
+#include "SCurveTimelineView.h"
+#include "IRewindDebugger.h"
+#include "GameplayProvider.h"
+#include "AnimationProvider.h"
+#include "SAnimationCurvesView.h"
+
+#define LOCTEXT_NAMESPACE "ExternalMorphTrack"
+
+namespace RewindDebugger
+{
+	FExternalMorphSetGroupTrack::FExternalMorphSetGroupTrack(uint64 InObjectId)
+		: ObjectId(InObjectId)
+	{
+		SetIsExpanded(false);
+		Icon = FSlateIcon("EditorStyle", "AnimGraph.Attribute.Curves.Icon", "AnimGraph.Attribute.Curves.Icon");
+	}
+
+	void FExternalMorphSetGroupTrack::IterateSubTracksInternal(TFunction<void(TSharedPtr<FRewindDebuggerTrack> SubTrack)> IteratorFunction)
+	{
+		for (TSharedPtr<FExternalMorphSetTrack>& Track : Children)
+		{
+			IteratorFunction(Track);
+		}
+	};
+
+	bool FExternalMorphSetGroupTrack::UpdateInternal()
+	{
+		IRewindDebugger* RewindDebugger = IRewindDebugger::Instance();
+		const TraceServices::IAnalysisSession* AnalysisSession = RewindDebugger->GetAnalysisSession();
+
+		// Convert time range to from rewind debugger times to profiler times.
+		TRange<double> TraceTimeRange = RewindDebugger->GetCurrentTraceRange();
+		double StartTime = TraceTimeRange.GetLowerBoundValue();
+		double EndTime = TraceTimeRange.GetUpperBoundValue();
+	
+		const FGameplayProvider* GameplayProvider = AnalysisSession->ReadProvider<FGameplayProvider>(FGameplayProvider::ProviderName);
+		const FAnimationProvider* AnimationProvider = AnalysisSession->ReadProvider<FAnimationProvider>(FAnimationProvider::ProviderName);
+
+		TArray<int32> MorphSetIndexValues;
+
+		bool bChanged = false;
+		if (GameplayProvider && AnimationProvider)
+		{
+			TraceServices::FAnalysisSessionReadScope SessionReadScope(*AnalysisSession);
+
+			AnimationProvider->ReadSkeletalMeshPoseTimeline(ObjectId, [this, AnimationProvider, StartTime, EndTime, &MorphSetIndexValues](const FAnimationProvider::SkeletalMeshPoseTimeline& InTimeline, bool bHasCurves)
+			{
+				InTimeline.EnumerateEvents(StartTime, EndTime, [this, StartTime, EndTime, AnimationProvider, &MorphSetIndexValues](double InStartTime, double InEndTime, uint32 InDepth, const FSkeletalMeshPoseMessage& InMessage)
+				{
+					if (InEndTime > StartTime && InStartTime < EndTime)
+					{
+						AnimationProvider->EnumerateExternalMorphSets(InMessage, [this, &MorphSetIndexValues](const FExternalMorphWeightMessage& InSet)
+						{
+							MorphSetIndexValues.AddUnique(InSet.Index);
+						});
+					}
+					return TraceServices::EEventEnumerate::Continue;
+				});
+			});
+
+			bChanged = false;
+			if (MorphSetIndexValues.Num() != Children.Num() || PrevMorphSetIndexValues != MorphSetIndexValues)
+			{
+				bChanged = true;
+				Children.SetNum(MorphSetIndexValues.Num());
+				for (int32 Index = 0; Index < MorphSetIndexValues.Num(); ++Index)
+				{
+					TSharedPtr<FExternalMorphSetTrack> NewTrack = MakeShared<FExternalMorphSetTrack>(ObjectId, MorphSetIndexValues[Index]);
+					Children[Index] = NewTrack;
+				}
+			}
+
+			for (int32 Index = 0; Index < Children.Num(); ++Index)
+			{
+				bChanged |= Children[Index]->Update();
+			}
+
+			PrevMorphSetIndexValues = MorphSetIndexValues;
+		}
+
+		return bChanged;
+	}
+
+	//----------------------------------------------------------------------------------------------
+
+	FExternalMorphSetTrack::FExternalMorphSetTrack(uint64 InObjectId, int32 InMorphSetIndex)
+		: ObjectId(InObjectId)
+		, MorphSetIndex(InMorphSetIndex)
+	{
+		SetIsExpanded(false);
+		Icon = FSlateIcon("EditorStyle", "AnimGraph.Attribute.Curves.Icon", "AnimGraph.Attribute.Curves.Icon");
+	}
+
+	TSharedPtr<SWidget> FExternalMorphSetTrack::GetDetailsViewInternal()
+	{
+		IRewindDebugger* RewindDebugger = IRewindDebugger::Instance();
+		return SNew(SAnimationCurvesView, ObjectId, RewindDebugger->CurrentTraceTime(), *RewindDebugger->GetAnalysisSession())
+					.CurrentTime_Lambda([RewindDebugger]{ return RewindDebugger->CurrentTraceTime(); });
+	
+	}
+
+	void FExternalMorphSetTrack::IterateSubTracksInternal(TFunction<void(TSharedPtr<FRewindDebuggerTrack> SubTrack)> IteratorFunction)
+	{
+		for (TSharedPtr<FExternalMorphTrack>& Track : Children)
+		{
+			IteratorFunction(Track);
+		}
+	};
+
+	bool FExternalMorphSetTrack::UpdateInternal()
+	{
+		IRewindDebugger* RewindDebugger = IRewindDebugger::Instance();
+		const TraceServices::IAnalysisSession* AnalysisSession = RewindDebugger->GetAnalysisSession();
+
+		// Convert time range to from rewind debugger times to profiler times.
+		TRange<double> TraceTimeRange = RewindDebugger->GetCurrentTraceRange();
+		double StartTime = TraceTimeRange.GetLowerBoundValue();
+		double EndTime = TraceTimeRange.GetUpperBoundValue();
+	
+		const FGameplayProvider* GameplayProvider = AnalysisSession->ReadProvider<FGameplayProvider>(FGameplayProvider::ProviderName);
+		const FAnimationProvider* AnimationProvider = AnalysisSession->ReadProvider<FAnimationProvider>(FAnimationProvider::ProviderName);
+
+		int32 NumMorphTracks = 0;
+		bool bChanged = false;
+		if (GameplayProvider && AnimationProvider)
+		{
+			TraceServices::FAnalysisSessionReadScope SessionReadScope(*AnalysisSession);
+
+			AnimationProvider->ReadSkeletalMeshPoseTimeline(ObjectId, [this, AnimationProvider, StartTime, EndTime, &NumMorphTracks](const FAnimationProvider::SkeletalMeshPoseTimeline& InTimeline, bool bHasCurves)
+			{
+				InTimeline.EnumerateEvents(StartTime, EndTime, [this, StartTime, EndTime, AnimationProvider, &NumMorphTracks](double InStartTime, double InEndTime, uint32 InDepth, const FSkeletalMeshPoseMessage& InMessage)
+				{
+					if (InEndTime > StartTime && InStartTime < EndTime)
+					{
+						AnimationProvider->EnumerateExternalMorphSets(InMessage, [this, &NumMorphTracks](const FExternalMorphWeightMessage& InSet)
+						{
+							if (InSet.Index == MorphSetIndex)
+							{
+								NumMorphTracks = InSet.Weights.Num();
+							}
+						});
+					}
+					return TraceServices::EEventEnumerate::Continue;
+				});
+			});
+
+			bChanged = false;
+			if (NumMorphTracks != Children.Num())
+			{
+				bChanged = true;
+				Children.SetNum(NumMorphTracks);
+				for (int32 Index = 0; Index < NumMorphTracks; ++Index)
+				{
+					TSharedPtr<FExternalMorphTrack> NewTrack = MakeShared<FExternalMorphTrack>(ObjectId, Index, MorphSetIndex);
+					Children[Index] = NewTrack;
+				}
+			}
+
+			check(Children.Num() == NumMorphTracks);
+			for (int32 Index = 0; Index < NumMorphTracks; ++Index)
+			{
+				bChanged |= Children[Index]->Update();
+			}
+		}
+
+		return bChanged;
+	}
+
+	//----------------------------------------------------------------------------------------------
+
+	FExternalMorphTrack::FExternalMorphTrack(uint64 InObjectId, int32 InMorphIndex, int32 InMorphSetIndex) :
+		MorphIndex(InMorphIndex),
+		MorphSetIndex(InMorphSetIndex),
+		ObjectId(InObjectId)
+	{
+		Icon = FSlateIcon("EditorStyle", "AnimGraph.Attribute.Curves.Icon", "AnimGraph.Attribute.Curves.Icon");
+	}
+
+	TSharedPtr<SCurveTimelineView::FTimelineCurveData> FExternalMorphTrack::GetCurveData() const
+	{
+		if (!CurveData.IsValid())
+		{
+			CurveData = MakeShared<SCurveTimelineView::FTimelineCurveData>();
+		}
+	
+		CurvesUpdateRequested++;
+		return CurveData;
+	}
+
+	bool FExternalMorphTrack::UpdateInternal()
+	{
+		IRewindDebugger* RewindDebugger = IRewindDebugger::Instance();
+		const TraceServices::IAnalysisSession* AnalysisSession = RewindDebugger->GetAnalysisSession();
+		const FAnimationProvider* AnimationProvider = AnalysisSession->ReadProvider<FAnimationProvider>(FAnimationProvider::ProviderName);
+
+		if (CurvesUpdateRequested > 25 && AnimationProvider)
+		{
+			CurvesUpdateRequested = 0;
+			UpdateCurvePointsInternal();
+		}
+
+		bool bChanged = false;
+		if (Name.IsEmpty())
+		{
+			Name = FText::Format(LOCTEXT("MorphNameFormat", "Morph Target {0}"), MorphIndex);
+			bChanged = true;
+		}
+
+		return bChanged;
+	}
+
+	void FExternalMorphTrack::UpdateCurvePointsInternal()
+	{
+		IRewindDebugger* RewindDebugger = IRewindDebugger::Instance();
+		const TraceServices::IAnalysisSession* AnalysisSession = RewindDebugger->GetAnalysisSession();
+		const FAnimationProvider* AnimationProvider = AnalysisSession->ReadProvider<FAnimationProvider>(FAnimationProvider::ProviderName);
+
+		// Convert time range to from rewind debugger times to profiler times.
+		TRange<double> TraceTimeRange = RewindDebugger->GetCurrentTraceRange();
+		double StartTime = TraceTimeRange.GetLowerBoundValue();
+		double EndTime = TraceTimeRange.GetUpperBoundValue();
+	
+		TArray<SCurveTimelineView::FTimelineCurveData::CurvePoint>& CurvePoints = CurveData->Points;
+		CurvePoints.Reset();
+		CurvePoints.Reserve(256);
+
+		TraceServices::FAnalysisSessionReadScope SessionReadScope(*AnalysisSession);
+		AnimationProvider->ReadSkeletalMeshPoseTimeline(ObjectId, [this, AnimationProvider, StartTime, EndTime, AnalysisSession, &CurvePoints](const FAnimationProvider::SkeletalMeshPoseTimeline& InTimeline, bool bHasCurves)
+		{
+			InTimeline.EnumerateEvents(StartTime, EndTime, [this, StartTime, EndTime, AnimationProvider, AnalysisSession, &CurvePoints](double InStartTime, double InEndTime, uint32 InDepth, const FSkeletalMeshPoseMessage& InMessage)
+			{
+				if (InEndTime > StartTime && InStartTime < EndTime)
+				{
+					const double Time = InMessage.RecordingTime;		
+					AnimationProvider->EnumerateExternalMorphSets(InMessage, [this, Time, &CurvePoints](const FExternalMorphWeightMessage& InWeight)
+					{
+						if (InWeight.Index == MorphSetIndex)
+						{
+							CurvePoints.Add({Time, InWeight.Weights[MorphIndex]});
+						}
+					});
+				}
+				return TraceServices::EEventEnumerate::Continue;
+			});
+		});
+	}
+
+	TSharedPtr<SWidget> FExternalMorphTrack::GetDetailsViewInternal() 
+	{
+		IRewindDebugger* RewindDebugger = IRewindDebugger::Instance();
+		TSharedPtr<SAnimationCurvesView> AnimationCurvesView = SNew(SAnimationCurvesView, ObjectId, RewindDebugger->CurrentTraceTime(), *RewindDebugger->GetAnalysisSession())
+			.CurrentTime_Lambda([RewindDebugger]{ return RewindDebugger->CurrentTraceTime(); });
+		return AnimationCurvesView;
+	}
+
+	TSharedPtr<SWidget> FExternalMorphTrack::GetTimelineViewInternal()
+	{
+		FLinearColor CurveColor(0.5,0.5,0.5);
+	
+		return SNew(SCurveTimelineView)
+			.CurveColor(CurveColor)
+			.ViewRange_Lambda([]() { return IRewindDebugger::Instance()->GetCurrentViewRange(); })
+			.RenderFill(false)
+			.CurveData_Raw(this, &FExternalMorphTrack::GetCurveData);
+	}
+
+	//----------------------------------------------------------------------------------------------
+
+	FName FExternalMorphSetGroupTrackCreator::GetTargetTypeNameInternal() const
+	{
+		static const FName SkelMeshName("SkeletalMeshComponent");
+		return SkelMeshName;
+	}
+	
+	FName FExternalMorphSetGroupTrackCreator::GetNameInternal() const
+	{
+		static const FName ExternalMorphSetName("ExternalMorphSetGroup");
+		return ExternalMorphSetName;
+	}
+
+	TSharedPtr<RewindDebugger::FRewindDebuggerTrack> FExternalMorphSetGroupTrackCreator::CreateTrackInternal(uint64 ObjectId) const
+	{
+		return MakeShared<RewindDebugger::FExternalMorphSetGroupTrack>(ObjectId);
+	}
+
+	bool FExternalMorphSetGroupTrackCreator::HasDebugInfoInternal(uint64 ObjectId) const
+	{
+		const TraceServices::IAnalysisSession* AnalysisSession = IRewindDebugger::Instance()->GetAnalysisSession();
+		TraceServices::FAnalysisSessionReadScope SessionReadScope(*AnalysisSession);
+		if (const FAnimationProvider* AnimationProvider = AnalysisSession->ReadProvider<FAnimationProvider>(FAnimationProvider::ProviderName))
+		{
+			return AnimationProvider->HasExternalMorphSets(ObjectId);
+		}
+		return false;
+	}
+}
+
+#undef LOCTEXT_NAMESPACE
diff --git a/Engine/Plugins/Animation/GameplayInsights/Source/GameplayInsights/Private/ExternalMorphTrack.h b/Engine/Plugins/Animation/GameplayInsights/Source/GameplayInsights/Private/ExternalMorphTrack.h
new file mode 100644
--- /dev/null
+++ b/Engine/Plugins/Animation/GameplayInsights/Source/GameplayInsights/Private/ExternalMorphTrack.h
@@ -0,0 +1,113 @@
+// Copyright Epic Games, Inc. All Rights Reserved.
+#pragma once
+
+#include "IAnimationProvider.h"
+#include "IRewindDebuggerTrackCreator.h"
+#include "RewindDebuggerTrack.h"
+#include "SCurveTimelineView.h"
+#include "Textures/SlateIcon.h"
+
+namespace RewindDebugger
+{
+	// Sub Track for a external morph target.
+	class FExternalMorphTrack
+		: public FRewindDebuggerTrack
+	{
+	public:
+		FExternalMorphTrack(uint64 InObjectId, int32 InMorphIndex, int32 InMorphSetIndex);
+	
+		int32 GetMorphIndex() const { return MorphIndex; }
+
+	private:
+		// FRewindDebuggerTrack overrides.
+		virtual bool UpdateInternal() override;
+		TSharedPtr<SWidget> GetDetailsViewInternal() override;
+		virtual TSharedPtr<SWidget> GetTimelineViewInternal() override;
+		virtual FSlateIcon GetIconInternal() override						{ return Icon; }
+		virtual FName GetNameInternal() const override						{ return "ExternalMorphTarget"; }
+		virtual FText GetDisplayNameInternal() const override				{ return Name; }
+		virtual uint64 GetObjectIdInternal() const override					{ return ObjectId; }
+		// ~END FRewindDebuggerTrack overrides.
+	
+		TSharedPtr<SCurveTimelineView::FTimelineCurveData> GetCurveData() const;
+
+	protected:
+		void UpdateCurvePointsInternal();
+
+	protected:
+
+		int32 MorphIndex = INDEX_NONE;
+		int32 MorphSetIndex = INDEX_NONE;
+		uint64 ObjectId = 0;
+		FText Name;
+		FSlateIcon Icon;
+		mutable TSharedPtr<SCurveTimelineView::FTimelineCurveData> CurveData;
+		mutable int CurvesUpdateRequested = 0;
+	};
+
+
+	// The external morph target set, which is the parent of a collection of FExternalMorphTrack items.
+	class FExternalMorphSetTrack
+		: public FRewindDebuggerTrack
+	{
+	public:
+		FExternalMorphSetTrack(uint64 InObjectId, int32 InMorphSetIndex);
+
+	private:
+		// FRewindDebuggerTrack overrides.
+		virtual bool UpdateInternal() override;
+		virtual TSharedPtr<SWidget> GetDetailsViewInternal() override;
+		virtual FSlateIcon GetIconInternal() override						{ return Icon; }
+		virtual FName GetNameInternal() const override						{ return "ExternalMorphSet"; }
+		virtual FText GetDisplayNameInternal() const override				{ return NSLOCTEXT("RewindDebugger", "ExternalMorphSetTrackName", "External Morph Set"); }
+		virtual uint64 GetObjectIdInternal() const override					{ return ObjectId; }
+		virtual void IterateSubTracksInternal(TFunction<void(TSharedPtr<FRewindDebuggerTrack> SubTrack)> IteratorFunction) override;
+		// ~END FRewindDebuggerTrack overrides.
+	
+	protected:
+		FSlateIcon Icon;
+		uint64 ObjectId;
+		int32 MorphSetIndex;
+		TArray<TSharedPtr<FExternalMorphTrack>> Children;
+	};
+
+
+	// The external morph set group track, which is the parent of all morph sets.
+	class FExternalMorphSetGroupTrack
+		: public FRewindDebuggerTrack
+	{
+	public:
+		FExternalMorphSetGroupTrack(uint64 InObjectId);
+
+	private:
+		// FRewindDebuggerTrack overrides.
+		virtual bool UpdateInternal() override;
+		virtual FSlateIcon GetIconInternal() override						{ return Icon; }
+		virtual FName GetNameInternal() const override						{ return "ExternalMorphSetGroup"; }
+		virtual FText GetDisplayNameInternal() const override				{ return NSLOCTEXT("RewindDebugger", "ExternalMorphSetGroupTrackName", "External Morph Sets"); }
+		virtual uint64 GetObjectIdInternal() const override					{ return ObjectId; }
+		virtual void IterateSubTracksInternal(TFunction<void(TSharedPtr<FRewindDebuggerTrack> SubTrack)> IteratorFunction) override;
+		// ~END FRewindDebuggerTrack overrides.
+	
+	protected:
+		FSlateIcon Icon;
+		uint64 ObjectId;
+		TArray<TSharedPtr<FExternalMorphSetTrack>> Children;
+		TArray<int32> PrevMorphSetIndexValues;
+	};
+
+
+	// The creator for the external morph set track.
+	class FExternalMorphSetGroupTrackCreator
+		: public IRewindDebuggerTrackCreator
+	{
+	private:
+		// IRewindDebuggerTrackCreator overrides.
+		virtual FName GetTargetTypeNameInternal() const override;
+		virtual FName GetNameInternal() const override;
+		virtual TSharedPtr<FRewindDebuggerTrack> CreateTrackInternal(uint64 ObjectId) const override;
+		virtual bool HasDebugInfoInternal(uint64 ObjectId) const override;
+		// ~END IRewindDebuggerTrackCreator overrides.
+	};
+
+}	// namespace RewindDebugger
\ No newline at end of file
diff --git a/Engine/Plugins/Animation/GameplayInsights/Source/GameplayInsights/Private/GameplayInsightsModule.cpp b/Engine/Plugins/Animation/GameplayInsights/Source/GameplayInsights/Private/GameplayInsightsModule.cpp
--- a/Engine/Plugins/Animation/GameplayInsights/Source/GameplayInsights/Private/GameplayInsightsModule.cpp
+++ b/Engine/Plugins/Animation/GameplayInsights/Source/GameplayInsights/Private/GameplayInsightsModule.cpp
@@ -24,6 +24,7 @@
 #include "NotifiesTrack.h"
 #include "PoseWatchTrack.h"
 #include "PropertiesTrack.h"
+#include "ExternalMorphTrack.h"
 #include "PropertyWatchManager.h"
 
 #if WITH_EDITOR
@@ -92,6 +93,8 @@ void FGameplayInsightsModule::StartupModule()
 	IModularFeatures::Get().RegisterModularFeature(RewindDebugger::IRewindDebuggerTrackCreator::ModularFeatureName, &PoseWatchesTrackCreator);
 	static RewindDebugger::FPropertiesTrackCreator PropertyTrackCreator;
 	IModularFeatures::Get().RegisterModularFeature(RewindDebugger::IRewindDebuggerTrackCreator::ModularFeatureName, &PropertyTrackCreator);
+	static RewindDebugger::FExternalMorphSetGroupTrackCreator ExternalMorphSetGroupTrackCreator;
+	IModularFeatures::Get().RegisterModularFeature(RewindDebugger::IRewindDebuggerTrackCreator::ModularFeatureName, &ExternalMorphSetGroupTrackCreator);
 
 	FPropertyWatchManager::Initialize();
 	
diff --git a/Engine/Plugins/Animation/GameplayInsights/Source/GameplayInsights/Public/IAnimationProvider.h b/Engine/Plugins/Animation/GameplayInsights/Source/GameplayInsights/Public/IAnimationProvider.h
--- a/Engine/Plugins/Animation/GameplayInsights/Source/GameplayInsights/Public/IAnimationProvider.h
+++ b/Engine/Plugins/Animation/GameplayInsights/Source/GameplayInsights/Public/IAnimationProvider.h
@@ -29,6 +29,13 @@ struct FSkeletalMeshNamedCurve
 	float Value = 0.0f;
 };
 
+struct FExternalMorphWeightMessage
+{
+	int32 Index = 0;
+	int32 NumMorphs = 0;
+	TArray<float> Weights;
+};
+
 struct FSkeletalMeshPoseMessage
 {
 	FTransform ComponentToWorld;
@@ -37,11 +44,13 @@ struct FSkeletalMeshPoseMessage
 	uint64 CurveStartIndex = 0;
 	uint64 ComponentId = 0;	
 	uint64 MeshId = 0;
+	uint64 ExternalMorphStartIndex = 0;
 	const TCHAR* MeshName = nullptr;
 	uint16 NumTransforms = 0;
 	uint16 NumCurves = 0;
 	uint16 FrameCounter = 0;
 	uint16 LodIndex = 0;
+	uint16 NumExternalMorphSets = 0;
 };
 
 struct FPoseWatchMessage
@@ -313,6 +322,7 @@ public:
 	virtual void GetPoseWatchData(const FPoseWatchMessage& InMessage, TArray<FTransform>& BoneTransforms, TArray<uint16>& RequiredBones) const = 0;
 	virtual void EnumerateSkeletalMeshCurves(const FSkeletalMeshPoseMessage& InMessage, TFunctionRef<void(const FSkeletalMeshNamedCurve&)> Callback) const = 0;
 	virtual void EnumeratePoseWatchCurves(const FPoseWatchMessage& InMessage, TFunctionRef<void(const FSkeletalMeshNamedCurve&)> Callback) const = 0;
+	virtual void EnumerateExternalMorphSets(const FSkeletalMeshPoseMessage& InMessage, TFunctionRef<void(const FExternalMorphWeightMessage&)> Callback) const = 0;
 	virtual bool ReadTickRecordTimeline(uint64 InObjectId, TFunctionRef<void(const TickRecordTimeline&)> Callback) const = 0;
 	virtual void EnumerateTickRecordIds(uint64 InObjectId, TFunctionRef<void(uint64, int32)> Callback) const = 0;
 	virtual void EnumerateAnimGraphTimelines(TFunctionRef<void(uint64 ObjectId, const AnimGraphTimeline&)> Callback) const =0;
diff --git a/Engine/Plugins/Animation/MLDeformer/MLDeformerFramework/MLDeformerFramework.uplugin b/Engine/Plugins/Animation/MLDeformer/MLDeformerFramework/MLDeformerFramework.uplugin
--- a/Engine/Plugins/Animation/MLDeformer/MLDeformerFramework/MLDeformerFramework.uplugin
+++ b/Engine/Plugins/Animation/MLDeformer/MLDeformerFramework/MLDeformerFramework.uplugin
@@ -38,6 +38,10 @@
 		{
 			"Name": "GeometryCache",
 			"Enabled": true
+		},
+		{
+			"Name": "GameplayInsights",
+			"Enabled": true
 		}
 	]
 }
\ No newline at end of file
diff --git a/Engine/Plugins/Animation/MLDeformer/MLDeformerFramework/Source/MLDeformerFramework/Private/MLDeformerComponent.cpp b/Engine/Plugins/Animation/MLDeformer/MLDeformerFramework/Source/MLDeformerFramework/Private/MLDeformerComponent.cpp
--- a/Engine/Plugins/Animation/MLDeformer/MLDeformerFramework/Source/MLDeformerFramework/Private/MLDeformerComponent.cpp
+++ b/Engine/Plugins/Animation/MLDeformer/MLDeformerFramework/Source/MLDeformerFramework/Private/MLDeformerComponent.cpp
@@ -53,7 +53,7 @@ void UMLDeformerComponent::Init()
 	}
 	else
 	{
-		ModelInstance = nullptr;
+		ReleaseModelInstance();
 		UE_LOG(LogMLDeformer, Warning, TEXT("ML Deformer component on '%s' has a deformer asset that has no ML model setup."), *GetOuter()->GetName());
 	}
 }
@@ -136,12 +136,13 @@ void UMLDeformerComponent::UnbindDelegates()
 void UMLDeformerComponent::BeginDestroy()
 {
 	UnbindDelegates();
+	ReleaseModelInstance();
 	Super::BeginDestroy();
 }
 
 void UMLDeformerComponent::ReleaseModelInstance()
 {
-	if (ModelInstance)
+	if (ModelInstance && IsValid(ModelInstance))
 	{
 		ModelInstance->ConditionalBeginDestroy(); // Force destruction immediately instead of waiting for the next GC.
 		ModelInstance = nullptr;
@@ -155,7 +156,7 @@ USkeletalMeshComponent* UMLDeformerComponent::FindSkeletalMeshComponent(const UM
 	{
 		// First search for a skeletal mesh component that uses the same skeletal mesh as the ML Deformer asset was trained on.
 		const UMLDeformerModel* Model = Asset ? Asset->GetModel() : nullptr;
-		if (Model && Model->GetSkeletalMesh())
+		if (Model && Model->GetInputInfo() && Model->GetInputInfo()->GetSkeletalMesh().IsValid())
 		{
 			const FSoftObjectPath& ModelSkeletalMesh = Model->GetInputInfo()->GetSkeletalMesh();
 
@@ -185,8 +186,8 @@ USkeletalMeshComponent* UMLDeformerComponent::FindSkeletalMeshComponent(const UM
 
 void UMLDeformerComponent::UpdateSkeletalMeshComponent()
 {
-	SkelMeshComponent = FindSkeletalMeshComponent(DeformerAsset.Get());
-	SetupComponent(DeformerAsset, SkelMeshComponent);
+	USkeletalMeshComponent* Component = FindSkeletalMeshComponent(DeformerAsset.Get());
+	SetupComponent(DeformerAsset, Component);
 }
 
 void UMLDeformerComponent::Activate(bool bReset)
@@ -210,6 +211,17 @@ void UMLDeformerComponent::Deactivate()
 	Super::Deactivate();
 }
 
+
+float UMLDeformerComponent::GetFinalMLDeformerWeight() const
+{
+	float FinalWeight = Weight;
+	if (GMLDeformerOverrideWeight >= 0.0f)
+	{
+		FinalWeight = FMath::Min(GMLDeformerOverrideWeight, 1.0f);
+	}
+	return FinalWeight;
+}
+
 void UMLDeformerComponent::TickComponent(float DeltaTime, enum ELevelTick TickType, FActorComponentTickFunction* ThisTickFunction)
 {
 	#if WITH_EDITOR
@@ -224,11 +236,8 @@ void UMLDeformerComponent::TickComponent(float DeltaTime, enum ELevelTick TickTy
 			SkelMeshComponent->GetPredictedLODLevel() == 0)
 		{
 			TRACE_CPUPROFILER_EVENT_SCOPE(UMLDeformerComponent::TickComponent)
-			float ApplyWeight = Weight;
-			if (GMLDeformerOverrideWeight >= 0.0f)
-			{
-				ApplyWeight = FMath::Min(GMLDeformerOverrideWeight, 1.0f);
-			}
+			
+			const float ApplyWeight = GetFinalMLDeformerWeight();
 			ModelInstance->Tick(DeltaTime, ApplyWeight);
 
 			#if WITH_EDITOR
@@ -255,8 +264,8 @@ void UMLDeformerComponent::SetWeightInternal(const float NormalizedWeightValue)
 
 void UMLDeformerComponent::SetDeformerAssetInternal(UMLDeformerAsset* const InDeformerAsset)
 { 
-	DeformerAsset = InDeformerAsset;
-	UpdateSkeletalMeshComponent();
+	USkeletalMeshComponent* SkelMeshComp = FindSkeletalMeshComponent(InDeformerAsset);
+	SetupComponent(InDeformerAsset, SkelMeshComp);
 }
 
 #if WITH_EDITOR
diff --git a/Engine/Plugins/Animation/MLDeformer/MLDeformerFramework/Source/MLDeformerFramework/Private/MLDeformerModel.cpp b/Engine/Plugins/Animation/MLDeformer/MLDeformerFramework/Source/MLDeformerFramework/Private/MLDeformerModel.cpp
--- a/Engine/Plugins/Animation/MLDeformer/MLDeformerFramework/Source/MLDeformerFramework/Private/MLDeformerModel.cpp
+++ b/Engine/Plugins/Animation/MLDeformer/MLDeformerFramework/Source/MLDeformerFramework/Private/MLDeformerModel.cpp
@@ -15,6 +15,7 @@
 #include "UObject/UObjectGlobals.h"
 #include "RHICommandList.h"
 #include "AssetRegistry/AssetData.h"
+#include "Components/SkeletalMeshComponent.h"
 
 #include UE_INLINE_GENERATED_CPP_BY_NAME(MLDeformerModel)
 
@@ -196,6 +197,64 @@ void UMLDeformerModel::FloatArrayToVector3Array(const TArray<float>& FloatArray,
 	}
 }
 
+bool UMLDeformerModel::IsCompatibleDebugActor(const AActor* Actor, UMLDeformerComponent** OutDebugComponent) const
+{
+	// Set it to nullptr first, in case we exit this method early.
+	if (OutDebugComponent)
+	{
+		*OutDebugComponent = nullptr;
+	}
+
+	if (Actor == nullptr || !IsValid(Actor))
+	{
+		return false;
+	}
+
+	// Iterate over all skeletal mesh components, see if one matches our currently loaded character.
+	USkeletalMesh* SkelMesh = nullptr;
+	for (const UActorComponent* Component : Actor->GetComponents())
+	{
+		const USkeletalMeshComponent* SkelMeshComponent = Cast<USkeletalMeshComponent>(Component);
+		if (!SkelMeshComponent)
+		{
+			continue;
+		}
+
+		if (SkelMeshComponent->GetSkeletalMeshAsset() == SkeletalMesh)
+		{
+			SkelMesh = SkeletalMesh;
+			break;
+		}
+	}
+
+	// If we haven't found a matching skeletal mesh, we can ignore this actor.
+	if (!SkelMesh)
+	{
+		return false;
+	}
+
+	// Now check if we have an ML Deformer component on the actor uses the same ML Deformer asset.
+	for (UActorComponent* Component : Actor->GetComponents())
+	{
+		UMLDeformerComponent* MLDeformerComponent = Cast<UMLDeformerComponent>(Component);
+		if (!MLDeformerComponent)
+		{
+			continue;
+		}
+
+		if (MLDeformerComponent->GetDeformerAsset() == GetDeformerAsset())
+		{
+			if (OutDebugComponent)
+			{
+				*OutDebugComponent = MLDeformerComponent;
+			}
+			return true;
+		}
+	}
+
+	return false;
+}
+
 #if WITH_EDITOR
 	void UMLDeformerModel::UpdateNumTargetMeshVertices()
 	{ 
diff --git a/Engine/Plugins/Animation/MLDeformer/MLDeformerFramework/Source/MLDeformerFramework/Private/MLDeformerModelInstance.cpp b/Engine/Plugins/Animation/MLDeformer/MLDeformerFramework/Source/MLDeformerFramework/Private/MLDeformerModelInstance.cpp
--- a/Engine/Plugins/Animation/MLDeformer/MLDeformerFramework/Source/MLDeformerFramework/Private/MLDeformerModelInstance.cpp
+++ b/Engine/Plugins/Animation/MLDeformer/MLDeformerFramework/Source/MLDeformerFramework/Private/MLDeformerModelInstance.cpp
@@ -68,11 +68,6 @@ bool UMLDeformerModelInstance::HasPostInitialized() const
 	return bHasPostInitialized;
 }
 
-const TArray<FTransform>& UMLDeformerModelInstance::GetBoneTransforms() const
-{ 
-	return BoneTransforms;
-}
-
 bool UMLDeformerModelInstance::IsCompatible() const
 { 
 	return bIsCompatible;
@@ -101,6 +96,7 @@ void UMLDeformerModelInstance::Init(USkeletalMeshComponent* SkelMeshComponent)
 			}
 		}
 
+		UpdateCompatibilityStatus();
 		return;
 	}
 
@@ -132,7 +128,7 @@ void UMLDeformerModelInstance::Init(USkeletalMeshComponent* SkelMeshComponent)
 
 void UMLDeformerModelInstance::UpdateCompatibilityStatus()
 {
-	bIsCompatible = SkeletalMeshComponent->GetSkeletalMeshAsset() && CheckCompatibility(SkeletalMeshComponent, true).IsEmpty();
+	bIsCompatible = SkeletalMeshComponent && SkeletalMeshComponent->GetSkeletalMeshAsset() && CheckCompatibility(SkeletalMeshComponent, true).IsEmpty();
 }
 
 FString UMLDeformerModelInstance::CheckCompatibility(USkeletalMeshComponent* InSkelMeshComponent, bool bLogIssues)
@@ -157,47 +153,120 @@ FString UMLDeformerModelInstance::CheckCompatibility(USkeletalMeshComponent* InS
 	return ErrorText;
 }
 
+USkeletalMeshComponent* UMLDeformerModelInstance::GetFinalSkeletalMeshComponent() const
+{
+	#if WITH_EDITOR
+		const UMLDeformerComponent* MLDeformerComponent = GetMLDeformerComponent();
+		check(MLDeformerComponent);
+
+		const AActor* DebugActor = MLDeformerComponent->GetDebugActor();
+		if (DebugActor)
+		{
+			const USkeletalMesh* ModelSkelMesh = GetModel()->GetSkeletalMesh();
+			if (ModelSkelMesh)
+			{
+				for (UActorComponent* Component : DebugActor->GetComponents())
+				{
+					USkeletalMeshComponent* DebugActorSkelMeshComponent = Cast<USkeletalMeshComponent>(Component);
+					if (DebugActorSkelMeshComponent && DebugActorSkelMeshComponent->GetSkeletalMeshAsset() == ModelSkelMesh)
+					{
+						return DebugActorSkelMeshComponent;
+					}
+				}
+			}
+		}
+	#endif
+
+	return SkeletalMeshComponent;
+}
+
+
+namespace
+{
+	void CalculateBoneSpaceTransforms(const FReferenceSkeleton& RefSkel, const TArray<FTransform>& ComponentSpaceTransforms, TArray<FTransform>& OutBoneSpaceTransforms)
+	{
+		for (int32 Index = 0; Index < RefSkel.GetNum(); ++Index)
+		{
+			const int32 ParentIndex = RefSkel.GetParentIndex(Index);
+			if (ParentIndex != INDEX_NONE)
+			{
+				OutBoneSpaceTransforms[Index] = ComponentSpaceTransforms[Index].GetRelativeTransform(ComponentSpaceTransforms[ParentIndex]);
+				OutBoneSpaceTransforms[Index].NormalizeRotation();
+			}
+			else
+			{
+				OutBoneSpaceTransforms[Index] = ComponentSpaceTransforms[Index];
+			}
+		}
+	}
+}	// Anonymous namespace
+
+
 void UMLDeformerModelInstance::UpdateBoneTransforms()
 {
 	// If we use a leader component, we have to get the transforms from that component.
-	const USkinnedMeshComponent* LeaderPoseComponent = SkeletalMeshComponent->LeaderPoseComponent.Get();
+	USkeletalMeshComponent* FinalSkelMeshComponent = GetFinalSkeletalMeshComponent();
+	const USkinnedMeshComponent* LeaderPoseComponent = FinalSkelMeshComponent->LeaderPoseComponent.Get();
 	if (LeaderPoseComponent)
 	{
 		const TArray<FTransform>& LeaderTransforms = LeaderPoseComponent->GetComponentSpaceTransforms();
 		if (!LeaderTransforms.IsEmpty())
 		{
-			USkinnedAsset* SkinnedAsset = LeaderPoseComponent->GetSkinnedAsset();
+			const USkinnedAsset* SkinnedAsset = LeaderPoseComponent->GetSkinnedAsset();
+			const FReferenceSkeleton& RefSkel = SkinnedAsset->GetRefSkeleton();
 			const int32 NumTrainingBones = AssetBonesToSkelMeshMappings.Num();
 			for (int32 Index = 0; Index < NumTrainingBones; ++Index)
 			{
 				const int32 ComponentBoneIndex = AssetBonesToSkelMeshMappings[Index];
-				checkSlow(LeaderTransforms.IsValidIndex(ComponentBoneIndex));			
 				const FTransform& ComponentSpaceTransform = LeaderTransforms[ComponentBoneIndex];
-				const int32 ParentIndex = SkinnedAsset->GetRefSkeleton().GetParentIndex(ComponentBoneIndex);
+				const int32 ParentIndex = RefSkel.GetParentIndex(ComponentBoneIndex);
 				if (LeaderTransforms.IsValidIndex(ParentIndex))
 				{
 					TrainingBoneTransforms[Index] = ComponentSpaceTransform.GetRelativeTransform(LeaderTransforms[ParentIndex]);
+					TrainingBoneTransforms[Index].NormalizeRotation();
 				}
 				else
 				{
 					TrainingBoneTransforms[Index] = ComponentSpaceTransform;
 				}
-				TrainingBoneTransforms[Index].NormalizeRotation();
 			}
 		}
 	}
-	else
+	else // No leader pose is used.
 	{
-		// Grab the transforms from our own skeletal mesh component.
-		// These are local space transforms, relative to the parent bone.
-		const TArrayView<const FTransform> Transforms = SkeletalMeshComponent->GetBoneSpaceTransformsView();
-		if (!Transforms.IsEmpty())
+		if (FinalSkelMeshComponent == SkeletalMeshComponent)
 		{
+			// Grab the transforms from our own skeletal mesh component.
+			// These are local space transforms, relative to the parent bone.
+			const TArrayView<const FTransform> Transforms = FinalSkelMeshComponent->GetBoneSpaceTransformsView();
+			if (!Transforms.IsEmpty())
+			{
+				const int32 NumTrainingBones = AssetBonesToSkelMeshMappings.Num();
+				for (int32 Index = 0; Index < NumTrainingBones; ++Index)
+				{
+					const int32 ComponentBoneIndex = AssetBonesToSkelMeshMappings[Index];
+					TrainingBoneTransforms[Index] = Transforms[ComponentBoneIndex];
+				}
+			}
+		}
+		else // We are debugging and grabbing transforms from another component.
+		{
+			USkeletalMesh* SkelMeshAsset = FinalSkelMeshComponent->GetSkeletalMeshAsset();
+			check(SkelMeshAsset);
+
+			// Allocate enough space in our bone space transforms buffer.
+			// We need to calculate the bone space transforms, because when debugging only component space transforms are updated.
+			DebugBoneSpaceTransforms.Reset();
+			DebugBoneSpaceTransforms.SetNumUninitialized(FinalSkelMeshComponent->GetComponentSpaceTransforms().Num());
+			const FReferenceSkeleton& RefSkel = SkelMeshAsset->GetRefSkeleton();
+			CalculateBoneSpaceTransforms(RefSkel, FinalSkelMeshComponent->GetComponentSpaceTransforms(), DebugBoneSpaceTransforms);
+
+			// Now that we have the new bone space transforms, store them so they will be passed into our neural network later.
 			const int32 NumTrainingBones = AssetBonesToSkelMeshMappings.Num();
 			for (int32 Index = 0; Index < NumTrainingBones; ++Index)
 			{
 				const int32 ComponentBoneIndex = AssetBonesToSkelMeshMappings[Index];
-				TrainingBoneTransforms[Index] = Transforms[ComponentBoneIndex];
+				TrainingBoneTransforms[Index] = DebugBoneSpaceTransforms[ComponentBoneIndex];
 			}
 		}
 	}
@@ -279,7 +348,8 @@ bool UMLDeformerModelInstance::HasValidTransforms() const
 		return false;
 	}
 
-	const USkinnedMeshComponent* LeaderPoseComponent = SkeletalMeshComponent->LeaderPoseComponent.Get();
+	USkeletalMeshComponent* FinalSkeletalMeshComponent = GetFinalSkeletalMeshComponent();
+	const USkinnedMeshComponent* LeaderPoseComponent = FinalSkeletalMeshComponent->LeaderPoseComponent.Get();
 	if (LeaderPoseComponent)
 	{
 		if (LeaderPoseComponent->GetComponentSpaceTransforms().IsEmpty())
@@ -287,7 +357,25 @@ bool UMLDeformerModelInstance::HasValidTransforms() const
 			return false;
 		}
 	}
-	else if (SkeletalMeshComponent->GetBoneSpaceTransformsView().IsEmpty())
+	else if (FinalSkeletalMeshComponent->GetBoneSpaceTransformsView().IsEmpty())
+	{
+		return false;
+	}
+
+	// Make sure we use the same Skeletal Mesh on the current skeletal mesh component as that we trained the deformer on.
+	// This is already checked before by compatibility checks, but it can still fail when some other game code suddenly
+	// changes the Skeletal Mesh on the SkeletalMeshComponent after the ML Deformer component has been initialized already.
+	// This is just an extra safety check.
+	USkeletalMesh* SkelMesh = FinalSkeletalMeshComponent->GetSkeletalMeshAsset();
+	UMLDeformerInputInfo* InputInfo = Model->GetInputInfo();
+	if (SkelMesh && InputInfo)
+	{
+		if (FSoftObjectPath(SkelMesh) != InputInfo->GetSkeletalMesh())
+		{
+			return false;
+		}
+	}
+	else if (!SkelMesh && InputInfo && InputInfo->GetSkeletalMesh().IsValid())
 	{
 		return false;
 	}
@@ -304,15 +392,47 @@ void UMLDeformerModelInstance::Tick(float DeltaTime, float ModelWeight)
 		PostMLDeformerComponentInit();
 	}
 
+	bool bExecuteCalled = false;
 	if (ModelWeight > 0.0001f && HasValidTransforms() && SetupInputs())
 	{
 		// Execute the model instance.
 		// For models using neural networks this will perform the inference, 
 		// calculate the network outputs and possibly use them, depending on how the model works.
 		Execute(ModelWeight);
+		bExecuteCalled = true;
 	}
 	else
 	{
 		HandleZeroModelWeight();
 	}
+
+	// Execute the post tick.
+	PostTick(bExecuteCalled);
 }
+
+#if WITH_EDITOR
+	void UMLDeformerModelInstance::CopyDataFromCurrentDebugActor()
+	{
+		UMLDeformerComponent* MLDeformerComponent = GetMLDeformerComponent();
+		if (MLDeformerComponent)
+		{
+			// Check if we're debugging at all.
+			const AActor* DebugActor = MLDeformerComponent->GetDebugActor();
+			if (DebugActor == nullptr)
+			{
+				return;
+			}
+
+			// Make sure we use the same ML Deformer asset and skeletal mesh etc.
+			UMLDeformerComponent* DebugMLDeformerComponent = nullptr;
+			if (!Model->IsCompatibleDebugActor(DebugActor, &DebugMLDeformerComponent))
+			{
+				return;
+			}
+			check(DebugMLDeformerComponent);
+
+			// Now copy the actual data.
+			CopyDataFromDebugActor(DebugActor, DebugMLDeformerComponent);
+		}
+	}
+#endif
diff --git a/Engine/Plugins/Animation/MLDeformer/MLDeformerFramework/Source/MLDeformerFramework/Private/MLDeformerMorphModelInstance.cpp b/Engine/Plugins/Animation/MLDeformer/MLDeformerFramework/Source/MLDeformerFramework/Private/MLDeformerMorphModelInstance.cpp
--- a/Engine/Plugins/Animation/MLDeformer/MLDeformerFramework/Source/MLDeformerFramework/Private/MLDeformerMorphModelInstance.cpp
+++ b/Engine/Plugins/Animation/MLDeformer/MLDeformerFramework/Source/MLDeformerFramework/Private/MLDeformerMorphModelInstance.cpp
@@ -23,6 +23,13 @@ bool UMLDeformerMorphModelInstance::IsValidForDataProvider() const
 	return true;
 }
 
+void UMLDeformerMorphModelInstance::PostTick(bool bExecuteCalled)
+{
+#if WITH_EDITOR
+	CopyDataFromCurrentDebugActor();
+#endif
+}
+
 int32 UMLDeformerMorphModelInstance::GetExternalMorphSetID() const
 { 
 	return ExternalMorphSetID;
@@ -31,7 +38,7 @@ int32 UMLDeformerMorphModelInstance::GetExternalMorphSetID() const
 void UMLDeformerMorphModelInstance::BeginDestroy()
 {
 	// Try to unregister the morph target and morph target set.
-	if (SkeletalMeshComponent)
+	if (IsValid(SkeletalMeshComponent))
 	{
 		const UMLDeformerMorphModel* MorphModel = Cast<UMLDeformerMorphModel>(Model);
 		if (MorphModel)
@@ -67,6 +74,7 @@ void UMLDeformerMorphModelInstance::PostMLDeformerComponentInit()
 
 		// Register the morph set. This overwrites the existing one for this model, if it already exists.
 		// Only add to LOD 0 for now.
+		check(IsInGameThread());	// We don't want to call this multithreaded, as the AddExternalMorphSets etc isn't thread safe.
 		const int32 LOD = 0;
 		SkelMeshComponent->AddExternalMorphSet(LOD, ExternalMorphSetID, MorphTargetSet);
 
@@ -104,6 +112,9 @@ void UMLDeformerMorphModelInstance::PostMLDeformerComponentInit()
 
 void UMLDeformerMorphModelInstance::Tick(float DeltaTime, float ModelWeight)
 {
+	// We don't want to call this multithreaded, as the AddExternalMorphSets etc isn't thread safe.
+	check(IsInGameThread());
+
 	// Detect changes in quality level.
 	const int32 CurrentQualityLevel = GetMLDeformerComponent()->GetQualityLevel();
 	if (CurrentQualityLevel != LastQualityLevel)
diff --git a/Engine/Plugins/Animation/MLDeformer/MLDeformerFramework/Source/MLDeformerFramework/Public/MLDeformerComponent.h b/Engine/Plugins/Animation/MLDeformer/MLDeformerFramework/Source/MLDeformerFramework/Public/MLDeformerComponent.h
--- a/Engine/Plugins/Animation/MLDeformer/MLDeformerFramework/Source/MLDeformerFramework/Public/MLDeformerComponent.h
+++ b/Engine/Plugins/Animation/MLDeformer/MLDeformerFramework/Source/MLDeformerFramework/Public/MLDeformerComponent.h
@@ -132,6 +132,15 @@ public:
 	UFUNCTION(BlueprintCallable, Category = "MLDeformer")
 	void UpdateSkeletalMeshComponent();
 
+	/**
+	 * Get the final weight that is used when applying the ML Deformer.
+	 * Some console command might override the weight that was set to this component. This method will
+	 * return the real weight that will be applied, after console command modifications are applied as well.
+	 * So this is the actual final weight used on the ML Deformer.
+	 * @return The ML Deformer weight value that is used to deform the mesh, where 0 means it is not doing any deformations and 1 means it is fully active.
+	 */
+	float GetFinalMLDeformerWeight() const;
+
 	/**
 	 * Suppress logging warnings about mesh deformers not being set.
 	 * A warning is logged when an ML Deformer is used that requires a deformer graph, but the skeletal mesh has no deformer graph setup.
@@ -144,6 +153,9 @@ public:
 	 * Get the performance counter that measures how much time is spent inside the Tick function.
 	 */
 	const UE::MLDeformer::FMLDeformerPerfCounter& GetTickPerfCounter() const	{ return TickPerfCounter; }
+
+	TObjectPtr<AActor> GetDebugActor() const					{ return DebugActor; }
+	void SetDebugActor(TObjectPtr<AActor> Actor)				{ DebugActor = Actor; }
 #endif
 
 	// Get property names.
@@ -156,6 +168,12 @@ protected:
 	virtual void TickComponent(float DeltaTime, enum ELevelTick TickType, FActorComponentTickFunction* ThisTickFunction) override;
 	// ~END AActorComponent overrides.
 
+	/** Set the ML Deformer weight. */
+	virtual void SetWeightInternal(const float NormalizedWeightValue);
+
+	/** Set the ML Deformer asset. */
+	virtual void SetDeformerAssetInternal(UMLDeformerAsset* const InDeformerAsset);
+
 	/** 
 	 * Initialize the component. 
 	 * This releases any existing deformer instance that is active, and creates a new one.
@@ -164,40 +182,32 @@ protected:
 	 */
 	void Init();
 
-	/** Bind to the MLDeformerModel's NeuralNetworkModifyDelegate. */
-	UE_DEPRECATED(5.2, "This method will be removed.")
-	void AddNeuralNetworkModifyDelegate() {}
-
-	/** Unbind from the MLDeformerModel's NeuralNetworkModifyDelegate. */
-	UE_DEPRECATED(5.2, "This method will be removed.")
-	void RemoveNeuralNetworkModifyDelegate() {}
-
 	UE_DEPRECATED(5.3, "This method will be removed.")
-	void AddReleaseModelInstancesDelegate();
+	void AddReleaseModelInstancesDelegate() {}
 	
 	UE_DEPRECATED(5.3, "This method will be removed.")
-	void RemoveReleaseModelInstancesDelegate();
+	void RemoveReleaseModelInstancesDelegate() {}
 
 	void BindDelegates();
 	void UnbindDelegates();
-
 	void ReleaseModelInstance();
-
-	/** Set the ML Deformer weight. */
-	virtual void SetWeightInternal(const float NormalizedWeightValue);
-
-	/** Set the ML Deformer asset. */
-	virtual void SetDeformerAssetInternal(UMLDeformerAsset* const InDeformerAsset);
-
+	
 #if WITH_EDITOR
 	/** Reset the tick cycle counters. */
-	void ResetTickCycleCounters();
+	UE_DEPRECATED(5.4, "This method will be removed")
+	void ResetTickCycleCounters() {}
 #endif
 
 protected:
 #if WITH_EDITOR
 	/** The performance counter that measures timing of the Tick function. */
 	UE::MLDeformer::FMLDeformerPerfCounter TickPerfCounter;
+
+	/**
+	 * The actor we are currently debugging. This can be used to copy over specific information from another actor, such as copying over external morph target weights.
+	 * When set to a nullptr, then we're not debugging. This actor most likely is inside another UWorld.
+	 */
+	TObjectPtr<AActor> DebugActor;
 #endif
 
 	/** Render command fence that let's us wait for all other commands to finish. */
@@ -208,7 +218,7 @@ protected:
 	 * This can be a nullptr. When it is a nullptr then it will internally try to find the first skeletal mesh component on the actor.
 	 * You can see this as an override. You can specify this override through the SetupComponent function.
 	 */
-	TObjectPtr<USkeletalMeshComponent> SkelMeshComponent = nullptr;
+	TObjectPtr<USkeletalMeshComponent> SkelMeshComponent;
 
 	/** DelegateHandle for NeuralNetwork modification. This has been deprecated. */
 	UE_DEPRECATED(5.3, "This member has been deprecated.")
@@ -222,7 +232,7 @@ protected:
 
 	/** The deformer asset to use. */
 	UPROPERTY(EditAnywhere, DisplayName = "ML Deformer Asset", Category = "ML Deformer")
-	TObjectPtr<UMLDeformerAsset> DeformerAsset = nullptr;
+	TObjectPtr<UMLDeformerAsset> DeformerAsset;
 
 	/** How active is this deformer? Can be used to blend it in and out. */
 	UPROPERTY(EditAnywhere, Category = "ML Deformer", meta = (ClampMin = "0.0", ClampMax = "1.0"))
@@ -240,5 +250,5 @@ protected:
 
 	/** The deformation model instance. This is used to perform the runtime updates and run the inference. */
 	UPROPERTY(Transient)
-	TObjectPtr<UMLDeformerModelInstance> ModelInstance = nullptr;
+	TObjectPtr<UMLDeformerModelInstance> ModelInstance;
 };
diff --git a/Engine/Plugins/Animation/MLDeformer/MLDeformerFramework/Source/MLDeformerFramework/Public/MLDeformerModel.h b/Engine/Plugins/Animation/MLDeformer/MLDeformerFramework/Source/MLDeformerFramework/Public/MLDeformerModel.h
--- a/Engine/Plugins/Animation/MLDeformer/MLDeformerFramework/Source/MLDeformerFramework/Public/MLDeformerModel.h
+++ b/Engine/Plugins/Animation/MLDeformer/MLDeformerFramework/Source/MLDeformerFramework/Public/MLDeformerModel.h
@@ -205,6 +205,16 @@ public:
 	 */
 	void SetSkeletalMesh(USkeletalMesh* SkelMesh)			{ SkeletalMesh = SkelMesh; }
 
+	/**
+	 * Check if a given actor would be a compatible debugging actor.
+	 * We check this by checking if it uses the same skeletal mesh and whether it uses the same ML Deformer asset.
+	 * Debugging allows us to copy over the morph weights and pose of a character. But for that to be possible we need to 
+	 * make sure the other actor uses the same skeletal mesh and ML Deformer asset. This method helps us check that easily.
+	 * @param Actor The actor to check compatibility with.
+	 * @param OutDebugComponent A pointer to the ML Deformer component that we would be debugging. Can be set to nullptr to get it ignored.
+	 * @return Returns true when the provided actor is compatible for debugging, otherwise false is returned.
+	 */
+	virtual bool IsCompatibleDebugActor(const AActor* Actor, UMLDeformerComponent** OutDebugComponent = nullptr) const;
 
 #if WITH_EDITORONLY_DATA
 	/**
diff --git a/Engine/Plugins/Animation/MLDeformer/MLDeformerFramework/Source/MLDeformerFramework/Public/MLDeformerModelInstance.h b/Engine/Plugins/Animation/MLDeformer/MLDeformerFramework/Source/MLDeformerFramework/Public/MLDeformerModelInstance.h
--- a/Engine/Plugins/Animation/MLDeformer/MLDeformerFramework/Source/MLDeformerFramework/Public/MLDeformerModelInstance.h
+++ b/Engine/Plugins/Animation/MLDeformer/MLDeformerFramework/Source/MLDeformerFramework/Public/MLDeformerModelInstance.h
@@ -69,6 +69,14 @@ public:
 	 */
 	virtual void Tick(float DeltaTime, float ModelWeight);
 
+	/** 
+	 * Executed at the end of the tick. This is always executed when Tick is called, also if the model weight is 0 etc. 
+	 * When the model weight is tiny or setting the inputs failed, the Execute function wasn't called during Tick.
+	 * The parameter that is passed to this method will let you know whether Execute got called or not.
+	 * @param bExecuteCalled This is true when we called Execute, otherwise it is false.
+	 */
+	virtual void PostTick(bool bExecuteCalled) {}
+
 	/**
 	 * Check whether the deformer is compatible with a given skeletal mesh component.
 	 * This internally also edits the value returned by GetCompatibilityErrorText().
@@ -87,6 +95,25 @@ public:
 	 */
 	virtual bool IsValidForDataProvider() const;
 
+#if WITH_EDITOR
+	/**
+	 * Copy over data from the debug actor, if any is used.
+	 * You can get the debug actor using UMLDeformerComponent::GetDebugActor().
+	 * This should make calls to the CopyDataFromDebugActor, which takes the actor pointer and ML Deformer component as parameters.
+	 * Compatibility checks between the debug actor and the current model instance should be done inside this method.
+	 * This could for example copy over morph target weights etc.
+	 */
+	virtual void CopyDataFromCurrentDebugActor();
+
+	/**
+	 * Copy data from a specific debug actor.
+	 * This method already assumes DebugActor and MLDeformerComponent are compatible with the current model instance.
+	 * @param DebugActor The debug actor object to copy data from.
+	 * @param DebugMLDeformerComponent The debug ML Deformer component to copy data from.
+	 */
+	virtual void CopyDataFromDebugActor(const AActor* DebugActor, const UMLDeformerComponent* DebugMLDeformerComponent) {}
+#endif
+
 	/**
 	 * Is the deformer asset used compatible with the skeletal mesh component used during the Init call?
 	 * @return True if compatible, false if not.
@@ -133,6 +160,14 @@ public:
 	 */
 	UMLDeformerComponent* GetMLDeformerComponent() const;
 
+	/**
+	 * Get the skeletal mesh component to get transforms from etc.
+	 * If we are debugging this will return the debugged skeletal mesh component, otherwise the skel mesh component
+	 * of our own actor is returned.
+	 * @return The skeletal mesh component to grab input data from.
+	 */
+	USkeletalMeshComponent* GetFinalSkeletalMeshComponent() const;
+
 protected:
 	/**
 	 * Update the neural network input values directly inside its input tensor.
@@ -170,12 +205,6 @@ protected:
 	 */
 	virtual void HandleZeroModelWeight() {}
 
-	/**
-	 * Get the current component space bone transforms that we grabbed from the skeletal mesh component.
-	 * @return The cached bone transforms, in component space.
-	 */
-	const TArray<FTransform>& GetBoneTransforms() const;
-
 	/**
 	 * Set the bone transformations inside a given output buffer, starting from a given StartIndex.
 	 * @param OutputBuffer The buffer we need to write the transformations to.
@@ -220,8 +249,8 @@ protected:
 	/** The cached current local space bone transforms for the current frame. */
 	TArray<FTransform> TrainingBoneTransforms;
 
-	/** A temp array of bone transforms. */
-	TArray<FTransform> BoneTransforms;
+	/** The debug bone space transforms. */
+	TArray<FTransform> DebugBoneSpaceTransforms;
 
 	/** Maps the ML deformer asset bone index to a skeletal mesh component bone index. */
 	TArray<int32> AssetBonesToSkelMeshMappings;
diff --git a/Engine/Plugins/Animation/MLDeformer/MLDeformerFramework/Source/MLDeformerFramework/Public/MLDeformerMorphModelInstance.h b/Engine/Plugins/Animation/MLDeformer/MLDeformerFramework/Source/MLDeformerFramework/Public/MLDeformerMorphModelInstance.h
--- a/Engine/Plugins/Animation/MLDeformer/MLDeformerFramework/Source/MLDeformerFramework/Public/MLDeformerMorphModelInstance.h
+++ b/Engine/Plugins/Animation/MLDeformer/MLDeformerFramework/Source/MLDeformerFramework/Public/MLDeformerMorphModelInstance.h
@@ -31,6 +31,7 @@ public:
 	virtual void Tick(float DeltaTime, float ModelWeight) override;
 	virtual void HandleZeroModelWeight() override;
 	virtual bool IsValidForDataProvider() const override;
+	virtual void PostTick(bool bExecuteCalled) override;
 	// ~END UMLDeformerModelInstance overrides.
 
 	int32 GetExternalMorphSetID() const;
@@ -42,6 +43,16 @@ public:
 	 */
 	FExternalMorphSetWeights* FindWeightData(int32 LOD) const;
 
+protected:
+#if WITH_EDITOR
+	/**
+	 * Copy the morph weights from a given morph model instance.
+	 * We assume that the other instance is compatible with this one.
+	 * @param SourceModelInstance The instance to copy the weights from.
+	 */
+	void CopyMorphWeightsFromMorphInstance(const UMLDeformerMorphModelInstance* SourceModelInstance);
+#endif
+
 protected:
 	/** The next free morph target set ID. This is used to generate unique ID's for each morph model. */
 	static TAtomic<int32> NextFreeMorphSetID;
diff --git a/Engine/Plugins/Animation/MLDeformer/MLDeformerFramework/Source/MLDeformerFramework/Public/MLDeformerVizSettings.h b/Engine/Plugins/Animation/MLDeformer/MLDeformerFramework/Source/MLDeformerFramework/Public/MLDeformerVizSettings.h
--- a/Engine/Plugins/Animation/MLDeformer/MLDeformerFramework/Source/MLDeformerFramework/Public/MLDeformerVizSettings.h
+++ b/Engine/Plugins/Animation/MLDeformer/MLDeformerFramework/Source/MLDeformerFramework/Public/MLDeformerVizSettings.h
@@ -90,6 +90,7 @@ public:
 	bool GetDrawMLDeformedActor() const						{ return bDrawMLDeformedActor; }
 	bool GetDrawGroundTruthActor() const					{ return bDrawGroundTruthActor; }
 	bool GetDrawMLCompareActors() const						{ return bDrawMLCompareActors; }
+	bool GetDrawDebugActorBounds() const					{ return bDrawDebugActorBounds; }
 	bool GetShowHeatMap() const								{ return bShowHeatMap; }
 	EMLDeformerHeatMapMode GetHeatMapMode() const			{ return HeatMapMode; }
 	float GetHeatMapMax() const								{ return HeatMapMax; }
@@ -101,6 +102,7 @@ public:
 	int32 GetQualityLevel() const							{ return QualityLevel; }
 	const TArray<FMLDeformerCompareActor>& GetCompareActors() const { return CompareActors; }
 	TArray<FMLDeformerCompareActor>& GetCompareActors()		{ return CompareActors; }
+	FColor GetDebugBoundsColor() const						{ return DebugBoundsColor; }
 
 	// Get property names.
 	static FName GetVisualizationModePropertyName()			{ return GET_MEMBER_NAME_CHECKED(UMLDeformerVizSettings, VisualizationMode); }
@@ -116,6 +118,7 @@ public:
 	static FName GetDrawMLDeformedActorPropertyName()		{ return GET_MEMBER_NAME_CHECKED(UMLDeformerVizSettings, bDrawMLDeformedActor); }
 	static FName GetDrawGroundTruthActorPropertyName()		{ return GET_MEMBER_NAME_CHECKED(UMLDeformerVizSettings, bDrawGroundTruthActor); }
 	static FName GetDrawMLCompareActorsPropertyName()		{ return GET_MEMBER_NAME_CHECKED(UMLDeformerVizSettings, bDrawMLCompareActors); }
+	static FName GetDrawDebugActorBoundsPropertyName()		{ return GET_MEMBER_NAME_CHECKED(UMLDeformerVizSettings, bDrawDebugActorBounds); }
 	static FName GetShowHeatMapPropertyName()				{ return GET_MEMBER_NAME_CHECKED(UMLDeformerVizSettings, bShowHeatMap); }
 	static FName GetHeatMapModePropertyName()				{ return GET_MEMBER_NAME_CHECKED(UMLDeformerVizSettings, HeatMapMode); }
 	static FName GetHeatMapMaxPropertyName()				{ return GET_MEMBER_NAME_CHECKED(UMLDeformerVizSettings, HeatMapMax); }
@@ -126,6 +129,7 @@ public:
 	static FName GetDrawVertexDeltasPropertyName()			{ return GET_MEMBER_NAME_CHECKED(UMLDeformerVizSettings, bDrawDeltas); }
 	static FName GetQualityLevelPropertyName()				{ return GET_MEMBER_NAME_CHECKED(UMLDeformerVizSettings, QualityLevel); }
 	static FName GetCompareActorsPropertyName()				{ return GET_MEMBER_NAME_CHECKED(UMLDeformerVizSettings, CompareActors); }
+	static FName GetDebugBoundsColorPropertyName()			{ return GET_MEMBER_NAME_CHECKED(UMLDeformerVizSettings, DebugBoundsColor); }
 #endif
 
 protected:
@@ -227,6 +231,17 @@ protected:
 	UPROPERTY(EditAnywhere, Category = "Live Settings", DisplayName = "Draw ML Compare Actors")
 	bool bDrawMLCompareActors = true;
 
+	/** Draw the debug actor's bounds in the PIE viewport? This only renders when PIE is active. */
+	UPROPERTY(EditAnywhere, Category = "Live Settings", DisplayName = "Highlight Debug Actors in PIE")
+	bool bDrawDebugActorBounds = true;
+
+	/** 
+	 * The color of the bounding box rendered inside the PIE viewport.
+	 * Keep in mind that the actor you are currently debugging is always rendered in green.
+	 */
+	UPROPERTY(EditAnywhere, Category = "Live Settings")
+	FColor DebugBoundsColor = FColor::Purple;
+
 	/** The scale factor of the ML deformer deltas being applied on top of the linear skinned results. */
 	UPROPERTY(EditAnywhere, Transient, Category = "Live Settings", meta = (ClampMin = "0.0", ClampMax = "1.0"))
 	float Weight = 1.0f;
diff --git a/Engine/Plugins/Animation/MLDeformer/MLDeformerFramework/Source/MLDeformerFrameworkEditor/Private/MLDeformerEditorModel.cpp b/Engine/Plugins/Animation/MLDeformer/MLDeformerFramework/Source/MLDeformerFrameworkEditor/Private/MLDeformerEditorModel.cpp
--- a/Engine/Plugins/Animation/MLDeformer/MLDeformerFramework/Source/MLDeformerFrameworkEditor/Private/MLDeformerEditorModel.cpp
+++ b/Engine/Plugins/Animation/MLDeformer/MLDeformerFramework/Source/MLDeformerFrameworkEditor/Private/MLDeformerEditorModel.cpp
@@ -43,6 +43,10 @@
 #include "Engine/SkinnedAssetCommon.h"
 #include "SMLDeformerInputWidget.h"
 #include "SMLDeformerTimeline.h"
+#include "SMLDeformerDebugSelectionWidget.h"
+#include "EditorViewportClient.h"
+#include "DrawDebugHelpers.h"
+#include "Slate/SceneViewport.h"
 
 #define LOCTEXT_NAMESPACE "MLDeformerEditorModel"
 
@@ -391,6 +395,8 @@ namespace UE::MLDeformer
 
 	void FMLDeformerEditorModel::UpdateMeshOffsetFactors()
 	{
+		const bool bIsDebugging = (GetEditor()->GetDebugActor() != nullptr);
+
 		// Set the default mesh translation offsets for our ground truth actors.
 		for (FMLDeformerEditorActor* EditorActor : EditorActors)
 		{
@@ -412,7 +418,7 @@ namespace UE::MLDeformer
 				{
 					MeshOffsetFactor = 2.0f;
 
-					if (Model->GetVizSettings()->GetDrawMLCompareActors())
+					if (Model->GetVizSettings()->GetDrawMLCompareActors() && !bIsDebugging)
 					{
 						const TArray<FMLDeformerCompareActor>& CompareActors = Model->GetVizSettings()->GetCompareActors();
 						const int32 NumValidCompareActors = static_cast<float>(CalcNumValidCompareActorsPriorTo(CompareActors.Num()));
@@ -541,7 +547,7 @@ namespace UE::MLDeformer
 	}
 
 	void FMLDeformerEditorModel::Tick(FEditorViewportClient* ViewportClient, float DeltaTime)
-	{		
+	{
 		if (bNeedsAssetReinit)
 		{
 			TriggerInputAssetChanged();
@@ -567,6 +573,20 @@ namespace UE::MLDeformer
 		UpdateActorTransforms();
 		UpdateLabels();
 		CheckTrainingDataFrameChanged();
+		ApplyDebugActorTransforms();
+
+		// Debug draw elements inside the PIE viewport when PIE is active.
+		DrawPIEDebugActors();
+
+		// We need to tell the viewport that it needs to redraw, because otherwise while PIE is active
+		// it won't redraw the asset editor unless you interact with the viewport using the mouse.
+		if (GEditor->GetPIEViewport() && GetEditor()->GetViewport().IsValid())
+		{
+			if (Model->GetVizSettings()->GetDrawDebugActorBounds() || IsValid(GetEditor()->GetDebugActor()))
+			{
+				GetEditor()->GetViewport()->GetViewportClient().bNeedsRedraw = true;
+			}
+		}
 
 		// Update the ML Deformer components.
 		for (FMLDeformerEditorActor* EditorActor : EditorActors)
@@ -646,6 +666,7 @@ namespace UE::MLDeformer
 		const UMLDeformerVizSettings* VizSettings = Model->GetVizSettings();
 		const bool bShowTrainingData = (VizSettings->GetVisualizationMode() == EMLDeformerVizMode::TrainingData);
 		const bool bShowTestData = (VizSettings->GetVisualizationMode() == EMLDeformerVizMode::TestData);
+		const bool bIsDebugging = (GetEditor()->GetDebugActor() != nullptr);
 		for (FMLDeformerEditorActor* EditorActor : EditorActors)
 		{
 			if (EditorActor)
@@ -661,11 +682,11 @@ namespace UE::MLDeformer
 				}
 				else if (EditorActor->GetTypeID() == ActorID_Test_GroundTruth)
 				{
-					bIsVisible &= VizSettings->GetDrawGroundTruthActor();
+					bIsVisible &= VizSettings->GetDrawGroundTruthActor() && !bIsDebugging;
 				}
 				else if (EditorActor->GetTypeID() == ActorID_Test_Compare)
 				{
-					bIsVisible &= VizSettings->GetDrawMLCompareActors();
+					bIsVisible &= VizSettings->GetDrawMLCompareActors() && !bIsDebugging;
 				}
 
 				bIsVisible &= EditorActor->HasVisualMesh();
@@ -1646,6 +1667,53 @@ namespace UE::MLDeformer
 				}
 			}
 		}
+
+	}
+
+	void FMLDeformerEditorModel::DrawPIEDebugActors()
+	{
+		// Draw all debuggable actor bounds and names.
+		UMLDeformerVizSettings* VizSettings = Model->GetVizSettings();
+		if (Editor && VizSettings->GetDrawDebugActorBounds() && Editor->GetDebugWidget().IsValid())
+		{
+			const AActor* ActiveDebugActor = Editor->GetDebugActor();
+
+			FString Name;
+			TSharedPtr<SMLDeformerDebugSelectionWidget> DebugWidget = Editor->GetDebugWidget();
+			for (TSharedPtr<FMLDeformerDebugActor> DebuggableActor : DebugWidget->GetActors())
+			{
+				if (!DebuggableActor.IsValid() || 
+					!IsValid(DebuggableActor->Actor) || 
+					!DebuggableActor->Actor->GetOuter())
+				{
+					continue;
+				}
+
+				// Get the actor's bounding box.
+				const UWorld* World = DebuggableActor->Actor->GetWorld();
+				FVector Center;
+				FVector Extents;
+				DebuggableActor->Actor->GetActorBounds(true, Center, Extents, false);
+
+				const FVector TextStart = DebuggableActor->Actor->GetTransform().GetTranslation();
+				Name = DebuggableActor->Actor->GetName();
+				FColor TextColor;
+				FColor BoxColor;
+				if (DebuggableActor->Actor == ActiveDebugActor)
+				{
+					TextColor = FColor::Yellow;
+					BoxColor = FColor::Green;
+				}
+				else
+				{
+					TextColor = FColor::White;
+					BoxColor = VizSettings->GetDebugBoundsColor();
+				}
+
+				DrawDebugBox(World, Center, Extents, BoxColor, false);
+				DrawDebugString(World, TextStart, Name, nullptr, TextColor, 0.0f, false, 1.0f);
+			}
+		}
 	}
 
 	void FMLDeformerEditorModel::SampleDeltas()
@@ -2828,6 +2896,70 @@ namespace UE::MLDeformer
 		// Set the current training frame to -1, which will cause the deltas to get updated next frame.
 		CurrentTrainingFrame = -1;
 	}
+
+	void FMLDeformerEditorModel::ApplyDebugActorTransforms()
+	{
+		AActor* DebugActor = GetEditor()->GetDebugActor();
+		if (DebugActor && IsValid(DebugActor))
+		{
+			const TArray<FTransform> DebugActorComponentSpaceTransforms = GetEditor()->GetDebugActorComponentSpaceTransforms();
+			if (!DebugActorComponentSpaceTransforms.IsEmpty())
+			{
+				ApplyDebugActorTransforms(DebugActorComponentSpaceTransforms);
+			}
+		}
+	}
+
+	void FMLDeformerEditorModel::OnDebugActorChanged(TObjectPtr<AActor> DebugActor)
+	{
+		// Mark the test ML Deformed character as having debugging enabled.
+		FMLDeformerEditorActor* TestMLDeformedActor = FindEditorActor(ActorID_Test_MLDeformed);
+		if (TestMLDeformedActor)
+		{
+			UMLDeformerComponent* MLDComponent = TestMLDeformedActor->GetMLDeformerComponent();
+			if (MLDComponent)
+			{
+				MLDComponent->SetDebugActor(DebugActor);
+			}
+		}
+
+		// Update visibility, mesh offsets and transforms.
+		// This is needed as when we enable or disable debugging we can show/hide compare actors, ground truth actors, etc.
+		UpdateActorVisibility();
+		UpdateMeshOffsetFactors();
+		UpdateActorTransforms();
+
+		// Trigger a redraw of the MLD asset editor viewport.
+		if (GetEditor()->GetViewport().IsValid())
+		{
+			GetEditor()->GetViewport()->GetViewportClient().bNeedsRedraw = true;
+		}
+	}
+
+	void FMLDeformerEditorModel::ApplyDebugActorTransforms(const TArray<FTransform>& DebugActorComponentSpaceTransforms)
+	{
+		for (FMLDeformerEditorActor* EditorActor : EditorActors)
+		{
+			// Only overwrite the transforms for the test ML Deformed character and the linear skinned character, leave the rest as-is.
+			// It doesn't make sense to overwrite transforms for the training actors for example.
+			// We also shouldn't apply to the compare actors, as the morph weights aren't copied over, as they might have a different number of morphs etc.
+			if (EditorActor == nullptr || 
+				(EditorActor->GetTypeID() != ActorID_Test_MLDeformed && EditorActor->GetTypeID() != ActorID_Test_Base))
+			{
+				continue;
+			}
+
+			// Copy the transforms to all editor actors that have a skeletal mesh component with valid skeletal mesh setup.
+			UDebugSkelMeshComponent* SkelMeshComponent = EditorActor->GetSkeletalMeshComponent();
+			if (SkelMeshComponent && SkelMeshComponent->GetSkeletalMeshAsset())
+			{
+				TArray<FTransform>& EditableTransforms = SkelMeshComponent->GetEditableComponentSpaceTransforms();
+				check(EditableTransforms.Num() == DebugActorComponentSpaceTransforms.Num());
+				EditableTransforms = DebugActorComponentSpaceTransforms;
+				SkelMeshComponent->ApplyEditedComponentSpaceTransforms();
+			}
+		}
+	}
 }	// namespace UE::MLDeformer
 
 #undef LOCTEXT_NAMESPACE
diff --git a/Engine/Plugins/Animation/MLDeformer/MLDeformerFramework/Source/MLDeformerFrameworkEditor/Private/MLDeformerEditorStyle.cpp b/Engine/Plugins/Animation/MLDeformer/MLDeformerFramework/Source/MLDeformerFrameworkEditor/Private/MLDeformerEditorStyle.cpp
--- a/Engine/Plugins/Animation/MLDeformer/MLDeformerFramework/Source/MLDeformerFrameworkEditor/Private/MLDeformerEditorStyle.cpp
+++ b/Engine/Plugins/Animation/MLDeformer/MLDeformerFramework/Source/MLDeformerFrameworkEditor/Private/MLDeformerEditorStyle.cpp
@@ -27,6 +27,7 @@ namespace UE::MLDeformer
 		Set("MLDeformer.Timeline.PauseIcon", new MLDEFORMER_IMAGE_BRUSH("Icons/Pause_40x", Icon20x20));
 		Set("MLDeformer.Timeline.TabIcon", new MLDEFORMER_IMAGE_BRUSH("Icons/TimelineTab_40x", Icon16x16));
 		Set("MLDeformer.VizSettings.TabIcon", new MLDEFORMER_IMAGE_BRUSH("Icons/VizSettingsTab_40x", Icon40x40));
+		Set("MLDeformer.Debug.RefreshIcon", new MLDEFORMER_IMAGE_BRUSH("Icons/Refresh_40x", Icon40x40));
 
 		// Colors and sizes.
 		Set("MLDeformer.BaseMesh.WireframeColor", FLinearColor(0.0f, 1.0f, 0.0f));
@@ -63,6 +64,9 @@ namespace UE::MLDeformer
 		Set("MLDeformer.InputsWidget.HighlightColor", FLinearColor(0.0f, 1.0f, 0.0f));
 		Set("MLDeformer.InputsWidget.ErrorColor", FLinearColor(1.0f, 0.5f, 0.0f));
 
+		Set("MLDeformer.Debug.ActiveDebugColor", FLinearColor(0.0f, 1.0f, 0.0f));
+		Set("MLDeformer.Debug.EditorSelectedMarkColor", FLinearColor(1.0f, 1.0f, 0.0f));
+
 		FSlateStyleRegistry::RegisterSlateStyle(*this);
 	}
 
diff --git a/Engine/Plugins/Animation/MLDeformer/MLDeformerFramework/Source/MLDeformerFrameworkEditor/Private/MLDeformerEditorToolkit.cpp b/Engine/Plugins/Animation/MLDeformer/MLDeformerFramework/Source/MLDeformerFrameworkEditor/Private/MLDeformerEditorToolkit.cpp
--- a/Engine/Plugins/Animation/MLDeformer/MLDeformerFramework/Source/MLDeformerFrameworkEditor/Private/MLDeformerEditorToolkit.cpp
+++ b/Engine/Plugins/Animation/MLDeformer/MLDeformerFramework/Source/MLDeformerFrameworkEditor/Private/MLDeformerEditorToolkit.cpp
@@ -13,6 +13,8 @@
 #include "MLDeformerEditorStyle.h"
 #include "MLDeformerVizSettings.h"
 #include "MLDeformerSampler.h"
+#include "SMLDeformerTimeline.h"
+#include "SMLDeformerDebugSelectionWidget.h"
 #include "AnimationEditorViewportClient.h"
 #include "EditorModeManager.h"
 #include "EditorViewportClient.h"
@@ -24,7 +26,6 @@
 #include "Preferences/PersonaOptions.h"
 #include "UObject/Object.h"
 #include "SSimpleTimeSlider.h"
-#include "SMLDeformerTimeline.h"
 #include "Framework/MultiBox/MultiBoxBuilder.h"
 #include "Framework/Notifications/NotificationManager.h"
 #include "Misc/MessageDialog.h"
@@ -97,6 +98,10 @@ namespace UE::MLDeformer
 		FMLDeformerEditorMode* EditorMode = static_cast<FMLDeformerEditorMode*>(GetEditorModeManager().GetActiveMode(FMLDeformerEditorMode::ModeName));
 		EditorMode->SetEditorToolkit(this);
 
+		SAssignNew(DebugWidget, SMLDeformerDebugSelectionWidget)
+			.MLDeformerEditor(this)
+			.Visibility(this, &FMLDeformerEditorToolkit::GetDebuggingVisibility);
+
 		ExtendToolbar();
 		RegenerateMenusAndToolbars();
 
@@ -115,6 +120,11 @@ namespace UE::MLDeformer
 			GetVizSettingsDetailsView()->SetObject(DeformerModel->GetVizSettings());
 		}
 
+		if (DebugWidget.IsValid())
+		{
+			DebugWidget->Refresh();
+		}
+
 		GetModelDetailsView()->ForceRefresh();
 		GetVizSettingsDetailsView()->ForceRefresh();
 
@@ -391,6 +401,11 @@ namespace UE::MLDeformer
 		return bIsTraining;
 	}
 
+	AActor* FMLDeformerEditorToolkit::GetDebugActor() const
+	{
+		return DebugWidget.IsValid() ? DebugWidget->GetDebugActor() : nullptr;
+	}
+
 	bool FMLDeformerEditorToolkit::Train(bool bSuppressDialogs)
 	{
 		check(ActiveModel);
@@ -456,12 +471,12 @@ namespace UE::MLDeformer
 
 	void FMLDeformerEditorToolkit::FillToolbar(FToolBarBuilder& ToolbarBuilder)
 	{
+		// Training button and model selection.
 		ToolbarBuilder.BeginSection("Training");
 		{
 			ToolbarBuilder.AddToolBarButton(
 				FUIAction
 				(
-					//FExecuteAction::CreateRaw(this, &FMLDeformerEditorToolkit::Train, false),
 					FExecuteAction::CreateLambda
 					(
 						[this]()
@@ -497,6 +512,7 @@ namespace UE::MLDeformer
 		}
 		ToolbarBuilder.EndSection();
 
+		// Tools.
 		if (!ToolsMenuExtenders.IsEmpty())
 		{
 			ToolbarBuilder.BeginSection("Tools");
@@ -517,6 +533,44 @@ namespace UE::MLDeformer
 			}
 			ToolbarBuilder.EndSection();
 		}
+
+		// Debugging.
+		ToolbarBuilder.BeginSection("Debugging");
+		{
+			ToolbarBuilder.AddWidget(DebugWidget.ToSharedRef());
+			ToolbarBuilder.AddToolBarButton
+			(
+				FUIAction
+				(
+					FExecuteAction::CreateLambda([this]() { if (DebugWidget.IsValid()) { DebugWidget->Refresh(); } }),
+					FCanExecuteAction::CreateLambda([](){ return true; }),
+					FGetActionCheckState::CreateLambda([](){ return ECheckBoxState::Checked; }),
+					FIsActionButtonVisible::CreateLambda([this](){ return GetDebuggingVisibility() == EVisibility::Visible; })
+				),
+				NAME_None,
+				FText(),
+				LOCTEXT("RefreshDebugTooltip", "Refresh the list of debuggable actors."),
+				FSlateIcon(FMLDeformerEditorStyle::Get().GetStyleSetName(), "MLDeformer.Debug.RefreshIcon"),
+				EUserInterfaceActionType::Button
+			);
+		}
+		ToolbarBuilder.EndSection();
+	}
+
+	EVisibility FMLDeformerEditorToolkit::GetDebuggingVisibility() const
+	{
+		if (ActiveModel && ActiveModel->GetModel() && ActiveModel->GetModel()->GetVizSettings())
+		{
+			if (ActiveModel->GetModel()->GetVizSettings()->GetVisualizationMode() == EMLDeformerVizMode::TestData)
+			{
+				return EVisibility::Visible;
+			}
+			else
+			{
+				return EVisibility::Hidden;
+			}
+		}
+		return EVisibility::Visible;
 	}
 
 	bool FMLDeformerEditorToolkit::HandleTrainingResult(ETrainingResult TrainingResult, double TrainingDuration, bool& bOutUsePartiallyTrained, bool bSuppressDialogs, bool& bOutSuccess)
@@ -793,6 +847,68 @@ namespace UE::MLDeformer
 		ZoomOnActors();
 	}
 
+	TArray<FTransform> FMLDeformerEditorToolkit::GetDebugActorComponentSpaceTransforms() const
+	{
+		const FMLDeformerEditorModel* EditorModel = GetActiveModel();
+		if (EditorModel)
+		{
+			const AActor* DebugActor = GetDebugActor();
+			if (DebugActor && EditorModel->GetModel())
+			{				
+				const USkeletalMesh* ModelSkelMesh = EditorModel->GetModel()->GetSkeletalMesh();
+				if (ModelSkelMesh)
+				{
+					for (const UActorComponent* Component : DebugActor->GetComponents())
+					{
+						const USkeletalMeshComponent* DebugActorSkelMeshComponent = Cast<USkeletalMeshComponent>(Component);
+						if (DebugActorSkelMeshComponent)
+						{
+							if (DebugActorSkelMeshComponent->GetSkeletalMeshAsset() != ModelSkelMesh)
+							{
+								continue;
+							}
+
+							const USkinnedMeshComponent* LeaderPoseComponent = DebugActorSkelMeshComponent->LeaderPoseComponent.Get();
+							if (LeaderPoseComponent && !LeaderPoseComponent->GetComponentSpaceTransforms().IsEmpty())
+							{
+								const FReferenceSkeleton& RefSkel = ModelSkelMesh->GetRefSkeleton();
+								const int32 NumBones = RefSkel.GetNum();
+
+								TArray<FTransform> OutTransforms;	// TODO: Create some cached reusable buffer?
+								OutTransforms.SetNumUninitialized(NumBones);
+
+								const TArray<FTransform>& FollowerComponentTransforms = DebugActorSkelMeshComponent->GetComponentSpaceTransforms();
+								const TArray<FTransform>& LeaderComponentTransforms = LeaderPoseComponent->GetComponentSpaceTransforms();
+
+								const TArray<int32>& BoneMap = DebugActorSkelMeshComponent->GetLeaderBoneMap();
+								for (int32 Index = 0; Index < NumBones; ++Index)
+								{									
+									const int32 LeaderTransformIndex = BoneMap[Index];
+									if (LeaderTransformIndex != INDEX_NONE)
+									{
+										OutTransforms[Index] = LeaderComponentTransforms[LeaderTransformIndex];
+									}
+									else
+									{
+										OutTransforms[Index] = FollowerComponentTransforms[Index];
+									}
+								}
+
+								return MoveTemp(OutTransforms);
+							}
+							else
+							{
+								return DebugActorSkelMeshComponent->GetComponentSpaceTransforms();
+							}
+						}
+					}
+				}
+			}
+		}
+
+		return TArray<FTransform>();
+	}
+
 	void FMLDeformerEditorToolkit::ZoomOnActors()
 	{
 		// Zoom in on the actors.
diff --git a/Engine/Plugins/Animation/MLDeformer/MLDeformerFramework/Source/MLDeformerFrameworkEditor/Private/MLDeformerMorphModelEditorModel.cpp b/Engine/Plugins/Animation/MLDeformer/MLDeformerFramework/Source/MLDeformerFrameworkEditor/Private/MLDeformerMorphModelEditorModel.cpp
--- a/Engine/Plugins/Animation/MLDeformer/MLDeformerFramework/Source/MLDeformerFrameworkEditor/Private/MLDeformerMorphModelEditorModel.cpp
+++ b/Engine/Plugins/Animation/MLDeformer/MLDeformerFramework/Source/MLDeformerFrameworkEditor/Private/MLDeformerMorphModelEditorModel.cpp
@@ -136,30 +136,33 @@ namespace UE::MLDeformer
 		FMLDeformerGeomCacheEditorModel::OnPostTraining(TrainingResult, bUsePartiallyTrainedWhenAborted);
 	}
 
-	float CalcStandardDeviation(TArrayView<const float> Values)
+	namespace
 	{
-		if (Values.IsEmpty())
+		float CalcStandardDeviation(TArrayView<const float> Values)
 		{
-			return 0.0f;
-		}
+			if (Values.IsEmpty())
+			{
+				return 0.0f;
+			}
 
-		// First calculate the mean.
-		float Mean = 0.0f;
-		for (float Value : Values)
-		{
-			Mean += Value;
-		}
-		Mean /= static_cast<float>(Values.Num());
+			// First calculate the mean.
+			float Mean = 0.0f;
+			for (float Value : Values)
+			{
+				Mean += Value;
+			}
+			Mean /= static_cast<float>(Values.Num());
 
-		// Now calculate the standard deviation.
-		float Sum = 0.0f;
-		for (float Value : Values)
-		{
-			Sum += FMath::Square(Value - Mean);
-		}
-		Sum /= static_cast<float>(Values.Num());
+			// Now calculate the standard deviation.
+			float Sum = 0.0f;
+			for (float Value : Values)
+			{
+				Sum += FMath::Square(Value - Mean);
+			}
+			Sum /= static_cast<float>(Values.Num());
 
-		return FMath::Sqrt(Sum);
+			return FMath::Sqrt(Sum);
+		}
 	}
 
 	void FMLDeformerMorphModelEditorModel::UpdateMorphErrorValues(TArrayView<UMorphTarget*> MorphTargets)
diff --git a/Engine/Plugins/Animation/MLDeformer/MLDeformerFramework/Source/MLDeformerFrameworkEditor/Private/MLDeformerVizSettingsDetails.cpp b/Engine/Plugins/Animation/MLDeformer/MLDeformerFramework/Source/MLDeformerFrameworkEditor/Private/MLDeformerVizSettingsDetails.cpp
--- a/Engine/Plugins/Animation/MLDeformer/MLDeformerFramework/Source/MLDeformerFrameworkEditor/Private/MLDeformerVizSettingsDetails.cpp
+++ b/Engine/Plugins/Animation/MLDeformer/MLDeformerFramework/Source/MLDeformerFrameworkEditor/Private/MLDeformerVizSettingsDetails.cpp
@@ -11,6 +11,7 @@
 #include "MLDeformerComponent.h"
 #include "MLDeformerPerfCounter.h"
 #include "MLDeformerEditorStyle.h"
+#include "MLDeformerEditorToolkit.h"
 #include "DetailLayoutBuilder.h"
 #include "DetailCategoryBuilder.h"
 #include "DetailWidgetRow.h"
@@ -186,6 +187,40 @@ namespace UE::MLDeformer
 
 		AddAdditionalSettings();
 
+		// Debugging group.
+		IDetailGroup& DebuggingGroup = LiveSettingsCategory->AddGroup("Debugging", LOCTEXT("DebuggingLabel", "Debugging"), false, true);
+		FDetailWidgetRow& DebuggingRow = DebuggingGroup.HeaderRow();
+		DebuggingRow.NameContent().Widget = 
+			SNew(STextBlock)
+			.Font(IDetailLayoutBuilder::GetDetailFont())
+			.Text(FText::FromName(DebuggingGroup.GetGroupName()));	
+		DebuggingRow.ValueContent().Widget =
+			SNew(STextBlock)
+			.Font(IDetailLayoutBuilder::GetDetailFont())
+			.Text_Lambda
+			(
+				[this]()
+				{
+					FMLDeformerEditorToolkit* Editor = EditorModel->GetEditor();
+					if (Editor && IsValid(Editor->GetDebugActor()))
+					{
+						return LOCTEXT("DebuggingActiveText", "Debugging Active");
+					}
+
+					return GEditor->GetPIEViewport() ? LOCTEXT("DebuggingPIEActiveText", "PIE Active") : FText();
+				}
+			)
+			.ColorAndOpacity_Lambda			
+			(
+				[this]()
+				{
+					return GEditor->GetPIEViewport() ? FSlateColor(FColor::Green) : FSlateColor::UseForeground();
+				}
+			);
+		DebuggingGroup.AddPropertyRow(DetailBuilder.GetProperty(UMLDeformerVizSettings::GetDrawDebugActorBoundsPropertyName(), UMLDeformerVizSettings::StaticClass()));
+		DebuggingGroup.AddPropertyRow(DetailBuilder.GetProperty(UMLDeformerVizSettings::GetDebugBoundsColorPropertyName(), UMLDeformerVizSettings::StaticClass()));
+
+		// Visibility group.
 		IDetailGroup& VisGroup = LiveSettingsCategory->AddGroup("Visibility", LOCTEXT("VisibilityLabel", "Visibility"), false, true);
 		VisGroup.AddPropertyRow(DetailBuilder.GetProperty(UMLDeformerVizSettings::GetDrawLinearSkinnedActorPropertyName(), UMLDeformerVizSettings::StaticClass()));
 		VisGroup.AddPropertyRow(DetailBuilder.GetProperty(UMLDeformerVizSettings::GetDrawMLDeformedActorPropertyName(), UMLDeformerVizSettings::StaticClass()));
diff --git a/Engine/Plugins/Animation/MLDeformer/MLDeformerFramework/Source/MLDeformerFrameworkEditor/Private/SMLDeformerDebugSelectionWidget.cpp b/Engine/Plugins/Animation/MLDeformer/MLDeformerFramework/Source/MLDeformerFrameworkEditor/Private/SMLDeformerDebugSelectionWidget.cpp
new file mode 100644
--- /dev/null
+++ b/Engine/Plugins/Animation/MLDeformer/MLDeformerFramework/Source/MLDeformerFrameworkEditor/Private/SMLDeformerDebugSelectionWidget.cpp
@@ -0,0 +1,434 @@
+// Copyright Epic Games, Inc. All Rights Reserved.
+
+#include "SMLDeformerDebugSelectionWidget.h"
+#include "Widgets/Input/SComboBox.h"
+#include "Widgets/Text/STextBlock.h"
+#include "Widgets/SBoxPanel.h"
+#include "Widgets/Layout/SGridPanel.h"
+#include "MLDeformerComponent.h"
+#include "MLDeformerEditorToolkit.h"
+#include "MLDeformerModel.h"
+#include "MLDeformerEditorModel.h"
+#include "MLDeformerEditorStyle.h"
+#include "Components/SkeletalMeshComponent.h"
+#include "Selection.h"
+#include "Engine/World.h"
+#include "UObject/UObjectIterator.h"
+
+#define LOCTEXT_NAMESPACE "MLDeformerDebugSelectionWidget"
+
+namespace
+{
+	bool IsValidDebugWorld(UWorld* World)
+	{
+		return 
+			IsValid(World) &&
+			World->WorldType == EWorldType::PIE &&
+			World->PersistentLevel != nullptr &&
+			World->PersistentLevel->OwningWorld == World;
+	}
+
+	TArray<TObjectPtr<UWorld>> GetDebugWorlds()
+	{
+		TArray<TObjectPtr<UWorld>> Results;
+
+		for (TObjectIterator<UWorld> It; It; ++It)
+		{
+			UWorld* World = *It;
+			if (!IsValidDebugWorld(World))
+			{
+				continue;
+			}
+
+			Results.Emplace(TObjectPtr<UWorld>(World));
+		}
+
+		return MoveTemp(Results);
+	}
+
+	FString GetWorldName(UWorld* World)
+	{
+		FString WorldName;
+
+		ENetMode NetMode = World->GetNetMode();
+		switch (NetMode)
+		{
+			case NM_Standalone:
+				WorldName = LOCTEXT("DebugWorldStandalone", "Standalone").ToString();
+				break;
+
+			case NM_ListenServer:
+				WorldName = LOCTEXT("DebugWorldListenServer", "Listen Server").ToString();
+				break;
+
+			case NM_DedicatedServer:
+				WorldName = LOCTEXT("DebugWorldDedicatedServer", "Dedicated Server").ToString();
+				break;
+
+			case NM_Client:
+				if (FWorldContext* PieContext = GEngine->GetWorldContextFromWorld(World))
+				{
+					WorldName = FString::Printf(TEXT("%s %d"), *LOCTEXT("DebugWorldClient", "Client").ToString(), PieContext->PIEInstance - 1);
+				}
+				break;
+
+			default:
+				WorldName = LOCTEXT("UnknownNetModeName", "Unknown Net Mode").ToString();
+		};
+
+		if (FWorldContext* PieContext = GEngine->GetWorldContextFromWorld(World))
+		{
+			if (!PieContext->CustomDescription.IsEmpty())
+			{
+				WorldName += TEXT(" ") + PieContext->CustomDescription;
+			}
+		}
+
+		return WorldName;
+	}
+}
+
+namespace UE::MLDeformer
+{
+	void SMLDeformerDebugSelectionWidget::Construct(const FArguments& InArgs)
+	{
+		MLDeformerEditor = InArgs._MLDeformerEditor;
+
+		ChildSlot
+		[
+			SNew(SHorizontalBox)
+			+SHorizontalBox::Slot()
+			.AutoWidth()
+			.Padding(FMargin(4, 0, 4, 0))
+			.HAlign(EHorizontalAlignment::HAlign_Left)
+			.VAlign(EVerticalAlignment::VAlign_Center)
+			[
+				SNew(STextBlock)
+				.Text(LOCTEXT("DebugActorText", "Debug Actor:"))
+			]
+			+SHorizontalBox::Slot()
+			.AutoWidth()
+			.Padding(FMargin(2, 0, 0, 0))
+			.HAlign(EHorizontalAlignment::HAlign_Left)
+			.VAlign(EVerticalAlignment::VAlign_Center)
+			[
+				SAssignNew(ActorComboBox, SComboBox<TSharedPtr<FMLDeformerDebugActor>>)
+				.OptionsSource(&Actors)
+				.OnGenerateWidget(this, &SMLDeformerDebugSelectionWidget::OnGenerateActorComboBoxItemWidget)
+				.OnSelectionChanged(this, &SMLDeformerDebugSelectionWidget::OnActorSelectionChanged)
+				[
+					SNew(STextBlock)
+					.Text(this, &SMLDeformerDebugSelectionWidget::GetComboBoxText)
+					.ColorAndOpacity_Lambda([this]()
+					{
+						const FLinearColor ActiveColor = FMLDeformerEditorStyle::Get().GetColor("MLDeformer.Debug.ActiveDebugColor");
+						return (IsDebuggingDisabled() || DebugActorName.IsEmpty()) ? FSlateColor::UseForeground() : FSlateColor(ActiveColor);
+					})
+				]
+			]
+		];
+
+		FWorldDelegates::OnPostWorldInitialization.AddSPLambda(this, [this](UWorld* World, const UWorld::InitializationValues IVS) { Refresh(); });
+		FWorldDelegates::OnWorldCleanup.AddSPLambda(this, [this](UWorld* World, bool bSessionEnded, bool bCleanupResources) { Refresh(); } );
+		FWorldDelegates::OnPostDuplicate.AddSPLambda(this, [this](UWorld* World, bool bDuplicateForPIE, FWorldDelegates::FReplacementMap& ReplacementMap, TArray<UObject*>& ObjectsToFixReferences) { Refresh(); });
+		FWorldDelegates::OnPostWorldRename.AddSPLambda(this, [this](UWorld* World) { Refresh(); } );
+		
+		FWorldDelegates::OnPIEReady.AddSPLambda(this, [this](UGameInstance*) { Refresh(); });
+		FEditorDelegates::PostPIEStarted.AddSPLambda(this, [this](bool bSimulating) { Refresh(); });
+		FEditorDelegates::PausePIE.AddSPLambda(this, [this](bool bSimulating) { Refresh(); });
+		FEditorDelegates::ResumePIE.AddSPLambda(this, [this](bool bSimulating) { Refresh(); });
+		FEditorDelegates::SingleStepPIE.AddSPLambda(this, [this](bool bSimulating) { Refresh(); });
+		FEditorDelegates::EndPIE.AddSPLambda(this, [this](bool bSimulating) { Refresh(); });
+		FEditorDelegates::CancelPIE.AddSPLambda(this, [this]() { Refresh(); });
+		FEditorDelegates::OnNewActorsPlaced.AddSPLambda(this, [this](UObject*, const TArray<AActor*>&) { Refresh(); });
+		FEditorDelegates::OnDeleteActorsBegin.AddSPLambda(this, [this]() { Refresh(); });
+		FEditorDelegates::OnSwitchBeginPIEAndSIE.AddSPLambda(this, [this](bool bSimulating) { Refresh(); });
+
+		USelection::SelectObjectEvent.AddSPLambda(this, [this](UObject* NewSelection) { Refresh(); });
+		USelection::SelectionChangedEvent.AddSPLambda(this, [this](UObject* NewSelection) { Refresh(); });
+		USelection::SelectNoneEvent.AddSPLambda(this, [this]() { Refresh(); });
+	}
+
+	FText SMLDeformerDebugSelectionWidget::GetComboBoxText() const
+	{
+		if (!IsDebuggingDisabled())
+		{
+			if (ActiveDebugActor.IsValid() && !DebugActorName.IsEmpty())
+			{
+				return DebugActorName;
+			}
+		}
+		return LOCTEXT("EditorPreviewText", "Editor Preview (Debug Disabled)");
+	}
+
+	bool SMLDeformerDebugSelectionWidget::IsDebuggingDisabled() const
+	{
+		return Actors.IsEmpty() || !ActiveDebugActor.IsValid() || !IsValid(ActiveDebugActor->Actor);
+	}
+
+	void SMLDeformerDebugSelectionWidget::Refresh()
+	{
+		DebugActorName = FText();
+		RefreshActorList();
+	}
+
+	AActor* SMLDeformerDebugSelectionWidget::GetDebugActor() const
+	{
+		return !DebugActorName.IsEmpty() && ActiveDebugActor.IsValid() && IsValid(ActiveDebugActor->Actor) ? ActiveDebugActor->Actor.Get() : nullptr;
+	}
+
+	TSharedRef<SWidget> SMLDeformerDebugSelectionWidget::OnGenerateActorComboBoxItemWidget(TSharedPtr<FMLDeformerDebugActor> Item)
+	{
+		// If we have the first item in the actor list, generate a special widget.
+		const AActor* Actor = Item.IsValid() ? Item->Actor.Get () : nullptr;
+		FText ActorName = Actor ? FText::FromString(Actor->GetName()) : LOCTEXT("DestroyedActorText", "<Destroyed>");
+
+		const UWorld* World = Actor ? Actor->GetWorld() : nullptr;
+		FText WorldName = World ? FText::FromString(GetDebugStringForWorld(World)) : LOCTEXT("DestroyedWorldText", "<Destroyed>");
+
+		const FLinearColor ActiveColor = FMLDeformerEditorStyle::Get().GetColor("MLDeformer.Debug.ActiveDebugColor");
+		const FLinearColor EditorSelectedMarkColor = FMLDeformerEditorStyle::Get().GetColor("MLDeformer.Debug.EditorSelectedMarkColor");
+
+		// Special item to disable debugging.
+		if (Actor == nullptr)
+		{
+			ActorName = LOCTEXT("DebugDisabledActorText", "Editor Preview (Disable Debug)");
+			WorldName = LOCTEXT("DebugDisabledWorldText", "Editor Preview World");
+		}
+
+		TSharedPtr<SWidget> ItemWidget = 
+			SNew(SGridPanel)
+			+SGridPanel::Slot(0, 0)
+			.Padding(2.0f)
+			.HAlign(EHorizontalAlignment::HAlign_Right)
+			[
+				SNew(STextBlock)
+				.Text(LOCTEXT("ActorName", "Actor:"))
+				.ColorAndOpacity(FSlateColor::UseSubduedForeground())
+			]
+			+SGridPanel::Slot(1, 0)
+			.Padding(2.0f)
+			.HAlign(EHorizontalAlignment::HAlign_Left)
+			[
+				SNew(STextBlock)
+				.Text(ActorName)
+				.Font_Lambda
+				(
+					[Item]()
+					{
+						if (Item.IsValid() && Item->bSelectedInEngine)
+						{
+							return FAppStyle::GetFontStyle(TEXT("NormalFontBold"));
+						}
+						return FAppStyle::GetFontStyle(TEXT("NormalFont"));
+					}
+				)
+				.ColorAndOpacity_Lambda
+				(
+					[this, Item, ActiveColor]() 
+					{ 
+						const AActor* ItemActor = Item.IsValid() ? Item->Actor.Get() : nullptr;
+						return (ActiveDebugActor.IsValid() && ActiveDebugActor->Actor == ItemActor) || (!ItemActor && DebugActorName.IsEmpty()) ? ActiveColor : FSlateColor::UseForeground();
+					}
+				)
+			]
+			+SGridPanel::Slot(2, 0)
+			.Padding(2.0f)
+			.HAlign(EHorizontalAlignment::HAlign_Left)
+			[
+				SNew(STextBlock)
+				.Text_Lambda
+				(
+					[Item]()
+					{
+						if (Item.IsValid() && Item->bSelectedInEngine)
+						{
+							return LOCTEXT("SelectedText", "(Selected)");
+						}
+						return FText();
+					}
+				)
+				.TextStyle(FAppStyle::Get(), "RichTextBlock.Bold")
+				.ColorAndOpacity(EditorSelectedMarkColor)
+			]
+			+SGridPanel::Slot(0, 1)
+			.Padding(2.0f)
+			.HAlign(EHorizontalAlignment::HAlign_Right)
+			[
+				SNew(STextBlock)
+				.Text(LOCTEXT("WorldName", "World:"))
+				.ColorAndOpacity(FSlateColor::UseSubduedForeground())
+			]
+			+SGridPanel::Slot(1, 1)
+			.Padding(2.0f)
+			.HAlign(EHorizontalAlignment::HAlign_Left)
+			[
+				SNew(STextBlock)
+				.Text(WorldName)
+				.ColorAndOpacity_Lambda
+				(
+					[this, Item, ActiveColor]() 
+					{ 
+						const TObjectPtr<AActor> ItemActor = Item.IsValid() ? Item->Actor : nullptr;
+						return (ActiveDebugActor.IsValid() && ActiveDebugActor->Actor == ItemActor) || (!ItemActor && DebugActorName.IsEmpty()) ? ActiveColor : FSlateColor::UseForeground();
+					}
+				)
+			];
+
+		return ItemWidget.ToSharedRef();
+	}
+
+	void SMLDeformerDebugSelectionWidget::OnActorSelectionChanged(TSharedPtr<FMLDeformerDebugActor> Item, ESelectInfo::Type SelectInfo)
+	{
+		ActiveDebugActor = Item;
+		if (!Item.IsValid() || !IsValid(Item->Actor))
+		{
+			DebugActorName = FText();
+		}
+		else
+		{
+			if (IsValid(Item->Actor))
+			{
+				DebugActorName = IsValid(Item->Actor) ? FText::FromString(Item->Actor->GetName()) : FText();
+			}
+		}
+
+		if (MLDeformerEditor && MLDeformerEditor->GetActiveModel())
+		{
+			MLDeformerEditor->GetActiveModel()->OnDebugActorChanged(ActiveDebugActor.IsValid() ? ActiveDebugActor->Actor : nullptr);
+		}
+	}
+
+	TArray<TSharedPtr<FMLDeformerDebugActor>> SMLDeformerDebugSelectionWidget::GetDebugActorsForWorld(TObjectPtr<UWorld> World) const
+	{
+		TArray<TSharedPtr<FMLDeformerDebugActor>> Results;
+
+		FMLDeformerEditorModel* ActiveModel = MLDeformerEditor->GetActiveModel();
+		UMLDeformerModel* Model = ActiveModel ? ActiveModel->GetModel() : nullptr;
+		USkeletalMesh* ModelSkeletalMesh = Model ? Model->GetSkeletalMesh() : nullptr;
+		if (!Model || !ModelSkeletalMesh)
+		{
+			return Results;
+		}
+
+		// Iterate over all actors.
+		for (TObjectIterator<AActor> It; It; ++It)
+		{
+			// Make sure the actor is alive and is using the world we're interested in.
+			AActor* Actor = *It;
+			if (IsValid(Actor) && Actor->GetOuter() && Actor->GetWorld() == World)
+			{	
+				if (Model->IsCompatibleDebugActor(Actor))
+				{
+					TSharedPtr<FMLDeformerDebugActor> NewActor = MakeShared<FMLDeformerDebugActor>();
+					NewActor->Actor = Actor;
+					NewActor->bSelectedInEngine = false;
+					Results.Emplace(NewActor);
+				}
+			}
+		}
+
+		return MoveTemp(Results);
+	}
+
+	void SMLDeformerDebugSelectionWidget::RefreshActorList()
+	{
+		// Create an "Editor Preview (Disable Debug)" item, which is always the first item.
+		// The preview item's Actor member is a nullptr.
+		Actors.Reset();
+		TSharedPtr<FMLDeformerDebugActor> PreviewItem = MakeShared<FMLDeformerDebugActor>();
+		Actors.Emplace(PreviewItem);
+
+		// Add actors of all the worlds we are interested in.
+		const TArray<TObjectPtr<UWorld>> Worlds = GetDebugWorlds();
+		for (UWorld* World : Worlds)
+		{
+			if (IsValidDebugWorld(World))
+			{
+				Actors.Append(GetDebugActorsForWorld(World));
+			}
+		}
+
+		// Update the selected flags.
+		UpdateDebugActorSelectedFlags();
+
+		// Reselect the right item.
+		for (const TSharedPtr<FMLDeformerDebugActor>& CurActor : Actors)
+		{
+			const AActor* ActiveActorObject = ActiveDebugActor.IsValid() ? ActiveDebugActor->Actor.Get() : nullptr;
+			if (CurActor.IsValid() && CurActor->Actor.Get() == ActiveActorObject)
+			{
+				ActorComboBox->SetSelectedItem(CurActor);
+				return;
+			}
+		}
+
+		DebugActorName = FText();
+		ActorComboBox->SetSelectedItem(PreviewItem);
+	}
+
+	void SMLDeformerDebugSelectionWidget::UpdateDebugActorSelectedFlags()
+	{
+		// First set all selected flags to false.
+		for (TSharedPtr<FMLDeformerDebugActor>& DebugActor : Actors)
+		{
+			if (DebugActor.IsValid())
+			{
+				DebugActor->bSelectedInEngine = false;
+			}
+		}
+
+		// Get the selected actors in the editor.
+		const USelection* ActiveDebugActors = GEditor->GetSelectedActors();
+		if (ActiveDebugActors != nullptr)
+		{
+			// Processed in reverse order, as we want the last selected item to be the one we pick.
+			// There can only be one actor selected in the MLD editor to debug, while many can be selected in the editor itself.
+			for (int32 Index = ActiveDebugActors->Num() - 1; Index >= 0; --Index)
+			{
+				// Get the object.
+				const UObject* Object = ActiveDebugActors->GetSelectedObject(Index);
+				if (Object == nullptr)
+				{
+					continue;
+				}
+
+				// Are we an Actor?
+				const AActor* Actor = Cast<AActor>(Object);
+				if (Actor == nullptr)
+				{
+					continue;
+				}
+
+				// Find the actor in the list.
+				const TSharedPtr<FMLDeformerDebugActor>* DebugActor = Actors.FindByPredicate
+				(
+					[Actor](const TSharedPtr<FMLDeformerDebugActor>& DebugActor)
+					{
+						if (DebugActor.IsValid())
+						{
+							return DebugActor->Actor == Actor;
+						}
+						return false;
+					} 
+				);
+
+				// Make sure this actor also is inside our debugging list.
+				if (!DebugActor)
+				{
+					continue;
+				}
+
+				// Mark it as selected.
+				if ((*DebugActor).IsValid())
+				{
+					(*DebugActor)->bSelectedInEngine = true;
+					break;	// There can only be one selected actor, as we can't debug multiple at the same time.
+				}
+			}	// For all selected actors.
+		}
+	}
+
+}	// namespace UE::MLDeformer
+
+#undef LOCTEXT_NAMESPACE
diff --git a/Engine/Plugins/Animation/MLDeformer/MLDeformerFramework/Source/MLDeformerFrameworkEditor/Private/SMLDeformerDebugSelectionWidget.h b/Engine/Plugins/Animation/MLDeformer/MLDeformerFramework/Source/MLDeformerFrameworkEditor/Private/SMLDeformerDebugSelectionWidget.h
new file mode 100644
--- /dev/null
+++ b/Engine/Plugins/Animation/MLDeformer/MLDeformerFramework/Source/MLDeformerFrameworkEditor/Private/SMLDeformerDebugSelectionWidget.h
@@ -0,0 +1,80 @@
+// Copyright Epic Games, Inc. All Rights Reserved.
+
+#pragma once
+
+#include "CoreTypes.h"
+#include "Templates/SharedPointer.h"
+#include "Widgets/SCompoundWidget.h"
+#include "Widgets/Input/SComboBox.h"
+#include "Containers/UnrealString.h"
+
+class UWorld;
+class AActor;
+
+namespace UE::MLDeformer
+{
+	class FMLDeformerEditorToolkit;
+}
+
+namespace UE::MLDeformer
+{
+	struct FMLDeformerDebugActor
+	{
+		/** The actor inside the PIE viewport, or nullptr when debugging is disabled. */
+		TObjectPtr<AActor> Actor;
+
+		/**
+		 * Is this actor selected in the editor / engine?
+		 * This is different from the Debug Actor we picked, we call that Active debug actor. 
+		 */
+		bool bSelectedInEngine = false;
+	};
+
+	/**
+	 * 
+	 */
+	class MLDEFORMERFRAMEWORKEDITOR_API SMLDeformerDebugSelectionWidget
+		: public SCompoundWidget
+	{
+	public:
+		SLATE_BEGIN_ARGS(SMLDeformerDebugSelectionWidget) {}
+		SLATE_ARGUMENT(FMLDeformerEditorToolkit*, MLDeformerEditor)
+		SLATE_END_ARGS()
+
+		void Construct(const FArguments& InArgs);
+
+		/** Refresh the list of debuggable actors in the combobox. */
+		void Refresh();
+
+		/** Get the current actor we're debugging, or a nullptr when debugging is disabled. */
+		AActor* GetDebugActor() const;
+
+		const TArray<TSharedPtr<FMLDeformerDebugActor>>& GetActors() const { return Actors; }
+
+	private:
+		TSharedRef<SWidget> OnGenerateActorComboBoxItemWidget(TSharedPtr<FMLDeformerDebugActor> Item);
+		TArray<TSharedPtr<FMLDeformerDebugActor>> GetDebugActorsForWorld(TObjectPtr<UWorld> World) const;
+		FText GetComboBoxText() const;
+		void OnActorSelectionChanged(TSharedPtr<FMLDeformerDebugActor> Item, ESelectInfo::Type SelectInfo);
+		void RefreshActorList();
+		void UpdateDebugActorSelectedFlags();
+		bool IsDebuggingDisabled() const;
+
+	private:
+		/** The combobox that contains the actor names. */
+		TSharedPtr<SComboBox<TSharedPtr<FMLDeformerDebugActor>>> ActorComboBox;
+
+		/** The actors we can debug. */
+		TArray<TSharedPtr<FMLDeformerDebugActor>> Actors;
+
+		/** The selected actor. */
+		TSharedPtr<FMLDeformerDebugActor> ActiveDebugActor;
+
+		/** The name of the actor we are debugging. */
+		FText DebugActorName;
+
+		/** A pointer to our asset editor. */
+		FMLDeformerEditorToolkit* MLDeformerEditor = nullptr;
+	};
+
+}	// namespace UE::MLDeformer
diff --git a/Engine/Plugins/Animation/MLDeformer/MLDeformerFramework/Source/MLDeformerFrameworkEditor/Public/MLDeformerEditorModel.h b/Engine/Plugins/Animation/MLDeformer/MLDeformerFramework/Source/MLDeformerFrameworkEditor/Public/MLDeformerEditorModel.h
--- a/Engine/Plugins/Animation/MLDeformer/MLDeformerFramework/Source/MLDeformerFrameworkEditor/Public/MLDeformerEditorModel.h
+++ b/Engine/Plugins/Animation/MLDeformer/MLDeformerFramework/Source/MLDeformerFrameworkEditor/Public/MLDeformerEditorModel.h
@@ -628,6 +628,31 @@ namespace UE::MLDeformer
 		UE_DEPRECATED(5.3, "Please call FMLDeformerMorphModelEditorModel::CalcMorphTargetNormals instead.")
 		virtual void GenerateNormalsForMorphTarget(int32 LOD, USkeletalMesh* SkelMesh, int32 MorphTargetIndex, TArrayView<const FVector3f> Deltas, TArrayView<const FVector3f> BaseVertexPositions, TArrayView<FVector3f> BaseNormals, TArray<FVector3f>& OutDeltaNormals);
 
+		/**
+		 * Called whenever the actor to debug has been changed.
+		 * This can be changed throughout the UI, in testing mode.
+		 * When debugging gets disabled, this will contain a nullptr.
+		 * @param DebugActor The new actor we want to debug.
+		 */
+		virtual void OnDebugActorChanged(TObjectPtr<AActor> DebugActor);
+
+		/** 
+		 * Apply the transforms of the actor we debug to the editor actors in our world.
+		 * @param DebugActorComponentSpaceTransforms The component space transforms from the actor we are debugging.
+		 */
+		virtual void ApplyDebugActorTransforms(const TArray<FTransform>& DebugActorComponentSpaceTransforms);
+
+		/**
+		 * Debug draw helpers inside the PIE viewport that highlight debuggable actors.
+		 * On default this will draw a bounding box around the actors that can be debugged by this model.
+		 * Also it will render the actor names.
+		 * NOTE: This renders inside the PIE viewport, not our own MLD asset editor viewport.
+		 */
+		virtual void DrawPIEDebugActors();
+
+		/** Apply the transforms of the debug actor to the actors in the asset editor world. This will internally call ApplyDebugActorTransforms(DebugActorComponentSpaceTransforms). */
+		void ApplyDebugActorTransforms();
+
 		/** Invalidate the memory usage, so it gets updated in the UI again. */
 		void UpdateMemoryUsage();
 
diff --git a/Engine/Plugins/Animation/MLDeformer/MLDeformerFramework/Source/MLDeformerFrameworkEditor/Public/MLDeformerEditorToolkit.h b/Engine/Plugins/Animation/MLDeformer/MLDeformerFramework/Source/MLDeformerFrameworkEditor/Public/MLDeformerEditorToolkit.h
--- a/Engine/Plugins/Animation/MLDeformer/MLDeformerFramework/Source/MLDeformerFrameworkEditor/Public/MLDeformerEditorToolkit.h
+++ b/Engine/Plugins/Animation/MLDeformer/MLDeformerFramework/Source/MLDeformerFrameworkEditor/Public/MLDeformerEditorToolkit.h
@@ -25,6 +25,7 @@ namespace UE::MLDeformer
 {
 	class SMLDeformerTimeline;
 	class FMLDeformerApplicationMode;
+	class SMLDeformerDebugSelectionWidget;
 
 	namespace MLDeformerEditorModes
 	{
@@ -55,7 +56,6 @@ namespace UE::MLDeformer
 		friend class FMLDeformerApplicationMode;
 		friend struct FMLDeformerVizSettingsTabSummoner;
 
-
 		~FMLDeformerEditorToolkit();
 
 		/** Initialize the asset editor. This will register the application mode, init the preview scene, etc. */
@@ -129,10 +129,19 @@ namespace UE::MLDeformer
 		bool IsTrainButtonEnabled() const;
 		bool IsTraining() const;
 
+		/** Get the actor we want to debug, if any. Returns a nullptr when we don't want to debug anything. */
+		AActor* GetDebugActor() const;
+
+		/** Get the component space transforms of the actor we want to debug. Returns an empty array if GetDebugActor returns a nullptr. */
+		TArray<FTransform> GetDebugActorComponentSpaceTransforms() const;
+
 		void ZoomOnActors();
 
+		TSharedPtr<SMLDeformerDebugSelectionWidget> GetDebugWidget() const { return DebugWidget; }
+
 		static void AddToolsMenuExtender(TUniquePtr<FToolsMenuExtender> Extender);
 		static TConstArrayView<TUniquePtr<FToolsMenuExtender>> GetToolsMenuExtenders();
+
 	private:
 		UE_DEPRECATED(5.3, "Please use the OnModelChanged that takes two parameters instead.")
 		void OnModelChanged(int Index);
@@ -160,6 +169,7 @@ namespace UE::MLDeformer
 		FText GetCurrentVizModeName() const;
 		FText GetVizModeName(EMLDeformerVizMode Mode) const;
 		void ShowNoModelsWarningIfNeeded();
+		EVisibility GetDebuggingVisibility() const;
 
 		TSharedRef<SWidget> GenerateModelButtonContents(TSharedRef<FUICommandList> InCommandList);
 		TSharedRef<SWidget> GenerateVizModeButtonContents(TSharedRef<FUICommandList> InCommandList);
@@ -167,26 +177,30 @@ namespace UE::MLDeformer
 
 	private:
 		/** The persona toolkit. */	
-		TSharedPtr<IPersonaToolkit> PersonaToolkit = nullptr;
+		TSharedPtr<IPersonaToolkit> PersonaToolkit;
 
 		/** Model details view. */
-		TSharedPtr<IDetailsView> ModelDetailsView = nullptr;
+		TSharedPtr<IDetailsView> ModelDetailsView;
 
 		/** Model viz settings details view. */
-		TSharedPtr<IDetailsView> VizSettingsDetailsView = nullptr;
+		TSharedPtr<IDetailsView> VizSettingsDetailsView;
 
 		/** The timeline slider widget. */
-		TSharedPtr<SMLDeformerTimeline> TimeSlider = nullptr;
+		TSharedPtr<SMLDeformerTimeline> TimeSlider;
 
 		/** The currently active editor model. */
-		TSharedPtr<FMLDeformerEditorModel> ActiveModel = nullptr;
+		TSharedPtr<FMLDeformerEditorModel> ActiveModel;
 
 		// Persona viewport.
-		TSharedPtr<IPersonaViewport> PersonaViewport = nullptr;
+		TSharedPtr<IPersonaViewport> PersonaViewport;
 
 		/** The ML Deformer Asset. */
-		TObjectPtr<UMLDeformerAsset> DeformerAsset = nullptr;
+		TObjectPtr<UMLDeformerAsset> DeformerAsset;
+
+		/** The widget where you select which actor to debug. */
+		TSharedPtr<SMLDeformerDebugSelectionWidget> DebugWidget;
 
+		/** The active application mode. */
 		FMLDeformerApplicationMode* ApplicationMode = nullptr;
 
 		/** Has the asset editor been initialized? */
diff --git a/Engine/Plugins/Animation/MLDeformer/NearestNeighborModel/Source/NearestNeighborModel/Private/NearestNeighborModelInstance.cpp b/Engine/Plugins/Animation/MLDeformer/NearestNeighborModel/Source/NearestNeighborModel/Private/NearestNeighborModelInstance.cpp
--- a/Engine/Plugins/Animation/MLDeformer/NearestNeighborModel/Source/NearestNeighborModel/Private/NearestNeighborModelInstance.cpp
+++ b/Engine/Plugins/Animation/MLDeformer/NearestNeighborModel/Source/NearestNeighborModel/Private/NearestNeighborModelInstance.cpp
@@ -9,6 +9,11 @@
 #include "NearestNeighborModel.h"
 #include "NearestNeighborModelInputInfo.h"
 #include "NearestNeighborOptimizedNetwork.h"
+#include "MLDeformerAsset.h"
+#include "MLDeformerComponent.h"
+#include "Components/ExternalMorphSet.h"
+#include "Engine/SkeletalMesh.h"
+#include "Components/SkeletalMeshComponent.h"
 
 #include UE_INLINE_GENERATED_CPP_BY_NAME(NearestNeighborModelInstance)
 CSV_DECLARE_CATEGORY_MODULE_EXTERN(MLDEFORMERFRAMEWORK_API, MLDeformer);
@@ -92,6 +97,7 @@ void UNearestNeighborModelInstance::Tick(float DeltaTime, float ModelWeight)
 		PostMLDeformerComponentInit();
 	}
 
+	bool bCalledExecute = false;
 	if (ModelWeight > 0.0001f && HasValidTransforms() && SetupInputs())
 	{
 		// Execute the model instance.
@@ -99,11 +105,15 @@ void UNearestNeighborModelInstance::Tick(float DeltaTime, float ModelWeight)
 		// calculate the network outputs and possibly use them, depending on how the model works.
 		Execute(ModelWeight);
 		RunNearestNeighborModel(DeltaTime, ModelWeight);
+		bCalledExecute = true;
 	}
 	else
 	{
 		HandleZeroModelWeight();
 	}
+
+	// Do some things afterards, such as copying over debug actor data.
+	PostTick(bCalledExecute);
 }
 
 int64 UNearestNeighborModelInstance::SetBoneTransforms(float* OutputBuffer, int64 OutputBufferSize, int64 StartIndex)
diff --git a/Engine/Plugins/Animation/MLDeformer/NearestNeighborModel/Source/NearestNeighborModel/Public/NearestNeighborModelVizSettings.h b/Engine/Plugins/Animation/MLDeformer/NearestNeighborModel/Source/NearestNeighborModel/Public/NearestNeighborModelVizSettings.h
--- a/Engine/Plugins/Animation/MLDeformer/NearestNeighborModel/Source/NearestNeighborModel/Public/NearestNeighborModelVizSettings.h
+++ b/Engine/Plugins/Animation/MLDeformer/NearestNeighborModel/Source/NearestNeighborModel/Public/NearestNeighborModelVizSettings.h
@@ -21,7 +21,9 @@ class NEARESTNEIGHBORMODEL_API UNearestNeighborModelVizSettings
 	GENERATED_BODY()
 #if	WITH_EDITORONLY_DATA
 public:
-	static FName GetNearestNeighborActorsOffsetPropertyName() { return GET_MEMBER_NAME_CHECKED(UNearestNeighborModelVizSettings, NearestNeighborActorsOffset); }
+	UE_DEPRECATED(5.4, "This method will be removed.")
+	static FName GetNearestNeighborActorsOffsetPropertyName() { return FName(); }
+
 	static FName GetNearestNeighborIdsPropertyName() { return GET_MEMBER_NAME_CHECKED(UNearestNeighborModelVizSettings, NearestNeighborIds); }
 	
 	/** Whether to show verts */
@@ -32,15 +34,25 @@ public:
 	UPROPERTY(EditAnywhere, Category = "Training Meshes", Meta = (DisplayName = "Show Verts in", EditorCondition = "bDrawVerts"))
 	int32 VertVizSectionIndex = INDEX_NONE;
 
+	UPROPERTY()
+	float NearestNeighborActorsOffset_DEPRECATED = 2.0f;
+
 	/** The section used to display the nearest neighbor. */	
 	UPROPERTY(EditAnywhere, Category = "Live Settings", Meta = (DisplayName = "Actor Section Index"))
 	int32 NearestNeighborActorSectionIndex = 0;
 	
-	/** The offset of the nearest neighbor actor from the mesh. */
-	UPROPERTY(EditAnywhere, Category = "Live Settings", Meta = (DisplayName = "Actor Offset"))
-	float NearestNeighborActorsOffset = 2.0f;
-	
 	UPROPERTY(VisibleAnywhere, Category = "Live Settings")
 	TArray<int32> NearestNeighborIds;
+
+	UPROPERTY(EditAnywhere, Category = "Live Settings", meta = (DisplayName = "Part Id"))
+	int32 NeighborStatsPartId = 0;
+
+	UE_DEPRECATED(5.4, "This method will be removed.")
+	void SetNearestNeighborActorsOffset(float InOffset) { }
+
+	UE_DEPRECATED(5.4, "This method will be removed.")
+	float GetNearestNeighborActorsOffset() const { return 0.0f; }
+
+	void SetNearestNeighborIds(const TArray<int32>& InNearestNeighborIds) { NearestNeighborIds = InNearestNeighborIds; }
 #endif
 };
diff --git a/Engine/Plugins/Animation/MLDeformer/NearestNeighborModel/Source/NearestNeighborModelEditor/Private/NearestNeighborEditorModel.cpp b/Engine/Plugins/Animation/MLDeformer/NearestNeighborModel/Source/NearestNeighborModelEditor/Private/NearestNeighborEditorModel.cpp
--- a/Engine/Plugins/Animation/MLDeformer/NearestNeighborModel/Source/NearestNeighborModelEditor/Private/NearestNeighborEditorModel.cpp
+++ b/Engine/Plugins/Animation/MLDeformer/NearestNeighborModel/Source/NearestNeighborModelEditor/Private/NearestNeighborEditorModel.cpp
@@ -435,6 +435,8 @@ namespace UE::NearestNeighborModel
 
 	void FNearestNeighborEditorModel::UpdateNearestNeighborActor(FNearestNeighborEditorModelActor& Actor) const
 	{
+		using namespace UE::MLDeformer;
+
 		const UNearestNeighborModel* const NearestNeighborModel = GetCastModel();
 		if (!NearestNeighborModel)
 		{
@@ -452,8 +454,18 @@ namespace UE::NearestNeighborModel
 			return;
 		}
 
-		const float Offset = NNViz->NearestNeighborActorsOffset;
-		Actor.SetMeshOffsetFactor(Offset * 2.0f);
+		float MaxOffset = 0.0f;
+		for (FMLDeformerEditorActor* EditorActor : EditorActors)
+		{
+			if (EditorActor && 
+				EditorActor->GetTypeID() != ActorID_NearestNeighborActors &&
+				EditorActor->IsVisible()) 
+			{
+				MaxOffset = FMath::Max(EditorActor->GetMeshOffsetFactor(), MaxOffset);
+			}
+		}
+
+		Actor.SetMeshOffsetFactor(MaxOffset + 1.0f);
 
 		UGeometryCache* const GeomCache = NearestNeighborModel->GetSection(SectionIndex).GetMutableNeighborMeshes();
 		if (!GeomCache)
diff --git a/Engine/Plugins/Animation/MLDeformer/NearestNeighborModel/Source/NearestNeighborModelEditor/Private/NearestNeighborModelVizSettingsDetails.cpp b/Engine/Plugins/Animation/MLDeformer/NearestNeighborModel/Source/NearestNeighborModelEditor/Private/NearestNeighborModelVizSettingsDetails.cpp
--- a/Engine/Plugins/Animation/MLDeformer/NearestNeighborModel/Source/NearestNeighborModelEditor/Private/NearestNeighborModelVizSettingsDetails.cpp
+++ b/Engine/Plugins/Animation/MLDeformer/NearestNeighborModel/Source/NearestNeighborModelEditor/Private/NearestNeighborModelVizSettingsDetails.cpp
@@ -70,7 +70,6 @@ namespace UE::NearestNeighborModel
 	{
 		IDetailGroup& NNGroup = LiveSettingsCategory->AddGroup("Nearest Neighbor", LOCTEXT("NearestNeighborLabel", "Nearest Neighbor"), false, true);
 		NNGroup.AddPropertyRow(DetailLayoutBuilder->GetProperty(GET_MEMBER_NAME_STRING_CHECKED(UNearestNeighborModelVizSettings, NearestNeighborActorSectionIndex), UNearestNeighborModelVizSettings::StaticClass()));
-		NNGroup.AddPropertyRow(DetailLayoutBuilder->GetProperty(UNearestNeighborModelVizSettings::GetNearestNeighborActorsOffsetPropertyName(), UNearestNeighborModelVizSettings::StaticClass()));
 		NNGroup.AddPropertyRow(DetailLayoutBuilder->GetProperty(UNearestNeighborModelVizSettings::GetNearestNeighborIdsPropertyName(), UNearestNeighborModelVizSettings::StaticClass()));	
 	}
 
diff --git a/Engine/Plugins/Animation/MLDeformer/NearestNeighborModel/Source/NearestNeighborModelEditor/Public/NearestNeighborEditorModel.h b/Engine/Plugins/Animation/MLDeformer/NearestNeighborModel/Source/NearestNeighborModelEditor/Public/NearestNeighborEditorModel.h
--- a/Engine/Plugins/Animation/MLDeformer/NearestNeighborModel/Source/NearestNeighborModelEditor/Public/NearestNeighborEditorModel.h
+++ b/Engine/Plugins/Animation/MLDeformer/NearestNeighborModel/Source/NearestNeighborModelEditor/Public/NearestNeighborEditorModel.h
@@ -45,13 +45,13 @@ namespace UE::NearestNeighborModel
 		virtual void InitInputInfo(UMLDeformerInputInfo* InputInfo) override;
 		virtual ETrainingResult Train() override;
 		virtual bool LoadTrainedNetwork() const override;
-		virtual void OnPostTraining(ETrainingResult TrainingResult, bool bUsePartiallyTrainedWhenAborted) override;
 		virtual FMLDeformerTrainingInputAnim* GetTrainingInputAnim(int32 Index) const override;
 		virtual void UpdateTimelineTrainingAnimList() override;
 		virtual void OnPropertyChanged(FPropertyChangedEvent& PropertyChangedEvent) override;
+		virtual void OnPostTraining(ETrainingResult TrainingResult, bool bUsePartiallyTrainedWhenAborted) override;
 		virtual void Render(const FSceneView* View, FViewport* Viewport, FPrimitiveDrawInterface* PDI) override;
 		// ~END FMLDeformerEditorModel overrides.
-
+		
 		// UMLDeformerMorphModelEditorModel overrides.
 		virtual bool IsMorphWeightClampingSupported() const override	{ return false; }	// We already do input clamping, so output clamping really isn't needed.
 		// ~END UMLDeformerMorphModelEditorModel overrides.
@@ -84,12 +84,8 @@ namespace UE::NearestNeighborModel
 
 		TUniquePtr<FNearestNeighborEditorModelActor> CreateNearestNeighborActor(UWorld* World) const;
 		void UpdateNearestNeighborActor(FNearestNeighborEditorModelActor& Actor) const;
-
-		void GetNeighborStats();
-
-		/** This should be only set in CreateActors(). */ 
-		TUniquePtr<FNearestNeighborEditorModelActor> NearestNeighborActor;
-
+	
+		TUniquePtr<FNearestNeighborEditorModelActor> NearestNeighborActor;	// This should be only set in CreateActors().
 		TUniquePtr<FVertexMapSelector> VertexMapSelector;
 		TUniquePtr<FVertVizSelector> VertVizSelector;
 	};
@@ -103,6 +99,7 @@ namespace UE::NearestNeighborModel
 		TSharedPtr<FString> GetSelectedItem(const UNearestNeighborModelSection& Section) const;
 		FString GetVertexMapString(const UNearestNeighborModelSection& Section) const;
 		bool IsValid() const;
+
 	private:
 		void Reset();
 		TArray<TSharedPtr<FString>> Options;
@@ -120,6 +117,7 @@ namespace UE::NearestNeighborModel
 		TSharedPtr<FString> GetSelectedItem() const;
 		int32 GetSectionIndex(TSharedPtr<FString> Item) const;
 		void SelectSection(int32 SectionIndex);
+
 	private:
 		TObjectPtr<UNearestNeighborModelVizSettings> Settings;
 		TArray<TSharedPtr<FString>> Options;
diff --git a/Engine/Source/Runtime/Engine/Private/Animation/AnimTrace.cpp b/Engine/Source/Runtime/Engine/Private/Animation/AnimTrace.cpp
--- a/Engine/Source/Runtime/Engine/Private/Animation/AnimTrace.cpp
+++ b/Engine/Source/Runtime/Engine/Private/Animation/AnimTrace.cpp
@@ -6,6 +6,7 @@
 
 #include "Animation/AnimInstanceProxy.h"
 #include "Components/SkeletalMeshComponent.h"
+#include "Components/ExternalMorphSet.h"
 #include "Engine/SkeletalMesh.h"
 #include "Animation/AnimMontage.h"
 #include "Animation/AnimBlueprintGeneratedClass.h"
@@ -55,6 +56,21 @@ UE_TRACE_EVENT_BEGIN(Animation, SkeletalMeshComponent3)
 	UE_TRACE_EVENT_FIELD(uint16, FrameCounter)
 UE_TRACE_EVENT_END()
 
+UE_TRACE_EVENT_BEGIN(Animation, SkeletalMeshComponent4)
+	UE_TRACE_EVENT_FIELD(uint64, Cycle)
+	UE_TRACE_EVENT_FIELD(double, RecordingTime)
+	UE_TRACE_EVENT_FIELD(uint64, ComponentId)
+	UE_TRACE_EVENT_FIELD(uint64, MeshId)
+	UE_TRACE_EVENT_FIELD(float[], ComponentToWorld)
+	UE_TRACE_EVENT_FIELD(float[], Pose)
+	UE_TRACE_EVENT_FIELD(uint32[], CurveIds)
+	UE_TRACE_EVENT_FIELD(float[], CurveValues)
+	UE_TRACE_EVENT_FIELD(uint16, LodIndex)
+	UE_TRACE_EVENT_FIELD(uint16, FrameCounter)
+	UE_TRACE_EVENT_FIELD(float[], ExternalMorphSetWeights)		// Concatenated array [WeightsOfSet0, WeightsOfSet1, ..., WeightsOfSetN].
+	UE_TRACE_EVENT_FIELD(int32[], ExternalMorphSetWeightCounts)	// The number of weights in each external morph set.
+UE_TRACE_EVENT_END()
+
 UE_TRACE_EVENT_BEGIN(Animation, SkeletalMeshFrame)
 	UE_TRACE_EVENT_FIELD(uint64, Cycle)
 	UE_TRACE_EVENT_FIELD(uint64, ComponentId)
@@ -295,6 +311,8 @@ struct FAnimTraceScratchBuffers : public TThreadSingleton<FAnimTraceScratchBuffe
 	// Curve values/IDs for skeletal mesh component
 	TArray<float> CurveValues;
 	TArray<uint32> CurveIds;
+	TArray<float> ExternalMorphTargetWeights;
+	TArray<int32> ExternalMorphTargetWeightCounts;
 
 	// Parent indices for skeletal meshes
 	TArray<int32> ParentIndices;
@@ -565,8 +583,17 @@ void FAnimTrace::OutputSkeletalMeshComponent(const USkeletalMeshComponent* InCom
 			CurveCount += AnimInstance->GetAnimationCurveList(CurveType).Num();
 		}
 	}
+
+	// Get the external morph target sets.
+	const int32 LOD = InComponent->GetPredictedLODLevel();
 	
-	if(BoneCount > 0 || CurveCount > 0)
+	FExternalMorphSets* ExternalMorphSet = nullptr;
+	if (InComponent->IsValidExternalMorphSetLODIndex(LOD) && CVarRecordExternalMorphTargets->GetBool())
+	{
+		ExternalMorphSet = const_cast<FExternalMorphSets*>(&InComponent->GetExternalMorphSets(LOD));
+	}
+
+	if (BoneCount > 0 || CurveCount > 0 || (ExternalMorphSet && !ExternalMorphSet->IsEmpty()))
 	{
 		TRACE_OBJECT(InComponent);
 		TRACE_SKELETAL_MESH(InComponent->GetSkeletalMeshAsset());
@@ -592,17 +619,45 @@ void FAnimTrace::OutputSkeletalMeshComponent(const USkeletalMeshComponent* InCom
 			}
 		}
 
-		UE_TRACE_LOG(Animation, SkeletalMeshComponent3, AnimationChannel)
-			<< SkeletalMeshComponent3.Cycle(FPlatformTime::Cycles64())
-			<< SkeletalMeshComponent3.RecordingTime(FObjectTrace::GetWorldElapsedTime(InComponent->GetWorld()))
-			<< SkeletalMeshComponent3.ComponentId(FObjectTrace::GetObjectId(InComponent))
-			<< SkeletalMeshComponent3.MeshId(FObjectTrace::GetObjectId(InComponent->GetSkeletalMeshAsset()))
-			<< SkeletalMeshComponent3.ComponentToWorld(reinterpret_cast<const float*>(&InComponent->GetComponentToWorld()), sizeof(FTransform) / sizeof(float))
-			<< SkeletalMeshComponent3.Pose(reinterpret_cast<const float*>(InComponent->GetComponentSpaceTransforms().GetData()), BoneCount * (sizeof(FTransform) / sizeof(float)))
-			<< SkeletalMeshComponent3.CurveIds(CurveIds.GetData(), CurveIds.Num())
-			<< SkeletalMeshComponent3.CurveValues(CurveValues.GetData(), CurveValues.Num())
-			<< SkeletalMeshComponent3.LodIndex((uint16)InComponent->GetPredictedLODLevel())
-			<< SkeletalMeshComponent3.FrameCounter(FObjectTrace::GetObjectWorldTickCounter(InComponent));
+		TArray<float>& ExternalMorphTraceWeights = FAnimTraceScratchBuffers::Get().ExternalMorphTargetWeights;
+		TArray<int32>& ExternalMorphTraceWeightCounts = FAnimTraceScratchBuffers::Get().ExternalMorphTargetWeightCounts;
+		ExternalMorphTraceWeights.Reset();
+		ExternalMorphTraceWeightCounts.Reset();
+
+		// Get the weights for all external morph sets.
+		if (ExternalMorphSet && !ExternalMorphSet->IsEmpty())
+		{
+			const FExternalMorphWeightData& ExternalMorphWeightData = InComponent->GetExternalMorphWeights(LOD);
+			const int32 NumMorphWeightSets = ExternalMorphWeightData.MorphSets.Num();
+			const int32 NumMorphDataSets = ExternalMorphSet->Num();
+			if (!ExternalMorphWeightData.MorphSets.IsEmpty() && NumMorphDataSets == NumMorphWeightSets)
+			{
+				// Allocate enough data in the scratch buffer.
+				ExternalMorphTraceWeightCounts.SetNumZeroed(NumMorphDataSets);
+
+				int32 SetIndex = 0;
+				for (const auto& CurrentSet : ExternalMorphWeightData.MorphSets)
+				{
+					ExternalMorphTraceWeightCounts[SetIndex] = CurrentSet.Value.Weights.Num();
+					ExternalMorphTraceWeights.Append(CurrentSet.Value.Weights);
+					SetIndex++;
+				}
+			}
+		}
+
+		UE_TRACE_LOG(Animation, SkeletalMeshComponent4, AnimationChannel)
+			<< SkeletalMeshComponent4.Cycle(FPlatformTime::Cycles64())
+			<< SkeletalMeshComponent4.RecordingTime(FObjectTrace::GetWorldElapsedTime(InComponent->GetWorld()))
+			<< SkeletalMeshComponent4.ComponentId(FObjectTrace::GetObjectId(InComponent))
+			<< SkeletalMeshComponent4.MeshId(FObjectTrace::GetObjectId(InComponent->GetSkeletalMeshAsset()))
+			<< SkeletalMeshComponent4.ComponentToWorld(reinterpret_cast<const float*>(&InComponent->GetComponentToWorld()), sizeof(FTransform) / sizeof(float))
+			<< SkeletalMeshComponent4.Pose(reinterpret_cast<const float*>(InComponent->GetComponentSpaceTransforms().GetData()), BoneCount * (sizeof(FTransform) / sizeof(float)))
+			<< SkeletalMeshComponent4.CurveIds(CurveIds.GetData(), CurveIds.Num())
+			<< SkeletalMeshComponent4.CurveValues(CurveValues.GetData(), CurveValues.Num())
+			<< SkeletalMeshComponent4.LodIndex((uint16)InComponent->GetPredictedLODLevel())
+			<< SkeletalMeshComponent4.FrameCounter(FObjectTrace::GetObjectWorldTickCounter(InComponent))
+			<< SkeletalMeshComponent4.ExternalMorphSetWeights(ExternalMorphTraceWeights.GetData(), ExternalMorphTraceWeights.Num())
+			<< SkeletalMeshComponent4.ExternalMorphSetWeightCounts(ExternalMorphTraceWeightCounts.GetData(), ExternalMorphTraceWeightCounts.Num());
 	}
 }
 
diff --git a/Engine/Source/Runtime/Engine/Public/Animation/AnimTrace.h b/Engine/Source/Runtime/Engine/Public/Animation/AnimTrace.h
--- a/Engine/Source/Runtime/Engine/Public/Animation/AnimTrace.h
+++ b/Engine/Source/Runtime/Engine/Public/Animation/AnimTrace.h
@@ -36,6 +36,13 @@ struct FAnimSyncMarker;
 struct FAnimMontageInstance;
 class UPoseWatchPoseElement;
 
+static FAutoConsoleVariable CVarRecordExternalMorphTargets(
+	TEXT("RecordExternalMorphTargets"),
+	false,
+	TEXT("Record the external morph target weights inside animation insights. On default this is disabled, because it can slow down recording."),
+	ECVF_Default
+);
+
 struct FAnimTrace
 {
 	/** The various phases of anim graph processing */
@@ -114,7 +121,7 @@ struct FAnimTrace
 		float Thickness;
 		bool bPersistentLines;
 	};
-	
+
 	/** Reset Caches so a new trace can be started*/
 	ENGINE_API static void Reset();
 
