[MLDeformer] Crash fix when clearing the input animations array in the UI and when switching characters.

diff --git a/Engine/Plugins/Animation/MLDeformer/MLDeformerFramework/Source/MLDeformerFramework/Private/MLDeformerComponent.cpp b/Engine/Plugins/Animation/MLDeformer/MLDeformerFramework/Source/MLDeformerFramework/Private/MLDeformerComponent.cpp
--- a/Engine/Plugins/Animation/MLDeformer/MLDeformerFramework/Source/MLDeformerFramework/Private/MLDeformerComponent.cpp
+++ b/Engine/Plugins/Animation/MLDeformer/MLDeformerFramework/Source/MLDeformerFramework/Private/MLDeformerComponent.cpp
@@ -34,7 +34,7 @@ UMLDeformerComponent::UMLDeformerComponent(const FObjectInitializer& ObjectIniti
 void UMLDeformerComponent::Init()
 {
 	// If there is no deformer asset linked, release what we currently have.
-	if (DeformerAsset == nullptr)
+	if (!DeformerAsset)
 	{
 		ReleaseModelInstance();
 		return;
@@ -152,7 +152,7 @@ void UMLDeformerComponent::ReleaseModelInstance()
 USkeletalMeshComponent* UMLDeformerComponent::FindSkeletalMeshComponent(const UMLDeformerAsset* const Asset) const
 {
 	USkeletalMeshComponent* ResultingComponent = nullptr;
-	if (Asset != nullptr)
+	if (Asset)
 	{
 		// First search for a skeletal mesh component that uses the same skeletal mesh as the ML Deformer asset was trained on.
 		const UMLDeformerModel* Model = Asset ? Asset->GetModel() : nullptr;
@@ -162,7 +162,7 @@ USkeletalMeshComponent* UMLDeformerComponent::FindSkeletalMeshComponent(const UM
 
 			// Get a list of all skeletal mesh components on the actor.
 			TArray<USkeletalMeshComponent*> Components;
-			AActor* Actor = Cast<AActor>(GetOuter());
+			const AActor* Actor = Cast<AActor>(GetOuter());
 			if (Actor)
 			{
 				Actor->GetComponents<USkeletalMeshComponent>(Components);
@@ -203,11 +203,7 @@ void UMLDeformerComponent::Deactivate()
 	#endif
 
 	UnbindDelegates();
-	if (ModelInstance)
-	{
-		ModelInstance->ConditionalBeginDestroy();
-		ModelInstance = nullptr;
-	}
+	ReleaseModelInstance();
 	Super::Deactivate();
 }
 
diff --git a/Engine/Plugins/Animation/MLDeformer/MLDeformerFramework/Source/MLDeformerFramework/Private/MLDeformerModelInstance.cpp b/Engine/Plugins/Animation/MLDeformer/MLDeformerFramework/Source/MLDeformerFramework/Private/MLDeformerModelInstance.cpp
--- a/Engine/Plugins/Animation/MLDeformer/MLDeformerFramework/Source/MLDeformerFramework/Private/MLDeformerModelInstance.cpp
+++ b/Engine/Plugins/Animation/MLDeformer/MLDeformerFramework/Source/MLDeformerFramework/Private/MLDeformerModelInstance.cpp
@@ -218,16 +218,19 @@ void UMLDeformerModelInstance::UpdateBoneTransforms()
 			for (int32 Index = 0; Index < NumTrainingBones; ++Index)
 			{
 				const int32 ComponentBoneIndex = AssetBonesToSkelMeshMappings[Index];
-				const FTransform& ComponentSpaceTransform = LeaderTransforms[ComponentBoneIndex];
-				const int32 ParentIndex = RefSkel.GetParentIndex(ComponentBoneIndex);
-				if (LeaderTransforms.IsValidIndex(ParentIndex))
+				if (LeaderTransforms.IsValidIndex(ComponentBoneIndex))
 				{
-					TrainingBoneTransforms[Index] = ComponentSpaceTransform.GetRelativeTransform(LeaderTransforms[ParentIndex]);
-					TrainingBoneTransforms[Index].NormalizeRotation();
-				}
-				else
-				{
-					TrainingBoneTransforms[Index] = ComponentSpaceTransform;
+					const FTransform& ComponentSpaceTransform = LeaderTransforms[ComponentBoneIndex];
+					const int32 ParentIndex = RefSkel.GetParentIndex(ComponentBoneIndex);
+					if (LeaderTransforms.IsValidIndex(ParentIndex))
+					{
+						TrainingBoneTransforms[Index] = ComponentSpaceTransform.GetRelativeTransform(LeaderTransforms[ParentIndex]);
+						TrainingBoneTransforms[Index].NormalizeRotation();
+					}
+					else
+					{
+						TrainingBoneTransforms[Index] = ComponentSpaceTransform;
+					}
 				}
 			}
 		}
@@ -245,7 +248,10 @@ void UMLDeformerModelInstance::UpdateBoneTransforms()
 				for (int32 Index = 0; Index < NumTrainingBones; ++Index)
 				{
 					const int32 ComponentBoneIndex = AssetBonesToSkelMeshMappings[Index];
-					TrainingBoneTransforms[Index] = Transforms[ComponentBoneIndex];
+					if (Transforms.IsValidIndex(ComponentBoneIndex))
+					{
+						TrainingBoneTransforms[Index] = Transforms[ComponentBoneIndex];
+					}
 				}
 			}
 		}
@@ -266,7 +272,10 @@ void UMLDeformerModelInstance::UpdateBoneTransforms()
 			for (int32 Index = 0; Index < NumTrainingBones; ++Index)
 			{
 				const int32 ComponentBoneIndex = AssetBonesToSkelMeshMappings[Index];
-				TrainingBoneTransforms[Index] = DebugBoneSpaceTransforms[ComponentBoneIndex];
+				if (DebugBoneSpaceTransforms.IsValidIndex(ComponentBoneIndex))
+				{					
+					TrainingBoneTransforms[Index] = DebugBoneSpaceTransforms[ComponentBoneIndex];
+				}
 			}
 		}
 	}
diff --git a/Engine/Plugins/Animation/MLDeformer/MLDeformerFramework/Source/MLDeformerFrameworkEditor/Private/MLDeformerGeomCacheEditorModel.cpp b/Engine/Plugins/Animation/MLDeformer/MLDeformerFramework/Source/MLDeformerFrameworkEditor/Private/MLDeformerGeomCacheEditorModel.cpp
--- a/Engine/Plugins/Animation/MLDeformer/MLDeformerFramework/Source/MLDeformerFrameworkEditor/Private/MLDeformerGeomCacheEditorModel.cpp
+++ b/Engine/Plugins/Animation/MLDeformer/MLDeformerFramework/Source/MLDeformerFrameworkEditor/Private/MLDeformerGeomCacheEditorModel.cpp
@@ -266,6 +266,10 @@ namespace UE::MLDeformer
 					const int32 ArrayIndex = PropertyChangedEvent.GetArrayIndex(PropertyChangedEvent.GetMemberPropertyName().ToString());
 					GetEditor()->GetTimeSlider()->OnDeletedTrainingInputAnim(ArrayIndex);
 				}
+				else if (PropertyChangedEvent.ChangeType == EPropertyChangeType::ArrayClear)
+				{
+					ActiveTrainingInputAnimIndex = INDEX_NONE;
+				}
 			}
 			SetResamplingInputOutputsNeeded(true);
 			TriggerInputAssetChanged(true);
