MLDeformer: fix the dropmenu width on timeline

diff --git a/Engine/Plugins/Animation/MLDeformer/MLDeformerFramework/Source/MLDeformerFrameworkEditor/Private/SMLDeformerTimeline.cpp b/Engine/Plugins/Animation/MLDeformer/MLDeformerFramework/Source/MLDeformerFrameworkEditor/Private/SMLDeformerTimeline.cpp
--- a/Engine/Plugins/Animation/MLDeformer/MLDeformerFramework/Source/MLDeformerFrameworkEditor/Private/SMLDeformerTimeline.cpp
+++ b/Engine/Plugins/Animation/MLDeformer/MLDeformerFramework/Source/MLDeformerFrameworkEditor/Private/SMLDeformerTimeline.cpp
@@ -1479,7 +1479,7 @@ namespace UE::MLDeformer
 						[
 							SNew(SHorizontalBox)
 							+SHorizontalBox::Slot()
-							.AutoWidth()
+							.FillWidth(0.9)
 							.Padding(2.0f)
 							[
 								SNew(SComboBox<TSharedPtr<FMLDeformerTrainingInputAnimName>>)
@@ -1518,7 +1518,7 @@ namespace UE::MLDeformer
 								]
 							]
 							+SHorizontalBox::Slot()
-							.FillWidth(1.0f)
+							.FillWidth(0.1f)
 							.VAlign(VAlign_Center)
 							+SHorizontalBox::Slot()
 							.VAlign(VAlign_Center)
