[MLDeformer] Add support for LOD. Removed quality levels.

diff --git a/Engine/Plugins/Animation/MLDeformer/MLDeformerFramework/Source/MLDeformerFramework/Private/MLDeformerComponent.cpp b/Engine/Plugins/Animation/MLDeformer/MLDeformerFramework/Source/MLDeformerFramework/Private/MLDeformerComponent.cpp
--- a/Engine/Plugins/Animation/MLDeformer/MLDeformerFramework/Source/MLDeformerFramework/Private/MLDeformerComponent.cpp
+++ b/Engine/Plugins/Animation/MLDeformer/MLDeformerFramework/Source/MLDeformerFramework/Private/MLDeformerComponent.cpp
@@ -228,8 +228,9 @@ void UMLDeformerComponent::TickComponent(float DeltaTime, enum ELevelTick TickTy
 	if (TickType != ELevelTick::LEVELTICK_PauseTick)
 	{
 		if (ModelInstance &&
+			ModelInstance->GetModel() &&
 			SkelMeshComponent && 
-			SkelMeshComponent->GetPredictedLODLevel() == 0)
+			SkelMeshComponent->GetPredictedLODLevel() < ModelInstance->GetModel()->GetMaxNumLODs())
 		{
 			TRACE_CPUPROFILER_EVENT_SCOPE(UMLDeformerComponent::TickComponent)
 			
diff --git a/Engine/Plugins/Animation/MLDeformer/MLDeformerFramework/Source/MLDeformerFramework/Private/MLDeformerGeomCacheModel.cpp b/Engine/Plugins/Animation/MLDeformer/MLDeformerFramework/Source/MLDeformerFramework/Private/MLDeformerGeomCacheModel.cpp
--- a/Engine/Plugins/Animation/MLDeformer/MLDeformerFramework/Source/MLDeformerFramework/Private/MLDeformerGeomCacheModel.cpp
+++ b/Engine/Plugins/Animation/MLDeformer/MLDeformerFramework/Source/MLDeformerFramework/Private/MLDeformerGeomCacheModel.cpp
@@ -4,6 +4,7 @@
 #include "MLDeformerGeomCacheHelpers.h"
 #include "MLDeformerComponent.h"
 #include "MLDeformerAsset.h"
+#include "MLDeformerObjectVersion.h"
 #include "UObject/Object.h"
 #include "UObject/UObjectGlobals.h"
 #include "GeometryCache.h"
@@ -14,6 +15,8 @@
 
 void UMLDeformerGeomCacheModel::Serialize(FArchive& Archive)
 {
+	Archive.UsingCustomVersion(UE::MLDeformer::FMLDeformerObjectVersion::GUID);
+
 	#if WITH_EDITOR
 		if (Archive.IsSaving() && Archive.IsCooking())
 		{
diff --git a/Engine/Plugins/Animation/MLDeformer/MLDeformerFramework/Source/MLDeformerFramework/Private/MLDeformerGraphDebugDataInterface.cpp b/Engine/Plugins/Animation/MLDeformer/MLDeformerFramework/Source/MLDeformerFramework/Private/MLDeformerGraphDebugDataInterface.cpp
--- a/Engine/Plugins/Animation/MLDeformer/MLDeformerFramework/Source/MLDeformerFramework/Private/MLDeformerGraphDebugDataInterface.cpp
+++ b/Engine/Plugins/Animation/MLDeformer/MLDeformerFramework/Source/MLDeformerFramework/Private/MLDeformerGraphDebugDataInterface.cpp
@@ -156,11 +156,11 @@ namespace UE::MLDeformer
 			const UMLDeformerVizSettings* VizSettings = Model->GetVizSettings();
 			const UMLDeformerModelInstance* ModelInstance = DeformerComponent->GetModelInstance();
 
-			SkeletalMeshObject = ModelInstance->GetSkeletalMeshComponent()->MeshObject;
+			SkeletalMeshObject = (ModelInstance && ModelInstance->GetSkeletalMeshComponent()) ? ModelInstance->GetSkeletalMeshComponent()->MeshObject : nullptr;
 			VertexMapBufferSRV = Model->GetVertexMapBuffer().ShaderResourceViewRHI;
 			HeatMapMode = static_cast<int32>(VizSettings->GetHeatMapMode());
 			HeatMapMax = 1.0f / FMath::Max(VizSettings->GetHeatMapMax(), 0.00001f);
-			GroundTruthLerp = VizSettings->GetGroundTruthLerp();
+			GroundTruthLerp = (ModelInstance && ModelInstance->GetSkeletalMeshComponent()->GetPredictedLODLevel() == 0) ? VizSettings->GetGroundTruthLerp() : 0.0f;
 		}
 	}
 
@@ -188,11 +188,19 @@ namespace UE::MLDeformer
 		{
 			return false;
 		}
+
 		if (SkeletalMeshObject == nullptr || VertexMapBufferSRV == nullptr)
 		{
 			return false;
 		}
 
+		// Only support showing heatmaps in groundtruth mode on LOD 0.
+		const FSkeletalMeshRenderData& SkeletalMeshRenderData = SkeletalMeshObject->GetSkeletalMeshRenderData();
+		if (SkeletalMeshRenderData.GetPendingFirstLODIdx(0) != 0 && HeatMapMode == static_cast<int32>(EMLDeformerHeatMapMode::GroundTruth))
+		{
+			return false;
+		}
+
 		return true;
 	}
 
diff --git a/Engine/Plugins/Animation/MLDeformer/MLDeformerFramework/Source/MLDeformerFramework/Private/MLDeformerModel.cpp b/Engine/Plugins/Animation/MLDeformer/MLDeformerFramework/Source/MLDeformerFramework/Private/MLDeformerModel.cpp
--- a/Engine/Plugins/Animation/MLDeformer/MLDeformerFramework/Source/MLDeformerFramework/Private/MLDeformerModel.cpp
+++ b/Engine/Plugins/Animation/MLDeformer/MLDeformerFramework/Source/MLDeformerFramework/Private/MLDeformerModel.cpp
@@ -73,6 +73,8 @@ void UMLDeformerModel::Init(UMLDeformerAsset* InDeformerAsset)
 
 void UMLDeformerModel::Serialize(FArchive& Archive)
 {
+	Archive.UsingCustomVersion(UE::MLDeformer::FMLDeformerObjectVersion::GUID);
+
 	#if WITH_EDITOR
 		if (Archive.IsSaving())
 		{
@@ -163,6 +165,7 @@ void UMLDeformerModel::GetAssetRegistryTags(TArray<FAssetRegistryTag>& OutTags)
 	OutTags.Add(FAssetRegistryTag("MLDeformer.NumBaseMeshVerts", FString::FromInt(NumBaseMeshVerts), FAssetRegistryTag::TT_Numerical));
 	OutTags.Add(FAssetRegistryTag("MLDeformer.NumTargetMeshVerts", FString::FromInt(NumTargetMeshVerts), FAssetRegistryTag::TT_Numerical));
 	OutTags.Add(FAssetRegistryTag("MLDeformer.SkeletalMesh", SkeletalMesh ? FAssetData(SkeletalMesh).ToSoftObjectPath().ToString() : TEXT("None"), FAssetRegistryTag::TT_Alphabetical));
+	OutTags.Add(FAssetRegistryTag("MLDeformer.MaxNumLODs", FString::FromInt(GetMaxNumLODs()), FAssetRegistryTag::TT_Numerical));
 
 	#if WITH_EDITORONLY_DATA
 		OutTags.Add(FAssetRegistryTag("MLDeformer.NumBones", FString::FromInt(BoneIncludeList.Num()), FAssetRegistryTag::TT_Numerical));
@@ -255,6 +258,7 @@ bool UMLDeformerModel::IsCompatibleDebugActor(const AActor* Actor, UMLDeformerCo
 	return false;
 }
 
+
 #if WITH_EDITOR
 	void UMLDeformerModel::UpdateNumTargetMeshVertices()
 	{ 
diff --git a/Engine/Plugins/Animation/MLDeformer/MLDeformerFramework/Source/MLDeformerFramework/Private/MLDeformerModule.cpp b/Engine/Plugins/Animation/MLDeformer/MLDeformerFramework/Source/MLDeformerFramework/Private/MLDeformerModule.cpp
--- a/Engine/Plugins/Animation/MLDeformer/MLDeformerFramework/Source/MLDeformerFramework/Private/MLDeformerModule.cpp
+++ b/Engine/Plugins/Animation/MLDeformer/MLDeformerFramework/Source/MLDeformerFramework/Private/MLDeformerModule.cpp
@@ -20,7 +20,20 @@ namespace UE::MLDeformer
 		// Register an additional shader path for our shaders used inside the deformer graph system.
 		const FString PluginShaderDir = FPaths::Combine(IPluginManager::Get().FindPlugin(TEXT("MLDeformerFramework"))->GetBaseDir(), TEXT("Shaders"));
 		AddShaderSourceDirectoryMapping(TEXT("/Plugin/MLDeformerFramework"), PluginShaderDir);
+
+		CVarMLDeformerMaxLODLevels = IConsoleManager::Get().RegisterConsoleVariable(
+			TEXT("sg.MLDeformerMaxLODLevelsOnCook"),
+			8,
+			TEXT("Limit the number of ML Deformer LOD levels that can be generated. Can be used to setup per device or platform limits. Minimum value is 1."),
+			ECVF_Scalability);
 	}
+
+	void FMLDeformerModule::ShutdownModule()
+	{
+		IConsoleManager::Get().UnregisterConsoleObject(CVarMLDeformerMaxLODLevels);
+		IModuleInterface::ShutdownModule();
+	}
+
 }	// namespace UE::MLDeformer
 
 #undef LOCTEXT_NAMESPACE
diff --git a/Engine/Plugins/Animation/MLDeformer/MLDeformerFramework/Source/MLDeformerFramework/Private/MLDeformerMorphModel.cpp b/Engine/Plugins/Animation/MLDeformer/MLDeformerFramework/Source/MLDeformerFramework/Private/MLDeformerMorphModel.cpp
--- a/Engine/Plugins/Animation/MLDeformer/MLDeformerFramework/Source/MLDeformerFramework/Private/MLDeformerMorphModel.cpp
+++ b/Engine/Plugins/Animation/MLDeformer/MLDeformerFramework/Source/MLDeformerFramework/Private/MLDeformerMorphModel.cpp
@@ -5,6 +5,8 @@
 #include "MLDeformerModelInstance.h"
 #include "MLDeformerComponent.h"
 #include "MLDeformerMorphModelInputInfo.h"
+#include "MLDeformerObjectVersion.h"
+#include "MLDeformerModule.h"
 #include "Components/ExternalMorphSet.h"
 #include "Components/SkeletalMeshComponent.h"
 #include "Rendering/MorphTargetVertexInfoBuffers.h"
@@ -16,8 +18,6 @@
 UMLDeformerMorphModel::UMLDeformerMorphModel(const FObjectInitializer& ObjectInitializer)
 	: Super(ObjectInitializer)
 {
-	MorphTargetSet = MakeShared<FExternalMorphSet>();
-	MorphTargetSet->Name = GetClass()->GetFName();
 }
 
 void UMLDeformerMorphModel::Serialize(FArchive& Archive)
@@ -25,26 +25,81 @@ void UMLDeformerMorphModel::Serialize(FArchive& Archive)
 	TRACE_CPUPROFILER_EVENT_SCOPE(UMLDeformerMorphModel::Serialize)
 
 	Super::Serialize(Archive);
+	Archive.UsingCustomVersion(UE::MLDeformer::FMLDeformerObjectVersion::GUID);
 
 	// Check if we have initialized our compressed morph buffers.
 	bool bHasMorphData = false;
 	if (Archive.IsSaving())
 	{
 		// Strip editor only data on cook.
+		int32 NumLODs = GetNumLODs();
 		if (Archive.IsCooking())
 		{			
 			MorphTargetDeltas.Empty();
+
+			// Check if we want to limit the number of LODs (can be per platform/device).
+			UE::MLDeformer::FMLDeformerModule& MLDeformerModule = FModuleManager::LoadModuleChecked<UE::MLDeformer::FMLDeformerModule>("MLDeformerFramework");
+			const int32 MaxLODLevels = FMath::Clamp(MLDeformerModule.GetMaxLODLevelsOnCookCVar().GetInt(), 1, 1000);	// Limit to 1000 LODs, which should never be reached.
+
+			// Get lowest value between what we generated, console variable and the UI/property max lods value.
+			NumLODs = FMath::Min3(NumLODs, MaxLODLevels, GetMaxNumLODs());
+
+			UE_LOG(LogMLDeformer, Display, TEXT("Cooking MLD asset '%s' with %d LOD levels"), *GetFullName(), NumLODs);
+		}
+		else
+		{
+			// Get lowest number between how many LODs we have generated and the number of LODs we setup in the UI/Property.
+			NumLODs = FMath::Min(NumLODs, GetMaxNumLODs());
 		}
 
-		bHasMorphData = MorphTargetSet.IsValid() ? MorphTargetSet->MorphBuffers.IsMorphCPUDataValid() : false;
+		// Save all LOD levels, strip out LODs we don't want.
+		Archive << NumLODs;
+		for (int32 LOD = 0; LOD < NumLODs; ++LOD)
+		{
+			bHasMorphData = GetMorphTargetSet(LOD).IsValid() ? GetMorphTargetSet(LOD)->MorphBuffers.IsMorphCPUDataValid() : false;
+			Archive << bHasMorphData;
+
+			// Load or save the compressed morph buffers, if they exist.
+			if (bHasMorphData)
+			{
+				Archive << GetMorphTargetSet(LOD)->MorphBuffers;
+			}
+		}
 	}
-	Archive << bHasMorphData;
 
-	// Load or save the compressed morph buffers, if they exist.
-	if (bHasMorphData)
+	if (Archive.IsLoading())
 	{
-		check(MorphTargetSet.IsValid());
-		Archive << MorphTargetSet->MorphBuffers;
+		// If we only support 1 LOD, in older files.
+		if (Archive.CustomVer(UE::MLDeformer::FMLDeformerObjectVersion::GUID) < UE::MLDeformer::FMLDeformerObjectVersion::LODSupportAdded)
+		{
+			AddMorphSets(1);
+			Archive << bHasMorphData;
+
+			// Load or save the compressed morph buffers, if they exist.
+			if (bHasMorphData)
+			{
+				Archive << GetMorphTargetSet(0)->MorphBuffers;
+			}
+		}
+		else // We support multiple LOD levels.
+		{
+			int32 NumLODs = 1;
+			Archive << NumLODs;
+
+			check(GetNumLODs() == 0);
+			AddMorphSets(NumLODs);
+
+			for (int32 LOD = 0; LOD < NumLODs; ++LOD)
+			{
+				Archive << bHasMorphData;
+
+				// Load or save the compressed morph buffers, if they exist.
+				if (bHasMorphData)
+				{
+					Archive << GetMorphTargetSet(LOD)->MorphBuffers;
+				}
+			}
+		}
 	}
 }
 
@@ -78,17 +133,27 @@ void UMLDeformerMorphModel::GetAssetRegistryTags(TArray<FAssetRegistryTag>& OutT
 	OutTags.Add(FAssetRegistryTag("MLDeformer.MorphModel.ClampMorphWeights", bClampMorphWeights ? TEXT("True") : TEXT("False"), FAssetRegistryTag::TT_Alphabetical));
 	OutTags.Add(FAssetRegistryTag("MLDeformer.MorphModel.InvertMaskChannel", bInvertMaskChannel? TEXT("True") : TEXT("False"), FAssetRegistryTag::TT_Alphabetical));
 	OutTags.Add(FAssetRegistryTag("MLDeformer.MorphModel.IncludeNormals", bIncludeNormals ? TEXT("True") : TEXT("False"), FAssetRegistryTag::TT_Alphabetical));
-	OutTags.Add(FAssetRegistryTag("MLDeformer.MorphModel.NumMorphTargets", FString::FromInt(NumMorphTargets), FAssetRegistryTag::TT_Numerical));
+	OutTags.Add(FAssetRegistryTag("MLDeformer.MorphModel.NumMorphTargets", FString::FromInt(GetNumMorphTargets(0)), FAssetRegistryTag::TT_Numerical));
 	OutTags.Add(FAssetRegistryTag("MLDeformer.MorphModel.DeltaZeroThreshold", FString::Printf(TEXT("%f"), MorphDeltaZeroThreshold), FAssetRegistryTag::TT_Numerical));
 	OutTags.Add(FAssetRegistryTag("MLDeformer.MorphModel.CompressionLevel", FString::Printf(TEXT("%f"), MorphCompressionLevel), FAssetRegistryTag::TT_Numerical));
-	OutTags.Add(FAssetRegistryTag("MLDeformer.MorphModel.NumQualityLevels", FString::FromInt(QualityLevels.Num()), FAssetRegistryTag::TT_Numerical));
 	OutTags.Add(FAssetRegistryTag("MLDeformer.MorphModel.CompressedSize", FString::FromInt(CompressedMorphDataSizeInBytes), FAssetRegistryTag::TT_Numerical));
 	OutTags.Add(FAssetRegistryTag("MLDeformer.MorphModel.UncompressedSize", FString::FromInt(UncompressedMorphDataSizeInBytes), FAssetRegistryTag::TT_Numerical));
+	OutTags.Add(FAssetRegistryTag("MLDeformer.MorphModel.NumLODs", FString::FromInt(GetNumLODs()), FAssetRegistryTag::TT_Numerical));
+}
+
+int32 UMLDeformerMorphModel::GetNumMorphTargets(int32 LOD) const
+{
+	if (GetNumLODs() > LOD)
+	{
+		return GetMorphTargetSet(LOD).IsValid() ? GetMorphTargetSet(LOD)->MorphBuffers.GetNumMorphs() : 0;
+	}
+	return 0;
 }
 
 bool UMLDeformerMorphModel::CanDynamicallyUpdateMorphTargets() const
 {
-	return GetMorphTargetDeltas().Num() == (GetNumBaseMeshVerts() * GetNumMorphTargets());
+	const int32 LOD = 0;
+	return GetMorphTargetDeltas().Num() == (GetNumBaseMeshVerts() * GetNumMorphTargets(LOD));
 }
 
 UMLDeformerModelInstance* UMLDeformerMorphModel::CreateModelInstance(UMLDeformerComponent* Component)
@@ -106,23 +171,54 @@ void UMLDeformerMorphModel::SetMorphTargetDeltas(const TArray<FVector3f>& Deltas
 	MorphTargetDeltas = Deltas;
 }
 
-int32 UMLDeformerMorphModel::GetMorphTargetDeltaStartIndex(int32 BlendShapeIndex) const
+void UMLDeformerMorphModel::ClearMorphTargetSets()
+{
+	const int32 NumLODs = GetNumLODs();
+	for (int32 LOD = 0; LOD < NumLODs; ++LOD)
+	{
+		if (GetMorphTargetSet(LOD).IsValid())
+		{
+			FMorphTargetVertexInfoBuffers& MorphBuffer = GetMorphTargetSet(LOD)->MorphBuffers;
+			if (MorphBuffer.IsRHIIntialized() && MorphBuffer.IsInitialized())
+			{
+				ReleaseResourceAndFlush(&MorphBuffer);
+			}
+		}
+	}
+
+	MorphTargetSets.Empty();
+}
+
+void UMLDeformerMorphModel::AddMorphSets(int32 NumToAdd)
+{
+	for (int32 Index = 0; Index < NumToAdd; ++Index)
+	{
+		TSharedPtr<FExternalMorphSet> NewSet = MakeShared<FExternalMorphSet>();
+		NewSet->Name = GetClass()->GetFName();
+		MorphTargetSets.Add(NewSet);
+	}
+}
+
+int32 UMLDeformerMorphModel::GetMorphTargetDeltaStartIndex(int32 MorphTargetIndex) const
 {
 	if (MorphTargetDeltas.Num() == 0)
 	{
 		return INDEX_NONE;
 	}
 
-	return GetNumBaseMeshVerts() * BlendShapeIndex;
+	return GetNumBaseMeshVerts() * MorphTargetIndex;
 }
 
 void UMLDeformerMorphModel::BeginDestroy()
 {
-	if (MorphTargetSet)
+	for (int32 LOD = 0; LOD < GetNumLODs(); ++LOD)
 	{
-		BeginReleaseResource(&MorphTargetSet->MorphBuffers);
-		RenderCommandFence.BeginFence();
+		if (GetMorphTargetSet(LOD).IsValid())
+		{
+			BeginReleaseResource(&GetMorphTargetSet(LOD)->MorphBuffers);
+		}
 	}
+	RenderCommandFence.BeginFence();
 	Super::BeginDestroy();
 }
 
@@ -140,8 +236,12 @@ void UMLDeformerMorphModel::SetMorphTargetsErrorOrder(const TArray<int32>& Morph
 
 void UMLDeformerMorphModel::UpdateStatistics()
 {
-	NumMorphTargets = MorphTargetSet.IsValid() ? MorphTargetSet->MorphBuffers.GetNumMorphs() : 0;
-	CompressedMorphDataSizeInBytes = MorphTargetSet.IsValid() ? MorphTargetSet->MorphBuffers.GetMorphDataSizeInBytes() : 0;
+	CompressedMorphDataSizeInBytes = 0;
+	const int32 NumLODs = GetNumLODs();
+	for (int32 LOD = 0; LOD < NumLODs; ++LOD)
+	{
+		CompressedMorphDataSizeInBytes += GetMorphTargetSet(LOD).IsValid() ? GetMorphTargetSet(LOD)->MorphBuffers.GetMorphDataSizeInBytes() : 0;		
+	}
 	UncompressedMorphDataSizeInBytes = GetMorphTargetDeltas().Num() * GetMorphTargetDeltas().GetTypeSize();
 }
 
@@ -190,12 +290,12 @@ TArrayView<const int32> UMLDeformerMorphModel::GetMorphTargetErrorOrder() const
 
 TArrayView<const FMLDeformerMorphModelQualityLevel> UMLDeformerMorphModel::GetQualityLevels() const
 {
-	return QualityLevels;
+	return QualityLevels_DEPRECATED;
 }
 
 TArray<FMLDeformerMorphModelQualityLevel>& UMLDeformerMorphModel::GetQualityLevelsArray()
 {
-	return QualityLevels;
+	return QualityLevels_DEPRECATED;
 }
 
 float UMLDeformerMorphModel::GetMorphTargetError(int32 MorphIndex) const
@@ -210,14 +310,7 @@ void UMLDeformerMorphModel::SetMorphTargetError(int32 MorphIndex, float Error)
 
 int32 UMLDeformerMorphModel::GetNumActiveMorphs(int32 QualityLevel) const
 {
-	if (QualityLevels.IsEmpty() || MorphTargetErrorOrder.IsEmpty() || MorphTargetErrors.IsEmpty())
-	{
-		return FMath::Max<int32>(0, NumMorphTargets - 1);	// -1 Because this number includes the means morph target.
-	}
-
-	const int32 ClampedQualityLevel = FMath::Clamp<int32>(QualityLevel, 0, QualityLevels.Num() - 1);
-	const int32 NumActiveMorphs = FMath::Clamp<int32>(QualityLevels[ClampedQualityLevel].GetMaxActiveMorphs(), 0, NumMorphTargets - 1);	// -1 Because we want to exclude the means morph target.
-	return NumActiveMorphs;
+	return 0;
 }
 
 #if WITH_EDITOR
diff --git a/Engine/Plugins/Animation/MLDeformer/MLDeformerFramework/Source/MLDeformerFramework/Private/MLDeformerMorphModelInstance.cpp b/Engine/Plugins/Animation/MLDeformer/MLDeformerFramework/Source/MLDeformerFramework/Private/MLDeformerMorphModelInstance.cpp
--- a/Engine/Plugins/Animation/MLDeformer/MLDeformerFramework/Source/MLDeformerFramework/Private/MLDeformerMorphModelInstance.cpp
+++ b/Engine/Plugins/Animation/MLDeformer/MLDeformerFramework/Source/MLDeformerFramework/Private/MLDeformerMorphModelInstance.cpp
@@ -25,9 +25,9 @@ bool UMLDeformerMorphModelInstance::IsValidForDataProvider() const
 
 void UMLDeformerMorphModelInstance::PostTick(bool bExecuteCalled)
 {
-#if WITH_EDITOR
-	CopyDataFromCurrentDebugActor();
-#endif
+	#if WITH_EDITOR
+		CopyDataFromCurrentDebugActor();
+	#endif
 }
 
 int32 UMLDeformerMorphModelInstance::GetExternalMorphSetID() const
@@ -43,8 +43,11 @@ void UMLDeformerMorphModelInstance::BeginDestroy()
 		const UMLDeformerMorphModel* MorphModel = Cast<UMLDeformerMorphModel>(Model);
 		if (MorphModel)
 		{
-			const int32 LOD = 0;
-			SkeletalMeshComponent->RemoveExternalMorphSet(LOD, ExternalMorphSetID);
+			const int32 NumLODs = SkeletalMeshComponent->GetNumLODs();
+			for (int32 LOD = 0; LOD < NumLODs; ++LOD)
+			{
+				SkeletalMeshComponent->RemoveExternalMorphSet(LOD, ExternalMorphSetID);
+			}
 			SkeletalMeshComponent->RefreshExternalMorphTargetWeights();
 		}
 	}
@@ -70,79 +73,45 @@ void UMLDeformerMorphModelInstance::PostMLDeformerComponentInit()
 		// Get the morph model and its morph target set.
 		UMLDeformerMorphModel* MorphModel = Cast<UMLDeformerMorphModel>(Model);
 		check(MorphModel);
-		TSharedPtr<FExternalMorphSet> MorphTargetSet = MorphModel->GetMorphTargetSet();
 
-		// Register the morph set. This overwrites the existing one for this model, if it already exists.
-		// Only add to LOD 0 for now.
 		check(IsInGameThread());	// We don't want to call this multithreaded, as the AddExternalMorphSets etc isn't thread safe.
-		const int32 LOD = 0;
-		SkelMeshComponent->AddExternalMorphSet(LOD, ExternalMorphSetID, MorphTargetSet);
-
-		// Update the weight information in the Skeletal Mesh.
-		SkelMeshComponent->RefreshExternalMorphTargetWeights();
-
-		// When we're in editor mode, keep the CPU data around, so we can re-initialize when needed.
-#if WITH_EDITOR
-		MorphTargetSet->MorphBuffers.SetEmptyMorphCPUDataOnInitRHI(false);
-#else
-		MorphTargetSet->MorphBuffers.SetEmptyMorphCPUDataOnInitRHI(true);
-#endif
-
-		// Release the render resources, but only in an editor build.
-		// The non-editor build shouldn't do this, as then it can't initialize again. The non-editor build assumes
-		// that the data doesn't change and we don't need to re-init.
-		// In the editor build we have to re-initialize the render resources as the morph targets can change after (re)training, so
-		// that is why we release them here, and intialize them again after.
-		FMorphTargetVertexInfoBuffers& MorphBuffers = MorphTargetSet->MorphBuffers;
-#if WITH_EDITOR
-		BeginReleaseResource(&MorphBuffers);
-#endif
-
-		// Reinitialize the GPU compressed buffers.
-		if (MorphBuffers.IsMorphCPUDataValid() && MorphBuffers.GetNumMorphs() > 0)
-		{
-			// In a non-editor build this will clear the CPU data.
-			// That also means it can't re-init the resources later on again.
-			BeginInitResource(&MorphBuffers);
-		}
-
-		SetHasPostInitialized(true);
-	}
-}
-
-void UMLDeformerMorphModelInstance::Tick(float DeltaTime, float ModelWeight)
-{
-	// We don't want to call this multithreaded, as the AddExternalMorphSets etc isn't thread safe.
-	check(IsInGameThread());
-
-	// Detect changes in quality level.
-	const int32 CurrentQualityLevel = GetMLDeformerComponent()->GetQualityLevel();
-	if (CurrentQualityLevel != LastQualityLevel)
-	{
-		LastQualityLevel = CurrentQualityLevel;
-		MorphLerpAlpha = 0.0f;
 
-		const int LOD = 0;	// For now we only support LOD 0, as we can't setup an ML Deformer per LOD yet.
-		FExternalMorphSetWeights* WeightData = FindWeightData(LOD);
-		if (WeightData)
+		const int32 NumLODs = MorphModel->GetNumLODs();
+		for (int32 LOD = 0; LOD < NumLODs; ++LOD)
 		{
-			StartMorphWeights = WeightData->Weights;
+			TSharedPtr<FExternalMorphSet> MorphTargetSet = MorphModel->GetMorphTargetSet(LOD);
+
+			// Register the morph set. This overwrites the existing one for this model, if it already exists.
+			SkelMeshComponent->AddExternalMorphSet(LOD, ExternalMorphSetID, MorphTargetSet);
+
+			// Update the weight information in the Skeletal Mesh.
+			SkelMeshComponent->RefreshExternalMorphTargetWeights();
+
+			// When we're in editor mode, keep the CPU data around, so we can re-initialize when needed.
+			#if WITH_EDITOR
+				MorphTargetSet->MorphBuffers.SetEmptyMorphCPUDataOnInitRHI(false);
+			#else
+				MorphTargetSet->MorphBuffers.SetEmptyMorphCPUDataOnInitRHI(true);
+			#endif
+
+			// Only release render resources in the editor build. Editor builds must re-initialize the render resources
+			// as the morph targets can change after training. Non-editor builds assume the data does not change, and there
+			// is no need to re-init.
+			FMorphTargetVertexInfoBuffers& MorphBuffers = MorphTargetSet->MorphBuffers;
+			#if WITH_EDITOR
+				BeginReleaseResource(&MorphBuffers);
+			#endif
+
+			// Reinitialize the GPU compressed buffers.
+			if (MorphBuffers.IsMorphCPUDataValid() && MorphBuffers.GetNumMorphs() > 0)
+			{
+				// In a non-editor build this will clear the CPU data.
+				// That also means it can't re-init the resources later on again.
+				BeginInitResource(&MorphBuffers);
+			}
 		}
-	}
 
-	if (StartMorphWeights.IsEmpty())
-	{
-		const UMLDeformerMorphModel* MorphModel = Cast<UMLDeformerMorphModel>(Model);
-		StartMorphWeights.SetNumZeroed(MorphModel->GetNumMorphTargets());
-	}
-
-	Super::Tick(DeltaTime, ModelWeight);
-
-	// Update the morph lerp towards the target.
-	MorphLerpAlpha += DeltaTime * 10.0f;
-	if (MorphLerpAlpha > 1.0f)
-	{
-		MorphLerpAlpha = 1.0f;
+		SetHasPostInitialized(true);
 	}
 }
 
@@ -162,7 +131,7 @@ FExternalMorphSetWeights* UMLDeformerMorphModelInstance::FindWeightData(int32 LO
 
 void UMLDeformerMorphModelInstance::HandleZeroModelWeight()
 {
-	const int LOD = 0;	// For now we only support LOD 0, as we can't setup an ML Deformer per LOD yet.
+	const int32 LOD = SkeletalMeshComponent->GetPredictedLODLevel();
 	FExternalMorphSetWeights* WeightData = FindWeightData(LOD);
 	if (WeightData)
 	{
diff --git a/Engine/Plugins/Animation/MLDeformer/MLDeformerFramework/Source/MLDeformerFramework/Public/MLDeformerModel.h b/Engine/Plugins/Animation/MLDeformer/MLDeformerFramework/Source/MLDeformerFramework/Public/MLDeformerModel.h
--- a/Engine/Plugins/Animation/MLDeformer/MLDeformerFramework/Source/MLDeformerFramework/Public/MLDeformerModel.h
+++ b/Engine/Plugins/Animation/MLDeformer/MLDeformerFramework/Source/MLDeformerFramework/Public/MLDeformerModel.h
@@ -12,6 +12,7 @@
 #include "RenderResource.h"
 #include "Animation/AnimSequence.h"
 #include "MLDeformerCurveReference.h"
+#include "PerQualityLevelProperties.h"
 #include "MLDeformerModel.generated.h"
 
 class UMLDeformerAsset;
@@ -20,7 +21,6 @@ class UMLDeformerModelInstance;
 class UMLDeformerComponent;
 class UMLDeformerInputInfo;
 
-
 /** The channel to get the mask data from. */
 UENUM()
 enum class EMLDeformerMaskChannel : uint8
@@ -152,6 +152,13 @@ public:
 	 */
 	virtual bool DoesSupportCurves() const					{ return true; }
 
+	/**
+	 * Check if this model supports LOD.
+	 * When this returns false, the UI will not show options to setup the maximum number of LOD levels.
+	 * @return Returns true if LOD is supported, otherwise false is returned, which means it only works on LOD0.
+	 */
+	virtual bool DoesSupportLOD() const						{ return false; }
+
 	/**
 	 * Does this model support deformer quality levels?
 	 * For example Morph based models can disable certain morph targets based on this quality level.
@@ -159,6 +166,7 @@ public:
 	 * Morph based models on enable this on default.
 	 * @return Returns true when Deformer Quality is supported, otherwise false is returned.
 	 */
+	UE_DEPRECATED(5.4, "This method will be removed.")
 	virtual bool DoesSupportQualityLevels() const			{ return false; }
 
 	/**
@@ -216,6 +224,19 @@ public:
 	 */
 	virtual bool IsCompatibleDebugActor(const AActor* Actor, UMLDeformerComponent** OutDebugComponent = nullptr) const;
 
+	/** 
+	 * Get the maximum number of LOD levels that we will generate.
+ 	 * Some examples:
+	 * A value of 1 means we only store one LOD, which means LOD0. 
+	 * A value of 2 means we support this ML Deformer on LOD0 and LOD1.
+	 * A value of 3 means we support this ML Deformer on LOD0 and LOD1 and LOD2.
+	 * We never generate more LOD levels for the ML Deformer than number of LOD levels in the Skeletal Mesh, so if 
+	 * this value is set to 100, while the Skeletal Mesh has only 4 LOD levels, we will only generate and store 4 ML Deformer LODs.
+	 * The default value of 1 means we do not support this ML Deformer at LOD levels other than LOD0.
+	 * @ return The maximum number of LOD levels we will generate.
+	 */
+	const int32 GetMaxNumLODs() const		{ return MaxNumLODs; }
+
 #if WITH_EDITORONLY_DATA
 	/**
 	 * Check whether this model currently has a training mesh setup or not.
@@ -567,6 +588,7 @@ public:
 	static FName GetCurveIncludeListPropertyName()		{ return GET_MEMBER_NAME_CHECKED(UMLDeformerModel, CurveIncludeList); }
 	static FName GetMaxTrainingFramesPropertyName()		{ return GET_MEMBER_NAME_CHECKED(UMLDeformerModel, MaxTrainingFrames); }
 	static FName GetDeltaCutoffLengthPropertyName()		{ return GET_MEMBER_NAME_CHECKED(UMLDeformerModel, DeltaCutoffLength); }
+	static FName GetMaxNumLODsPropertyName()			{ return GET_MEMBER_NAME_CHECKED(UMLDeformerModel, MaxNumLODs); }
 
 	UE_DEPRECATED(5.3, "This property has been removed and shouldn't be used anymore.")
 	static FName GetShouldIncludeBonesPropertyName()	{ return GET_MEMBER_NAME_CHECKED(UMLDeformerModel, bIncludeBones_DEPRECATED); }
@@ -665,11 +687,25 @@ private:
 	UPROPERTY()
 	int32 NumTargetMeshVerts = 0;
 
+	/** 
+	 * How many Skeletal Mesh LOD levels should we generate MLD lods for at most?
+	 * Some examples:
+	 * A value of 1 means we only store one LOD, which means LOD0. 
+	 * A value of 2 means we support this ML Deformer on LOD0 and LOD1.
+	 * A value of 3 means we support this ML Deformer on LOD0 and LOD1 and LOD2.
+	 * We never generate more LOD levels for the ML Deformer than number of LOD levels in the Skeletal Mesh, so if 
+	 * this value is set to 100, while the Skeletal Mesh has only 4 LOD levels, we will only generate and store 4 ML Deformer LODs.
+	 * The default value of 1 means we do not support this ML Deformer at LOD levels other than LOD0.
+	 * When cooking, the console variable "sg.MLDeformer.MaxLODLevelsOnCook" can be used to set the maximum value per device or platform.
+	 */
+	UPROPERTY(EditAnywhere, Category = "LOD Generation Settings", meta = (ClampMin = "1"))
+	int32 MaxNumLODs = 1;
+
 	/** 
 	 * The information about the neural network inputs. This contains things such as bone names and curve names.
 	 */
 	UPROPERTY()
-	TObjectPtr<UMLDeformerInputInfo> InputInfo = nullptr;
+	TObjectPtr<UMLDeformerInputInfo> InputInfo;
 
 	/** This is an index per vertex in the mesh, indicating the imported vertex number from the source asset. */
 	UPROPERTY()
@@ -677,7 +713,7 @@ private:
 
 	/** The skeletal mesh that represents the linear skinned mesh. */
 	UPROPERTY(EditAnywhere, Category = "Base Mesh")
-	TObjectPtr<USkeletalMesh> SkeletalMesh = nullptr;
+	TObjectPtr<USkeletalMesh> SkeletalMesh;
 
 	/** The number of floats per bone in network input. */
 	UE_DEPRECATED(5.3, "This will be removed")
@@ -689,7 +725,7 @@ private:
 
 #if WITH_EDITORONLY_DATA
 	UPROPERTY()
-	TObjectPtr<UMLDeformerVizSettings> VizSettings = nullptr;
+	TObjectPtr<UMLDeformerVizSettings> VizSettings;
 
 	/** Specifies whether bone transformations should be included as inputs during the training process. */
 	UPROPERTY()
diff --git a/Engine/Plugins/Animation/MLDeformer/MLDeformerFramework/Source/MLDeformerFramework/Public/MLDeformerModule.h b/Engine/Plugins/Animation/MLDeformer/MLDeformerFramework/Source/MLDeformerFramework/Public/MLDeformerModule.h
--- a/Engine/Plugins/Animation/MLDeformer/MLDeformerFramework/Source/MLDeformerFramework/Public/MLDeformerModule.h
+++ b/Engine/Plugins/Animation/MLDeformer/MLDeformerFramework/Source/MLDeformerFramework/Public/MLDeformerModule.h
@@ -4,6 +4,7 @@
 
 #include "CoreTypes.h"
 #include "Modules/ModuleManager.h"
+#include "HAL/IConsoleManager.h"
 
 // The log category for the ML Deformer framework.
 MLDEFORMERFRAMEWORK_API DECLARE_LOG_CATEGORY_EXTERN(LogMLDeformer, Log, All);
@@ -19,6 +20,12 @@ namespace UE::MLDeformer
 	public:
 		// IModuleInterface overrides.
 		void StartupModule() override;
+		void ShutdownModule() override;
 		// ~END IModuleInterface overrides.
+
+		const IConsoleVariable& GetMaxLODLevelsOnCookCVar() const		{ return *CVarMLDeformerMaxLODLevels; }
+
+	private:
+		IConsoleVariable* CVarMLDeformerMaxLODLevels = nullptr;
 	};
 }	// namespace UE::MLDeformer
diff --git a/Engine/Plugins/Animation/MLDeformer/MLDeformerFramework/Source/MLDeformerFramework/Public/MLDeformerMorphModel.h b/Engine/Plugins/Animation/MLDeformer/MLDeformerFramework/Source/MLDeformerFramework/Public/MLDeformerMorphModel.h
--- a/Engine/Plugins/Animation/MLDeformer/MLDeformerFramework/Source/MLDeformerFramework/Public/MLDeformerMorphModel.h
+++ b/Engine/Plugins/Animation/MLDeformer/MLDeformerFramework/Source/MLDeformerFramework/Public/MLDeformerMorphModel.h
@@ -28,6 +28,7 @@ public:
 	// UMLDeformerModel overrides.
 	virtual bool IsNeuralNetworkOnGPU() const override				{ return false; }	// CPU based neural network.
 	virtual bool DoesSupportQualityLevels() const override			{ return true; }	// We can disable morph targets based on the Deformer LOD float value.
+	virtual bool DoesSupportLOD() const override					{ return true; }	// We support deformations in different LOD levels.
 	virtual void Serialize(FArchive& Archive) override;
 	virtual UMLDeformerModelInstance* CreateModelInstance(UMLDeformerComponent* Component) override;
 #if WITH_EDITOR
@@ -42,7 +43,9 @@ public:
 	virtual void GetAssetRegistryTags(TArray<FAssetRegistryTag>& OutTags) const override;
 	// ~END UObject overrides.
 
-	int32 GetNumMorphTargets() const								{ return NumMorphTargets; }
+	UE_DEPRECATED(5.4, "Please use the GetNumMorphTargets(LOD) instead.")
+	int32 GetNumMorphTargets() const								{ return GetNumMorphTargets(0); }
+	int32 GetNumMorphTargets(int32 LOD) const;
 	uint64 GetCompressedMorphDataSizeInBytes() const				{ return CompressedMorphDataSizeInBytes; }
 	uint64 GetUncompressedMorphDataSizeInBytes() const				{ return UncompressedMorphDataSizeInBytes; }
 
@@ -83,12 +86,16 @@ public:
 	static FName GetIncludeMorphTargetNormalsPropertyName()			{ return GET_MEMBER_NAME_CHECKED(UMLDeformerMorphModel, bIncludeNormals); }
 	static FName GetMaskChannelPropertyName()						{ return GET_MEMBER_NAME_CHECKED(UMLDeformerMorphModel, MaskChannel); }
 	static FName GetInvertMaskChannelPropertyName()					{ return GET_MEMBER_NAME_CHECKED(UMLDeformerMorphModel, bInvertMaskChannel); }
-	static FName GetNumMorphTargetsPropertyName()					{ return GET_MEMBER_NAME_CHECKED(UMLDeformerMorphModel, NumMorphTargets); }
 	static FName GetCompressedMorphDataSizeInBytesPropertyName()	{ return GET_MEMBER_NAME_CHECKED(UMLDeformerMorphModel, CompressedMorphDataSizeInBytes); }
 	static FName GetUncompressedMorphDataSizeInBytesPropertyName()	{ return GET_MEMBER_NAME_CHECKED(UMLDeformerMorphModel, UncompressedMorphDataSizeInBytes); }
-	static FName GetQualityLevelsPropertyName()						{ return GET_MEMBER_NAME_CHECKED(UMLDeformerMorphModel, QualityLevels); }
 	static FName GetClampMorphTargetWeightsPropertyName()			{ return GET_MEMBER_NAME_CHECKED(UMLDeformerMorphModel, bClampMorphWeights); }
 
+	UE_DEPRECATED(5.4, "This method will be removed.")
+	static FName GetQualityLevelsPropertyName()						{ return GET_MEMBER_NAME_CHECKED(UMLDeformerMorphModel, QualityLevels_DEPRECATED); }
+
+	UE_DEPRECATED(5.4, "This method will be removed.")
+	static FName GetNumMorphTargetsPropertyName()					{ return GET_MEMBER_NAME_CHECKED(UMLDeformerMorphModel, NumMorphTargets_DEPRECATED); }
+
 	UE_DEPRECATED(5.2, "Please use GetMorphDeltaZeroThresholdPropertyName instead.")
 	static FName GetMorphTargetDeltaThresholdPropertyName()			{ return GET_MEMBER_NAME_CHECKED(UMLDeformerMorphModel, MorphDeltaZeroThreshold); }
 
@@ -201,12 +208,19 @@ public:
 	 * The layout of this array is [morphdeltas_target0, morphdeltas_target1, ..., morphdeltas_targetN].
 	 * So the total number of items in the array returned equals (NumMorphTargets * NumBaseMeshVerts).
 	 */
-	const TArray<FVector3f>& GetMorphTargetDeltas() const		{ return MorphTargetDeltas; }
+	const TArray<FVector3f>& GetMorphTargetDeltas() const				{ return MorphTargetDeltas; }
 
 	/**
 	 * Get the morph target set.
 	 */
-	TSharedPtr<FExternalMorphSet> GetMorphTargetSet() const		{ return MorphTargetSet; }
+	UE_DEPRECATED(5.4, "This method has been deprecated, please use the GetMorphTargetSet that takes a LOD level parameter.")
+	TSharedPtr<FExternalMorphSet> GetMorphTargetSet() const				{ return MorphTargetSet_DEPRECATED; }
+	TSharedPtr<FExternalMorphSet> GetMorphTargetSet(int32 LOD) const	{ return MorphTargetSets.IsValidIndex(LOD) ? MorphTargetSets[LOD] : MorphTargetSets.Last(); }
+
+	/** Get the number of LOD levels that we have deltas for. */
+	int32 GetNumLODs() const											{ return MorphTargetSets.Num(); }
+	void ClearMorphTargetSets();
+	void AddMorphSets(int32 NumToAdd);
 
 	/**
 	 * Get the start index into the array of deltas (vectors3's), for a given morph target.
@@ -236,12 +250,14 @@ public:
 	 * Get the quality levels.
 	 * @return An array view of the quality levels.
 	 */
+	UE_DEPRECATED(5.4, "This method will be removed.")
 	TArrayView<const FMLDeformerMorphModelQualityLevel> GetQualityLevels() const;
 
 	/**
 	 * Get the quality levels, with access to modify the quality levels.
 	 * @return The array containing the quality levels.
 	 */
+	UE_DEPRECATED(5.4, "This method will be removed.")
 	TArray<FMLDeformerMorphModelQualityLevel>& GetQualityLevelsArray();
 	
 	/**
@@ -260,19 +276,15 @@ public:
 	 */
 	void SetMorphTargetError(int32 MorphIndex, float Error);
 
-	/** 
-	 * Get the number of active morph targets for a given quality level.
-	 * If no morph error data is available, or if no quality levels have been setup, the total number of morph targets is returned.
-	 * If there is such data, the quality level will be clamped to be within a valid range.
-	 * Finally, if the number of active morph targets specified by the quality level will be clamped within a valid range as well.
-	 * So basically this method will always return a valid number of morph targets, whatever the input is.
-	 * @param QualityLevel The quality level to get the number of active morph targets for.
-	 */
+	UE_DEPRECATED(5.4, "Please use the GetNumActiveMorphsForLOD.")
 	int32 GetNumActiveMorphs(int32 QualityLevel) const;
 
 private:
 	/** The compressed morph target data, ready for the GPU. */
-	TSharedPtr<FExternalMorphSet> MorphTargetSet;
+	TSharedPtr<FExternalMorphSet> MorphTargetSet_DEPRECATED;
+
+	/** A morph target set for each LOD. The number of LODs that have a morph target set can be smaller than the number of LOD levels in the Skeletal Mesh, but should never be larger. */
+	TArray<TSharedPtr<FExternalMorphSet>> MorphTargetSets;
 
 	/** 
 	 * Should we enable morph target weight clamping?
@@ -293,7 +305,7 @@ private:
 
 	/** The number of morph targets. */
 	UPROPERTY()
-	int32 NumMorphTargets = 0;
+	int32 NumMorphTargets_DEPRECATED = 0;
 
 	/** The compressed memory usage of the morph targets. This is approximately what this MLD asset will use in your packaged project. */
 	UPROPERTY()
@@ -333,8 +345,8 @@ private:
 	 * The number in each quality level represents the number of active morph targets for that quality level.
 	 * These numbers will be clamped internally to be within valid ranges in case they go beyond the amount of morphs that exist.
 	 */
-	UPROPERTY(EditAnywhere, Category = "Deformer Quality", meta = (ClampMin = "1"))
-	TArray<FMLDeformerMorphModelQualityLevel> QualityLevels;
+	UPROPERTY()
+	TArray<FMLDeformerMorphModelQualityLevel> QualityLevels_DEPRECATED;
 
 	/**
 	 * Include vertex normals in the morph targets?
diff --git a/Engine/Plugins/Animation/MLDeformer/MLDeformerFramework/Source/MLDeformerFramework/Public/MLDeformerMorphModelInstance.h b/Engine/Plugins/Animation/MLDeformer/MLDeformerFramework/Source/MLDeformerFramework/Public/MLDeformerMorphModelInstance.h
--- a/Engine/Plugins/Animation/MLDeformer/MLDeformerFramework/Source/MLDeformerFramework/Public/MLDeformerMorphModelInstance.h
+++ b/Engine/Plugins/Animation/MLDeformer/MLDeformerFramework/Source/MLDeformerFramework/Public/MLDeformerMorphModelInstance.h
@@ -28,7 +28,6 @@ public:
 	virtual void Init(USkeletalMeshComponent* SkelMeshComponent) override;
 	virtual void PostMLDeformerComponentInit() override;
 	virtual void BeginDestroy() override;
-	virtual void Tick(float DeltaTime, float ModelWeight) override;
 	virtual void HandleZeroModelWeight() override;
 	virtual bool IsValidForDataProvider() const override;
 	virtual void PostTick(bool bExecuteCalled) override;
@@ -43,29 +42,10 @@ public:
 	 */
 	FExternalMorphSetWeights* FindWeightData(int32 LOD) const;
 
-protected:
-#if WITH_EDITOR
-	/**
-	 * Copy the morph weights from a given morph model instance.
-	 * We assume that the other instance is compatible with this one.
-	 * @param SourceModelInstance The instance to copy the weights from.
-	 */
-	void CopyMorphWeightsFromMorphInstance(const UMLDeformerMorphModelInstance* SourceModelInstance);
-#endif
-
 protected:
 	/** The next free morph target set ID. This is used to generate unique ID's for each morph model. */
 	static TAtomic<int32> NextFreeMorphSetID;
 
 	/** The ID of the external morph target set for this instance. This gets initialized during Init. */
 	int32 ExternalMorphSetID = -1;
-
-	/** The morph target weights we are blending towards. */
-	TArray<float> StartMorphWeights;
-
-	/** The current lerp alpha (between 0 and 1). */
-	float MorphLerpAlpha = 1.0f;
-
-	/** The quality level in the previous tick. */
-	int32 LastQualityLevel = -1;
 };
diff --git a/Engine/Plugins/Animation/MLDeformer/MLDeformerFramework/Source/MLDeformerFramework/Public/MLDeformerObjectVersion.h b/Engine/Plugins/Animation/MLDeformer/MLDeformerFramework/Source/MLDeformerFramework/Public/MLDeformerObjectVersion.h
--- a/Engine/Plugins/Animation/MLDeformer/MLDeformerFramework/Source/MLDeformerFramework/Public/MLDeformerObjectVersion.h
+++ b/Engine/Plugins/Animation/MLDeformer/MLDeformerFramework/Source/MLDeformerFramework/Public/MLDeformerObjectVersion.h
@@ -16,6 +16,9 @@ namespace UE::MLDeformer
 			// Before any version changes were made.
 			BeforeCustomVersionWasAdded,
 
+			// We added LOD support.
+			LODSupportAdded,
+
 			// ----- New versions can be added above this line -----
 			VersionPlusOne,
 			LatestVersion = VersionPlusOne - 1
diff --git a/Engine/Plugins/Animation/MLDeformer/MLDeformerFramework/Source/MLDeformerFramework/Public/MLDeformerVizSettings.h b/Engine/Plugins/Animation/MLDeformer/MLDeformerFramework/Source/MLDeformerFramework/Public/MLDeformerVizSettings.h
--- a/Engine/Plugins/Animation/MLDeformer/MLDeformerFramework/Source/MLDeformerFramework/Public/MLDeformerVizSettings.h
+++ b/Engine/Plugins/Animation/MLDeformer/MLDeformerFramework/Source/MLDeformerFramework/Public/MLDeformerVizSettings.h
@@ -73,7 +73,6 @@ public:
 	void SetTrainingFrameNumber(int32 FrameNumber)			{ TrainingFrameNumber = FrameNumber; }
 	void SetTestingFrameNumber(int32 FrameNumber)			{ TestingFrameNumber = FrameNumber; }
 	void SetWeight(float InWeight)							{ Weight = InWeight; }
-	void SetQualityLevel(int32 InQualityLevel)				{ QualityLevel = FMath::Max<int32>(InQualityLevel, 0); }
 
 	FVector GetMeshSpacingOffsetVector() const				{ return FVector(MeshSpacing, 0.0f, 0.0f); }
 	float GetMeshSpacing() const							{ return MeshSpacing; }
@@ -99,11 +98,17 @@ public:
 	float GetWeight() const									{ return Weight; }
 	bool GetXRayDeltas() const								{ return bXRayDeltas; }
 	bool GetDrawVertexDeltas() const						{ return bDrawDeltas; }
-	int32 GetQualityLevel() const							{ return QualityLevel; }
 	const TArray<FMLDeformerCompareActor>& GetCompareActors() const { return CompareActors; }
 	TArray<FMLDeformerCompareActor>& GetCompareActors()		{ return CompareActors; }
 	FColor GetDebugBoundsColor() const						{ return DebugBoundsColor; }
 
+
+	UE_DEPRECATED(5.4, "This method will be removed.")
+	void SetQualityLevel(int32 InQualityLevel)				{ QualityLevel_DEPRECATED = FMath::Max<int32>(InQualityLevel, 0); }
+
+	UE_DEPRECATED(5.4, "This method will be removed.")
+	int32 GetQualityLevel() const							{ return QualityLevel_DEPRECATED; }
+
 	// Get property names.
 	static FName GetVisualizationModePropertyName()			{ return GET_MEMBER_NAME_CHECKED(UMLDeformerVizSettings, VisualizationMode); }
 	static FName GetMeshSpacingPropertyName()				{ return GET_MEMBER_NAME_CHECKED(UMLDeformerVizSettings, MeshSpacing); }
@@ -127,9 +132,11 @@ public:
 	static FName GetWeightPropertyName()					{ return GET_MEMBER_NAME_CHECKED(UMLDeformerVizSettings, Weight); }
 	static FName GetXRayDeltasPropertyName()				{ return GET_MEMBER_NAME_CHECKED(UMLDeformerVizSettings, bXRayDeltas); }
 	static FName GetDrawVertexDeltasPropertyName()			{ return GET_MEMBER_NAME_CHECKED(UMLDeformerVizSettings, bDrawDeltas); }
-	static FName GetQualityLevelPropertyName()				{ return GET_MEMBER_NAME_CHECKED(UMLDeformerVizSettings, QualityLevel); }
 	static FName GetCompareActorsPropertyName()				{ return GET_MEMBER_NAME_CHECKED(UMLDeformerVizSettings, CompareActors); }
 	static FName GetDebugBoundsColorPropertyName()			{ return GET_MEMBER_NAME_CHECKED(UMLDeformerVizSettings, DebugBoundsColor); }
+
+	UE_DEPRECATED(5.4, "This method will be removed.")
+	static FName GetQualityLevelPropertyName()				{ return GET_MEMBER_NAME_CHECKED(UMLDeformerVizSettings, QualityLevel_DEPRECATED); }
 #endif
 
 protected:
@@ -181,13 +188,9 @@ protected:
 	UPROPERTY(EditAnywhere, Category = "Live Settings", meta = (ClampMin = "0"))
 	int32 TestingFrameNumber = 0;
 
-	/*
-	 * The ML Deformer LOD value. This is a continuous value between 0 and 1, where 0 means maximum quality and 1 means the lowest quality.
-	 * Morph based models will disable certain morph targets when increasing this value, which will lead to lower quality, but better GPU performance.
-	 * Each model can decide how to use this value. Some models might not support this.
-	 */
-	UPROPERTY(EditAnywhere, Transient, Category = "Live Settings", meta = (ClampMin = "0"))
-	int32 QualityLevel = 0;
+	/** Deprecated quality level property, deprecated in 5.4. */
+	UPROPERTY(Transient)
+	int32 QualityLevel_DEPRECATED = 0;
 
 	/** Specify whether the heatmap is enabled or not. */
 	UPROPERTY(EditAnywhere, Category = "Live Settings")
diff --git a/Engine/Plugins/Animation/MLDeformer/MLDeformerFramework/Source/MLDeformerFrameworkEditor/Private/MLDeformerEditorModel.cpp b/Engine/Plugins/Animation/MLDeformer/MLDeformerFramework/Source/MLDeformerFrameworkEditor/Private/MLDeformerEditorModel.cpp
--- a/Engine/Plugins/Animation/MLDeformer/MLDeformerFramework/Source/MLDeformerFrameworkEditor/Private/MLDeformerEditorModel.cpp
+++ b/Engine/Plugins/Animation/MLDeformer/MLDeformerFramework/Source/MLDeformerFrameworkEditor/Private/MLDeformerEditorModel.cpp
@@ -15,6 +15,7 @@
 #include "MLDeformerSampler.h"
 #include "MLDeformerModelInstance.h"
 #include "MLDeformerTrainingInputAnim.h"
+#include "MLDeformerOctree.h"
 #include "AnimationEditorPreviewActor.h"
 #include "AnimationEditorViewportClient.h"
 #include "EditorModeManager.h"
@@ -47,6 +48,8 @@
 #include "EditorViewportClient.h"
 #include "DrawDebugHelpers.h"
 #include "Slate/SceneViewport.h"
+#include "Rendering/SkeletalMeshLODModel.h"
+#include "Async/ParallelFor.h"
 
 #define LOCTEXT_NAMESPACE "MLDeformerEditorModel"
 
@@ -118,6 +121,8 @@ namespace UE::MLDeformer
 				}
 			}
 		);
+
+		UpdateLODMappings();
 	}
 
 	void FMLDeformerEditorModel::CreateSamplers()
@@ -393,6 +398,75 @@ namespace UE::MLDeformer
 		return FMath::IsWithinInclusive<int32>(AnimIndex, 0, GetNumTrainingInputAnims());
 	}
 
+	void FMLDeformerEditorModel::UpdateLODMappings()
+	{
+		LODMappings.Reset();
+
+		const USkeletalMesh* SkelMesh = Model->GetSkeletalMesh();
+		if (!SkelMesh || !SkelMesh->GetImportedModel())
+		{
+			return;
+		}
+
+		// Check some expectations on the number of LODS.
+		const FSkeletalMeshModel* ImportedModel = SkelMesh->GetImportedModel();
+		int32 NumLODLevels = SkelMesh->GetLODNum();
+		if (ImportedModel->LODModels.Num() != NumLODLevels)
+		{
+			return;
+		}
+		NumLODLevels = FMath::Min(NumLODLevels, Model->GetMaxNumLODs());
+
+		const double StartTime = FPlatformTime::Seconds();
+
+		TArray<FSoftSkinVertex> VertexPositionsLODZero;
+		ImportedModel->LODModels[0].GetVertices(VertexPositionsLODZero);
+
+		// Build an octree that contains all our vertices.
+		// We can use this to quickly get the closest vertex index for a given position.
+		// This isn't perfectly accurate in some edge cases, but good enough.
+		FMLDeformerOctree Octree;
+		Octree.Build(VertexPositionsLODZero, /*MaxVertsPerOctreeNode*/300, /*MaxSplitHierarchyDepth*/6);
+
+		// For each additional LOD level we are interested in.
+		// We skip LOD 0 because we don't need a mapping for that one, as it's our source LOD that we map into.
+		LODMappings.SetNum(NumLODLevels);
+		for (int32 LodLevel = 1; LodLevel < NumLODLevels; ++LodLevel)
+		{
+			FLODInfo& LodInfo = LODMappings[LodLevel];
+
+			// Get the number of vertices.
+			TArray<FSoftSkinVertex> VertexPositions;
+			ImportedModel->LODModels[LodLevel].GetVertices(VertexPositions);
+			const int32 NumVerts = ImportedModel->LODModels[LodLevel].NumVertices;
+			LodInfo.VtxMappingToLODZero.SetNum(NumVerts);
+
+			// Process all vertices in batches, and run all batches in parallel.
+			const int32 BatchSize = 100; // Number of verts per batch.
+			const int32 NumBatches = (NumVerts / BatchSize) + 1;
+			ParallelFor(NumBatches, [&LodInfo, &Octree, &VertexPositions, NumVerts](int32 BatchIndex)
+			{
+				// Make sure we don't go past the last vertex.
+				const int32 StartVertex = BatchIndex * BatchSize;
+				if (StartVertex >= NumVerts)
+				{
+					return;
+				}
+
+				// Process all vertices inside this batch.
+				const int32 NumVertsInBatch = (StartVertex + BatchSize) < NumVerts ? BatchSize : FMath::Max(NumVerts - StartVertex, 0);
+				for (int32 VertexIndex = StartVertex; VertexIndex < StartVertex + NumVertsInBatch; ++VertexIndex)
+				{
+					// Find the best vertex index in LOD 0.
+					LodInfo.VtxMappingToLODZero[VertexIndex] = Octree.FindClosestMeshVertexIndex((FVector)VertexPositions[VertexIndex].Position, 0.1);
+				}
+			});
+		} // For all LOD levels
+
+		const double TotalTime = FPlatformTime::Seconds() - StartTime;
+		UE_LOG(LogMLDeformer, Display, TEXT("Finished Updating ML Deformer LOD mappings in %.2f seconds"), TotalTime);
+	}
+
 	void FMLDeformerEditorModel::UpdateMeshOffsetFactors()
 	{
 		const bool bIsDebugging = (GetEditor()->GetDebugActor() != nullptr);
@@ -551,6 +625,7 @@ namespace UE::MLDeformer
 		if (bNeedsAssetReinit)
 		{
 			TriggerInputAssetChanged();
+			UpdateLODMappings();
 			Model->InitVertexMap();
 			Model->InitGPUData();
 			bNeedsAssetReinit = false;
@@ -572,6 +647,7 @@ namespace UE::MLDeformer
 		UpdateMeshOffsetFactors();
 		UpdateActorTransforms();
 		UpdateLabels();
+		UpdateActorLODs();
 		CheckTrainingDataFrameChanged();
 		ApplyDebugActorTransforms();
 
@@ -596,7 +672,6 @@ namespace UE::MLDeformer
 				const UMLDeformerVizSettings* VizSettings = Model->GetVizSettings();
 				UMLDeformerComponent* DeformerComponent = EditorActor->GetMLDeformerComponent();
 				DeformerComponent->SetWeight(VizSettings->GetWeight());
-				DeformerComponent->SetQualityLevel(VizSettings->GetQualityLevel());
 			}
 		}
 	}
@@ -947,6 +1022,7 @@ namespace UE::MLDeformer
 		if (Property->GetFName() == UMLDeformerModel::GetSkeletalMeshPropertyName())
 		{
 			TriggerInputAssetChanged();
+			UpdateLODMappings();
 			SetResamplingInputOutputsNeeded(true);
 			Model->InitVertexMap();
 			Model->InitGPUData();
@@ -1036,6 +1112,13 @@ namespace UE::MLDeformer
 			UpdateIsReadyForTrainingState();
 			RefreshInputWidget();
 		}
+		else if (Property->GetFName() == UMLDeformerModel::GetMaxNumLODsPropertyName())
+		{
+			if (PropertyChangedEvent.ChangeType == EPropertyChangeType::ValueSet)
+			{
+				OnMaxNumLODsChanged();
+			}
+		}
 		else if (Property->GetFName() == UMLDeformerVizSettings::GetAnimPlaySpeedPropertyName())
 		{
 			UpdateTestAnimPlaySpeed();
@@ -1649,25 +1732,57 @@ namespace UE::MLDeformer
 			}
 
 			// Draw the deltas for the current frame.
-			if (bDrawDeltas)
+			USkeletalMesh* SkelMesh = Model->GetSkeletalMesh();
+			if (bDrawDeltas && Model->GetSkeletalMesh() && SkelMesh->GetImportedModel())
 			{
 				const FLinearColor DeltasColor = FMLDeformerEditorStyle::Get().GetColor("MLDeformer.Deltas.Color");
 				const FLinearColor DebugVectorsColor = FMLDeformerEditorStyle::Get().GetColor("MLDeformer.DebugVectors.Color");
 				const FLinearColor DebugVectorsColor2 = FMLDeformerEditorStyle::Get().GetColor("MLDeformer.DebugVectors.Color2");
 				const uint8 DepthGroup = VizSettings->GetXRayDeltas() ? 100 : 0;
-				for (int32 Index = 0; Index < LinearSkinnedPositions.Num(); ++Index)
+
+				const int32 LOD = FindEditorActor(ActorID_Train_Base)->GetSkeletalMeshComponent()->GetPredictedLODLevel();
+				if (LODMappings.IsValidIndex(LOD))
 				{
-					const int32 ArrayIndex = 3 * Index;
-					const FVector Delta(
-						VertexDeltas[ArrayIndex], 
-						VertexDeltas[ArrayIndex + 1], 
-						VertexDeltas[ArrayIndex + 2]);
-					const FVector VertexPos = (FVector)LinearSkinnedPositions[Index];
-					PDI->DrawLine(VertexPos, VertexPos + Delta, DeltasColor, DepthGroup);
+					const TArray<int32>& Mapping = LODMappings[LOD].VtxMappingToLODZero;	// Map from current LOD render vertex index into LOD0 render vertex index.
+					const FSkeletalMeshLODModel& LODModel = SkelMesh->GetImportedModel()->LODModels[LOD];
+
+					// Get the current LOD skinned positions. We use those as starting point for the delta lines.
+					Sampler->ExtractSkinnedPositions(LOD, SkinnedPositions);
+
+					// Render a delta for every render vertex in this LOD.
+					if (!Mapping.IsEmpty()) // If we have LOD mapping data.
+					{
+						for (int32 RenderVertexIndex = 0; RenderVertexIndex < Mapping.Num(); ++RenderVertexIndex)
+						{
+							const int32 LOD0VertexIndex = Mapping[RenderVertexIndex]; // Get the vertex number in LOD0.
+							const int32 VertexIndex = SkelMesh->GetImportedModel()->LODModels[0].MeshToImportVertexMap[LOD0VertexIndex]; // Get the imported vertex number.
+							const int32 ArrayIndex = 3 * VertexIndex; // 3 floats per delta, so multiply by 3.
+							const FVector Delta(
+								VertexDeltas[ArrayIndex], 
+								VertexDeltas[ArrayIndex + 1], 
+								VertexDeltas[ArrayIndex + 2]);
+					
+							// The skinned vertex positions have only the number of imported vertices number of positions, which can differ from the render vertex count.
+							const FVector VertexPos = (FVector)SkinnedPositions[LODModel.MeshToImportVertexMap[RenderVertexIndex]];
+							PDI->DrawLine(VertexPos, VertexPos + Delta, DeltasColor, DepthGroup);
+						}
+					}
+					else
+					{
+						for (int32 Index = 0; Index < LinearSkinnedPositions.Num(); ++Index)
+						{
+							const int32 ArrayIndex = 3 * Index;
+							const FVector Delta(
+								VertexDeltas[ArrayIndex], 
+								VertexDeltas[ArrayIndex + 1], 
+								VertexDeltas[ArrayIndex + 2]);
+							const FVector VertexPos = (FVector)LinearSkinnedPositions[Index];
+							PDI->DrawLine(VertexPos, VertexPos + Delta, DeltasColor, DepthGroup);
+						}				
+					}
 				}
 			}
 		}
-
 	}
 
 	void FMLDeformerEditorModel::DrawPIEDebugActors()
@@ -2897,6 +3012,29 @@ namespace UE::MLDeformer
 		CurrentTrainingFrame = -1;
 	}
 
+	void FMLDeformerEditorModel::UpdateActorLODs()
+	{
+		FMLDeformerEditorActor* MLDeformedActor = FindEditorActor(ActorID_Test_MLDeformed);
+		if (MLDeformedActor)
+		{		
+			USkeletalMeshComponent* BaseSkelMeshComponent = MLDeformedActor->GetSkeletalMeshComponent();		
+			if (BaseSkelMeshComponent)
+			{
+				const int32 LOD = BaseSkelMeshComponent->GetPredictedLODLevel();
+				for (FMLDeformerEditorActor* EditorActor : EditorActors)
+				{
+					if (EditorActor->GetTypeID() == ActorID_Test_Compare || EditorActor->GetTypeID() == ActorID_Test_Base)
+					{
+						if (EditorActor->GetSkeletalMeshComponent())
+						{
+							EditorActor->GetSkeletalMeshComponent()->SetPredictedLODLevel(LOD);
+						}
+					}
+				}
+			}
+		}
+	}
+
 	void FMLDeformerEditorModel::ApplyDebugActorTransforms()
 	{
 		AActor* DebugActor = GetEditor()->GetDebugActor();
diff --git a/Engine/Plugins/Animation/MLDeformer/MLDeformerFramework/Source/MLDeformerFrameworkEditor/Private/MLDeformerEditorModule.cpp b/Engine/Plugins/Animation/MLDeformer/MLDeformerFramework/Source/MLDeformerFrameworkEditor/Private/MLDeformerEditorModule.cpp
--- a/Engine/Plugins/Animation/MLDeformer/MLDeformerFramework/Source/MLDeformerFrameworkEditor/Private/MLDeformerEditorModule.cpp
+++ b/Engine/Plugins/Animation/MLDeformer/MLDeformerFramework/Source/MLDeformerFrameworkEditor/Private/MLDeformerEditorModule.cpp
@@ -6,7 +6,6 @@
 #include "MLDeformerEditorModel.h"
 #include "MLDeformerModel.h"
 #include "MLDeformerCurveReferenceCustomization.h"
-#include "MLDeformerMorphModelQualityLevelCustomization.h"
 #include "MLDeformerGeomCacheTrainingInputAnimCustomize.h"
 #include "SMLDeformerInputWidget.h"
 #include "Modules/ModuleManager.h"
@@ -26,7 +25,6 @@ namespace UE::MLDeformer
 		// Register object detail customizations.
 		FPropertyEditorModule& PropertyModule = FModuleManager::GetModuleChecked<FPropertyEditorModule>("PropertyEditor");
 		PropertyModule.RegisterCustomPropertyTypeLayout("MLDeformerCurveReference", FOnGetPropertyTypeCustomizationInstance::CreateStatic(&FMLDeformerCurveReferenceCustomization::MakeInstance) );
-		PropertyModule.RegisterCustomPropertyTypeLayout("MLDeformerMorphModelQualityLevel", FOnGetPropertyTypeCustomizationInstance::CreateStatic(&FMLDeformerMorphModelQualityLevelCustomization::MakeInstance) );
 		PropertyModule.RegisterCustomPropertyTypeLayout("MLDeformerGeomCacheTrainingInputAnim", FOnGetPropertyTypeCustomizationInstance::CreateStatic(&FMLDeformerGeomCacheTrainingInputAnimCustomization::MakeInstance) );
 		PropertyModule.NotifyCustomizationModuleChanged();
 
@@ -42,7 +40,6 @@ namespace UE::MLDeformer
 		{
 			FPropertyEditorModule& PropertyModule = FModuleManager::GetModuleChecked<FPropertyEditorModule>(TEXT("PropertyEditor"));
 			PropertyModule.UnregisterCustomPropertyTypeLayout(TEXT("MLDeformerCurveReference"));
-			PropertyModule.UnregisterCustomPropertyTypeLayout(TEXT("MLDeformerMorphModelQualityLevel"));
 			PropertyModule.UnregisterCustomPropertyTypeLayout(TEXT("MLDeformerGeomCacheTrainingInputAnim"));
 			PropertyModule.NotifyCustomizationModuleChanged();
 		}
diff --git a/Engine/Plugins/Animation/MLDeformer/MLDeformerFramework/Source/MLDeformerFrameworkEditor/Private/MLDeformerGeomCacheSampler.cpp b/Engine/Plugins/Animation/MLDeformer/MLDeformerFramework/Source/MLDeformerFrameworkEditor/Private/MLDeformerGeomCacheSampler.cpp
--- a/Engine/Plugins/Animation/MLDeformer/MLDeformerFramework/Source/MLDeformerFrameworkEditor/Private/MLDeformerGeomCacheSampler.cpp
+++ b/Engine/Plugins/Animation/MLDeformer/MLDeformerFramework/Source/MLDeformerFrameworkEditor/Private/MLDeformerGeomCacheSampler.cpp
@@ -104,7 +104,7 @@ namespace UE::MLDeformer
 				ParallelFor(NumBatches, [&](int32 BatchIndex)
 				{
 					const int32 StartVertex = BatchIndex * BatchSize;
-					if (StartVertex >= MeshInfo.NumVertices)
+					if (StartVertex >= MeshInfo.NumVertices || VertexDeltas.IsEmpty())
 					{
 						return;
 					}
diff --git a/Engine/Plugins/Animation/MLDeformer/MLDeformerFramework/Source/MLDeformerFrameworkEditor/Private/MLDeformerModelDetails.cpp b/Engine/Plugins/Animation/MLDeformer/MLDeformerFramework/Source/MLDeformerFrameworkEditor/Private/MLDeformerModelDetails.cpp
--- a/Engine/Plugins/Animation/MLDeformer/MLDeformerFramework/Source/MLDeformerFrameworkEditor/Private/MLDeformerModelDetails.cpp
+++ b/Engine/Plugins/Animation/MLDeformer/MLDeformerFramework/Source/MLDeformerFrameworkEditor/Private/MLDeformerModelDetails.cpp
@@ -58,6 +58,8 @@ namespace UE::MLDeformer
 		BaseMeshCategoryBuilder = &DetailLayoutBuilder->EditCategory("Base Mesh", FText::GetEmpty(), ECategoryPriority::Important);
 		InputOutputCategoryBuilder = &DetailLayoutBuilder->EditCategory("Inputs", FText::GetEmpty(), ECategoryPriority::Important);
 		TrainingSettingsCategoryBuilder = &DetailLayoutBuilder->EditCategory("Training Settings", FText::GetEmpty(), ECategoryPriority::Important);
+		LODSettingsCategoryBuilder = &DetailLayoutBuilder->EditCategory("LOD Generation Settings", FText::GetEmpty(), ECategoryPriority::Important);
+		LODSettingsCategoryBuilder->SetCategoryVisibility(Model->DoesSupportLOD());
 	}
 
 	void FMLDeformerModelDetails::CustomizeDetails(IDetailLayoutBuilder& DetailBuilder)
@@ -193,6 +195,8 @@ namespace UE::MLDeformer
 		TrainingSettingsCategoryBuilder->AddProperty(UMLDeformerModel::GetMaxTrainingFramesPropertyName(), UMLDeformerModel::StaticClass());
 		TrainingSettingsCategoryBuilder->AddProperty(UMLDeformerModel::GetDeltaCutoffLengthPropertyName(), UMLDeformerModel::StaticClass());
 		AddTrainingSettingsErrors();
+
+		LODSettingsCategoryBuilder->AddProperty(UMLDeformerModel::GetMaxNumLODsPropertyName(), UMLDeformerModel::StaticClass());
 	}
 
 	bool FMLDeformerModelDetails::FilterAnimSequences(const FAssetData& AssetData, USkeleton* Skeleton)
diff --git a/Engine/Plugins/Animation/MLDeformer/MLDeformerFramework/Source/MLDeformerFrameworkEditor/Private/MLDeformerMorphModelDetails.cpp b/Engine/Plugins/Animation/MLDeformer/MLDeformerFramework/Source/MLDeformerFrameworkEditor/Private/MLDeformerMorphModelDetails.cpp
--- a/Engine/Plugins/Animation/MLDeformer/MLDeformerFramework/Source/MLDeformerFrameworkEditor/Private/MLDeformerMorphModelDetails.cpp
+++ b/Engine/Plugins/Animation/MLDeformer/MLDeformerFramework/Source/MLDeformerFrameworkEditor/Private/MLDeformerMorphModelDetails.cpp
@@ -52,8 +52,6 @@ namespace UE::MLDeformer
 		MaskGroup.AddPropertyRow(DetailLayoutBuilder->GetProperty(UMLDeformerMorphModel::GetMaskChannelPropertyName(), UMLDeformerMorphModel::StaticClass()));
 		MaskGroup.AddPropertyRow(DetailLayoutBuilder->GetProperty(UMLDeformerMorphModel::GetInvertMaskChannelPropertyName(), UMLDeformerMorphModel::StaticClass()));
 
-		MorphTargetCategoryBuilder->AddProperty(DetailLayoutBuilder->GetProperty(UMLDeformerMorphModel::GetQualityLevelsPropertyName(), UMLDeformerMorphModel::StaticClass()));
-
 		if (MorphModel && !MorphModel->CanDynamicallyUpdateMorphTargets())
 		{
 			const FText DeltaCountMismatchErrorText = LOCTEXT("MorphDeltaCountMismatch", "Dynamic morph target updates disabled until retrained. This is because the vertex count changed after the model was trained.");
diff --git a/Engine/Plugins/Animation/MLDeformer/MLDeformerFramework/Source/MLDeformerFrameworkEditor/Private/MLDeformerMorphModelEditorModel.cpp b/Engine/Plugins/Animation/MLDeformer/MLDeformerFramework/Source/MLDeformerFrameworkEditor/Private/MLDeformerMorphModelEditorModel.cpp
--- a/Engine/Plugins/Animation/MLDeformer/MLDeformerFramework/Source/MLDeformerFrameworkEditor/Private/MLDeformerMorphModelEditorModel.cpp
+++ b/Engine/Plugins/Animation/MLDeformer/MLDeformerFramework/Source/MLDeformerFrameworkEditor/Private/MLDeformerMorphModelEditorModel.cpp
@@ -37,6 +37,15 @@ namespace UE::MLDeformer
 		return false;
 	}
 
+	void FMLDeformerMorphModelEditorModel::OnMaxNumLODsChanged()
+	{
+		UpdateLODMappings();
+		if (GetMorphModel()->CanDynamicallyUpdateMorphTargets())
+		{
+			InitEngineMorphTargets(GetMorphModel()->GetMorphTargetDeltas());
+		}
+	}
+
 	void FMLDeformerMorphModelEditorModel::OnPropertyChanged(FPropertyChangedEvent& PropertyChangedEvent)
 	{
 		const FProperty* Property = PropertyChangedEvent.Property;
@@ -65,31 +74,14 @@ namespace UE::MLDeformer
 		{
 			ClampMorphTargetNumber();
 		}
-		else if (Property->GetFName() == UMLDeformerMorphModel::GetQualityLevelsPropertyName())
-		{
-			if (PropertyChangedEvent.ChangeType == EPropertyChangeType::ArrayAdd)
-			{
-				TArray<FMLDeformerMorphModelQualityLevel>& QualityLevels = GetMorphModel()->GetQualityLevelsArray();				
-				if (QualityLevels.Num() == 1)
-				{
-					const int32 NewQualityValue = FMath::Max(GetMorphModel()->GetNumMorphTargets() - 1, 1);
-					QualityLevels[QualityLevels.Num() - 1].SetMaxActiveMorphs(NewQualityValue);
-				}
-				else
-				{
-					int32 NewQuality = static_cast<int32>(QualityLevels[QualityLevels.Num() - 2].GetMaxActiveMorphs() * 0.75f);
-					NewQuality = FMath::Clamp<int32>(NewQuality, 1, GetMorphModel()->GetNumMorphTargets() - 1);
-					QualityLevels[QualityLevels.Num() - 1].SetMaxActiveMorphs(NewQuality);
-				}
-			}
-		}
 	}
 
 	void FMLDeformerMorphModelEditorModel::ClampMorphTargetNumber()
 	{
+		const int32 LOD = 0;
 		const UMLDeformerMorphModel* MorphModel = GetMorphModel();			
 		UMLDeformerMorphModelVizSettings* MorphViz = GetMorphModelVizSettings();
-		const int32 NumMorphTargets = MorphModel->GetMorphTargetSet().IsValid() ? MorphModel->GetMorphTargetSet()->MorphBuffers.GetNumMorphs() : 0;
+		const int32 NumMorphTargets = MorphModel->GetMorphTargetSet(LOD).IsValid() ? MorphModel->GetMorphTargetSet(LOD)->MorphBuffers.GetNumMorphs() : 0;
 		const int32 ClampedMorphTargetNumber = (NumMorphTargets > 0) ? FMath::Min<int32>(MorphViz->GetMorphTargetNumber(), NumMorphTargets - 1) : 0;
 		MorphViz->SetMorphTargetNumber(ClampedMorphTargetNumber);
 	}
@@ -168,8 +160,7 @@ namespace UE::MLDeformer
 	void FMLDeformerMorphModelEditorModel::UpdateMorphErrorValues(TArrayView<UMorphTarget*> MorphTargets)
 	{
 		if (MorphTargets.IsEmpty())
-		{
-			
+		{			
 			return;
 		}
 
@@ -178,13 +169,18 @@ namespace UE::MLDeformer
 		UMLDeformerMorphModel* MorphModel = GetMorphModel();
 		const TArray<FFloatInterval>& MinMaxMorphWeights = MorphModel->GetMorphTargetsMinMaxWeights();
 
+		const int32 NumMorphs = MinMaxMorphWeights.Num();
+
 		// Preallocate space for the standard deviation of each morph target.
 		TArray<float> ErrorValues;
-		ErrorValues.SetNumZeroed(MorphTargets.Num() - 1);
+		if (NumMorphs > 0)
+		{
+			ErrorValues.SetNumZeroed(NumMorphs - 1);
+		}
 
 		const int32 LOD = 0;
 		TArray<float> DeltaLengths;
-		for (int32 MorphIndex = 0; MorphIndex < MorphTargets.Num() - 1; ++MorphIndex)	// We have one extra morph for the means, skip that one.
+		for (int32 MorphIndex = 0; MorphIndex < NumMorphs - 1; ++MorphIndex)	// We have one extra morph for the means, skip that one.
 		{
 			const UMorphTarget* MorphTarget = MorphTargets[MorphIndex + 1];
 
@@ -211,20 +207,23 @@ namespace UE::MLDeformer
 
 		// Build a list of array indices, so we know the order in which things got sorted.
 		TArray<int32> SortedIndices;
-		SortedIndices.SetNumUninitialized(MorphTargets.Num() - 1);
-		for (int32 Index = 0; Index < SortedIndices.Num(); ++Index)
+		if (NumMorphs > 0)
 		{
-			SortedIndices[Index] = Index;
-		}
-
-		// Now that we have a list of standard deviations, sort them.
-		SortedIndices.Sort
-		(
-			[&ErrorValues](const int32& IndexA, const int32& IndexB)
+			SortedIndices.SetNumUninitialized(NumMorphs - 1);
+			for (int32 Index = 0; Index < SortedIndices.Num(); ++Index)
 			{
-				return ErrorValues[IndexA] > ErrorValues[IndexB];
+				SortedIndices[Index] = Index;
 			}
-		);
+
+			// Now that we have a list of standard deviations, sort them.
+			SortedIndices.Sort
+			(
+				[&ErrorValues](const int32& IndexA, const int32& IndexB)
+				{
+					return ErrorValues[IndexA] > ErrorValues[IndexB];
+				}
+			);
+		}
 
 		// Update the morph model with the newly calculated error values.
 		MorphModel->SetMorphTargetsErrorOrder(SortedIndices, ErrorValues);
@@ -430,7 +429,7 @@ namespace UE::MLDeformer
 		check(!Model->GetVertexMap().IsEmpty());
 
 		FScopedSlowTask Task(NumMorphTargets, LOCTEXT("CreateMorphTargetProgress", "Creating morph targets"));
-		Task.MakeDialog(false);	
+		Task.MakeDialog(false);
 
 		USkeletalMesh* SkelMesh = Model->GetSkeletalMesh();
 		FSkeletalMeshRenderData* RenderData = SkelMesh->GetResourceForRendering();
@@ -564,18 +563,6 @@ namespace UE::MLDeformer
 			BeginInitResource(&OutMorphBuffers);
 		}
 
-		// Update the editor actor skel mesh components for all the ones that also have an ML Deformer on it.
-		for (FMLDeformerEditorActor* EditorActor : EditorActors)
-		{
-			const UMLDeformerComponent* MLDeformerComponent = EditorActor->GetMLDeformerComponent();
-			USkeletalMeshComponent* SkelMeshComponent = EditorActor->GetSkeletalMeshComponent();
-			if (SkelMeshComponent && MLDeformerComponent)
-			{
-				SkelMeshComponent->RefreshExternalMorphTargetWeights();
-			}
-		}
-
-		UpdateMemoryUsage();
 		Task.EnterProgressFrame();
 	}
 
@@ -637,7 +624,7 @@ namespace UE::MLDeformer
 		ZeroDeltasByLengthThreshold(MorphTargetDeltas, MorphModel->GetMorphDeltaZeroThreshold());
 
 		// Turn the delta buffer in a set of engine morph targets.
-		const int32 LOD = 0;
+		const int32 LODZero = 0;
 		const bool bIncludeNormals = MorphModel->GetIncludeMorphTargetNormals();
 		const EMLDeformerMaskChannel MaskChannel = MorphModel->GetMaskChannel();
 		const bool bInvertMaskChannel = MorphModel->GetInvertMaskChannel();
@@ -648,7 +635,7 @@ namespace UE::MLDeformer
 			MorphTargets,
 			Deltas, 
 			TEXT("MLDeformerMorph_"), 
-			LOD,
+			LODZero,
 			MorphModel->GetMorphDeltaZeroThreshold(),
 			bIncludeNormals,
 			MaskChannel,
@@ -657,14 +644,24 @@ namespace UE::MLDeformer
 		// Analyze the error values of the morph targets.
 		UpdateMorphErrorValues(MorphTargets);
 
-		// Now compress the morph targets to GPU friendly buffers.
-		check(MorphModel->GetMorphTargetSet().IsValid());
-		FMorphTargetVertexInfoBuffers& MorphBuffers = MorphModel->GetMorphTargetSet()->MorphBuffers;
-		CompressMorphTargets(MorphBuffers, MorphTargets, LOD, MorphModel->GetMorphCompressionLevel());
+		// Transfer morphs to the LOD levels.
+		TransferMorphTargets(MorphTargets);
+
+		// Resize to the new desired size.
+		MorphModel->ClearMorphTargetSets();
+		const int32 NumLODs = !MorphTargets.IsEmpty() ? MorphTargets[0]->GetMorphLODModels().Num() : 0;
+		MorphModel->AddMorphSets(NumLODs);
 
-		if (MorphBuffers.GetNumBatches() == 0 || MorphBuffers.GetNumMorphs() == 0)
+		// Now compress the morph targets to GPU friendly buffers.
+		for (int32 LOD = 0; LOD < NumLODs; ++LOD)
 		{
-			MorphBuffers = FMorphTargetVertexInfoBuffers();
+			FMorphTargetVertexInfoBuffers& MorphBuffers = MorphModel->GetMorphTargetSet(LOD)->MorphBuffers;
+			CompressMorphTargets(MorphBuffers, MorphTargets, LOD, MorphModel->GetMorphCompressionLevel());
+
+			if (MorphBuffers.GetNumBatches() == 0 || MorphBuffers.GetNumMorphs() == 0)
+			{
+				MorphBuffers = FMorphTargetVertexInfoBuffers();
+			}
 		}
 
 		MorphModel->UpdateStatistics();
@@ -674,6 +671,119 @@ namespace UE::MLDeformer
 		{
 			MorphTarget->ConditionalBeginDestroy();
 		}
+
+		// Update the editor actor skel mesh components for all the ones that also have an ML Deformer on it.
+		for (FMLDeformerEditorActor* EditorActor : EditorActors)
+		{
+			UMLDeformerComponent* MLDeformerComponent = EditorActor->GetMLDeformerComponent();
+			USkeletalMeshComponent* SkelMeshComponent = EditorActor->GetSkeletalMeshComponent();
+			if (SkelMeshComponent && MLDeformerComponent)
+			{
+				SkelMeshComponent->RefreshExternalMorphTargetWeights();			
+				MLDeformerComponent->UpdateSkeletalMeshComponent();
+			}
+		}
+
+		UpdateMemoryUsage();
+	}
+
+	void FMLDeformerMorphModelEditorModel::TransferMorphTargets(TArray<UMorphTarget*> MorphTargetsLODZero)
+	{
+		const USkeletalMesh* SkelMesh = Model->GetSkeletalMesh();
+		if (!SkelMesh || !SkelMesh->GetImportedModel())
+		{
+			return;
+		}
+
+		const int32 MaxLodWithMorphs = Model->GetMaxNumLODs();
+		const int32 NumLODs = FMath::Min(SkelMesh->GetLODNum(), MaxLodWithMorphs);
+		const int32 NumTaskItems = 1 + (NumLODs - 1) * MorphTargetsLODZero.Num();	// +1 because we see preparing the lookup tables as one task item as well.
+		FScopedSlowTask Task(NumTaskItems, LOCTEXT("TransferMorphTargetProgress", "Generating morph target LODs"));
+		Task.MakeDialog(false);	
+
+		const double StartTime = FPlatformTime::Seconds();
+
+		FSkeletalMeshRenderData* RenderData = SkelMesh->GetResourceForRendering();
+		check(RenderData);
+
+		// Build a mapping table to eliminate linear searches.
+		// For every morph target build a map that maps a render vertex to a morph target vertex.
+		const FSkeletalMeshLODModel& LODModelZero = SkelMesh->GetImportedModel()->LODModels[0];
+		const int32 NumRenderVerticesLODZero = LODModelZero.NumVertices;
+		TArray<TMap<int32, int32>> RenderVertexToMorphVertexLodZero;
+		RenderVertexToMorphVertexLodZero.SetNum(MorphTargetsLODZero.Num());
+		for (int32 MorphIndex = 0; MorphIndex < MorphTargetsLODZero.Num(); ++MorphIndex)
+		{
+			UMorphTarget* MorphTarget = MorphTargetsLODZero[MorphIndex];
+			const TArray<FMorphTargetDelta>& MorphVertices = MorphTarget->GetMorphLODModels()[0].Vertices;
+			const int32 NumMorphVerts = MorphVertices.Num();
+			RenderVertexToMorphVertexLodZero[MorphIndex].Reserve(NumMorphVerts);
+			for (int32 MorphTargetVertexIndex = 0; MorphTargetVertexIndex < NumMorphVerts; ++MorphTargetVertexIndex)
+			{
+				const int32 RenderVertexIndex = MorphVertices[MorphTargetVertexIndex].SourceIdx;
+				RenderVertexToMorphVertexLodZero[MorphIndex].Add(RenderVertexIndex, MorphTargetVertexIndex);
+			}
+		}
+
+		Task.EnterProgressFrame();
+
+		for (int32 LOD = 1; LOD < NumLODs; ++LOD)
+		{
+			for (int32 MorphIndex = 0; MorphIndex < MorphTargetsLODZero.Num(); ++MorphIndex)
+			{
+				// Add an LOD to this morph target.
+				UMorphTarget* MorphTarget = MorphTargetsLODZero[MorphIndex];
+				MorphTarget->GetMorphLODModels().AddDefaulted();
+				FMorphTargetLODModel& MorphLODModel = MorphTarget->GetMorphLODModels().Last();
+
+				const FSkeletalMeshLODModel& LODModel = SkelMesh->GetImportedModel()->LODModels[LOD];
+				const int32 NumRenderVertices = LODModel.NumVertices;
+
+				// Initialize the morph target LOD level.
+				MorphLODModel.Reset();
+				MorphLODModel.bGeneratedByEngine = true;
+				MorphLODModel.Vertices.Reserve(NumRenderVertices);
+
+				const TArray<int32>& MappingToLODZero = LODMappings[LOD].VtxMappingToLODZero;
+				const FMorphTargetLODModel& MorphLODModelZero = MorphTarget->GetMorphLODModels()[0];
+
+				// Add all vertices.
+				for (int32 RenderVertexIndex = 0; RenderVertexIndex < NumRenderVertices; ++RenderVertexIndex)
+				{
+					// Try to locate a vertex in the morph target that uses the same render vertex index.
+					const int32* MorphVertexIndexLODZero = RenderVertexToMorphVertexLodZero[MorphIndex].Find(MappingToLODZero[RenderVertexIndex]);
+					const FMorphTargetDelta* DeltaInLODZero = MorphVertexIndexLODZero ? &MorphLODModelZero.Vertices[*MorphVertexIndexLODZero] : nullptr;
+					
+					// Make sure we found one, if we didn't find it, we can skip this vertex.
+					if (!DeltaInLODZero)
+					{
+						continue;
+					}
+
+					// Add the vertex to the morph target.
+					MorphLODModel.Vertices.AddDefaulted();
+					FMorphTargetDelta& MorphTargetDelta = MorphLODModel.Vertices.Last();
+					MorphTargetDelta.PositionDelta = DeltaInLODZero->PositionDelta;
+					MorphTargetDelta.SourceIdx = RenderVertexIndex;
+					MorphTargetDelta.TangentZDelta = DeltaInLODZero->TangentZDelta;
+
+					// Make sure we update the list of sections that we touch.
+					int32 RenderSection = INDEX_NONE;
+					int32 TempVertexIndex = INDEX_NONE;
+					RenderData->LODRenderData[LOD].GetSectionFromVertexIndex(RenderVertexIndex, RenderSection, TempVertexIndex);
+					if (RenderSection != INDEX_NONE)
+					{
+						MorphLODModel.SectionIndices.AddUnique(RenderSection);
+					}
+				}
+				MorphLODModel.NumBaseMeshVerts = NumRenderVertices;
+				MorphLODModel.NumVertices = NumRenderVertices;
+				MorphLODModel.Vertices.Shrink();
+			} // For each morph target.
+		} // For each LOD to generate.
+
+		const double TotalTime = FPlatformTime::Seconds() - StartTime;
+		UE_LOG(LogMLDeformer, Display, TEXT("Finished Morph Target LOD generation in %.2f seconds"), TotalTime);
 	}
 
 	void FMLDeformerMorphModelEditorModel::Render(const FSceneView* View, FViewport* Viewport, FPrimitiveDrawInterface* PDI)
diff --git a/Engine/Plugins/Animation/MLDeformer/MLDeformerFramework/Source/MLDeformerFrameworkEditor/Private/MLDeformerMorphModelQualityLevelCustomization.cpp b/Engine/Plugins/Animation/MLDeformer/MLDeformerFramework/Source/MLDeformerFrameworkEditor/Private/MLDeformerMorphModelQualityLevelCustomization.cpp
deleted file mode 100644
diff --git a/Engine/Plugins/Animation/MLDeformer/MLDeformerFramework/Source/MLDeformerFrameworkEditor/Private/MLDeformerMorphModelVizSettingsDetails.cpp b/Engine/Plugins/Animation/MLDeformer/MLDeformerFramework/Source/MLDeformerFrameworkEditor/Private/MLDeformerMorphModelVizSettingsDetails.cpp
--- a/Engine/Plugins/Animation/MLDeformer/MLDeformerFramework/Source/MLDeformerFrameworkEditor/Private/MLDeformerMorphModelVizSettingsDetails.cpp
+++ b/Engine/Plugins/Animation/MLDeformer/MLDeformerFramework/Source/MLDeformerFrameworkEditor/Private/MLDeformerMorphModelVizSettingsDetails.cpp
@@ -53,7 +53,7 @@ namespace UE::MLDeformer
 			const UMLDeformerMorphModelInstance* MorphInstance = Cast<UMLDeformerMorphModelInstance>(MLDeformerComponent->GetModelInstance());
 			if (MorphInstance)
 			{
-				const int32 LOD = 0;
+				const int32 LOD = MorphInstance->GetFinalSkeletalMeshComponent()->GetPredictedLODLevel();
 				FExternalMorphSetWeights* WeightData = MorphInstance->FindWeightData(LOD);
 				if (WeightData)
 				{
@@ -91,14 +91,25 @@ namespace UE::MLDeformer
 				(
 					[MorphEditorModel]()
 					{
+						int32 LOD = 0;
+						const UMLDeformerComponent* MLDeformerComponent = MorphEditorModel->FindMLDeformerComponent(ActorID_Test_MLDeformed);
+						if (MLDeformerComponent)
+						{
+							const UMLDeformerMorphModelInstance* MorphInstance = Cast<UMLDeformerMorphModelInstance>(MLDeformerComponent->GetModelInstance());
+							if (MorphInstance)
+							{
+								LOD = MorphInstance->GetFinalSkeletalMeshComponent()->GetPredictedLODLevel();
+							}
+						}
+
 						const UMLDeformerMorphModel* MorphModelPtr = Cast<UMLDeformerMorphModel>(MorphEditorModel->GetModel());					
 						const int32 NumActiveMorphs = GetNumActiveMorphs(MorphEditorModel);
-						const int32 NumMorphs = MorphModelPtr->GetNumMorphTargets();
+						const int32 NumMorphs = MorphModelPtr->GetNumMorphTargets(LOD);
 						const float Percentage = (NumMorphs > 0) ? (NumActiveMorphs / static_cast<float>(NumMorphs)) * 100.0f : 100.0f;
 
 						FNumberFormattingOptions TwoDigitsFormat;
 						TwoDigitsFormat.SetMaximumFractionalDigits(0);
-						return FText::Format(LOCTEXT("ActiveMorphsFormat", "{0} / {1} ({2} %)"), NumActiveMorphs, NumMorphs, FText::AsNumber(Percentage, &TwoDigitsFormat));
+						return FText::Format(LOCTEXT("ActiveMorphsFormat", "{0} / {1} ({2} %) - LOD {3}"), NumActiveMorphs, NumMorphs, FText::AsNumber(Percentage, &TwoDigitsFormat), LOD);
 					}
 				)
 			];
diff --git a/Engine/Plugins/Animation/MLDeformer/MLDeformerFramework/Source/MLDeformerFrameworkEditor/Private/MLDeformerOctree.cpp b/Engine/Plugins/Animation/MLDeformer/MLDeformerFramework/Source/MLDeformerFrameworkEditor/Private/MLDeformerOctree.cpp
new file mode 100644
--- /dev/null
+++ b/Engine/Plugins/Animation/MLDeformer/MLDeformerFramework/Source/MLDeformerFrameworkEditor/Private/MLDeformerOctree.cpp
@@ -0,0 +1,227 @@
+// Copyright Epic Games, Inc. All Rights Reserved.
+
+#include "MLDeformerOctree.h"
+#include "SceneManagement.h"
+
+namespace UE::MLDeformer
+{
+	void FMLDeformerOctree::Build(const TArray<FSoftSkinVertex>& InVertices, int32 MaxNumVertsPerNode, int32 MaxDepth)
+	{
+		RootNode = MakeUnique<FMLDeformerOctree::FNode>();
+		MeshVertices = InVertices;
+
+		// First calculate the bounding box that contains all points as root node.
+		RootNode->AABB.Init();
+		const int32 NumVerts = MeshVertices.Num();
+		RootNode->VertexIndices.SetNum(NumVerts);
+		for (int32 VertexIndex = 0; VertexIndex < NumVerts; ++VertexIndex)
+		{
+			const FVector VertexPos = (FVector)MeshVertices[VertexIndex].Position;
+			RootNode->AABB += VertexPos;
+			RootNode->VertexIndices[VertexIndex] = VertexIndex;
+		}
+
+		// Recursively split if desired.
+		if (RootNode->VertexIndices.Num() > MaxNumVertsPerNode && MaxDepth > 0)
+		{
+			Split(RootNode.Get(), MaxNumVertsPerNode, MaxDepth, 1);
+		}
+	}
+
+	void FMLDeformerOctree::Render(FPrimitiveDrawInterface* PDI)
+	{
+		RecursiveRender(RootNode.Get(), PDI);
+	}
+
+	int32 FMLDeformerOctree::FindClosestMeshVertexIndex(FVector VertexPos, double FallBackDistanceThreshold) const
+	{
+		int32 ClosestVertexIndex = INDEX_NONE;
+		double ClosestVertexDist = DBL_MAX;
+
+		// Try to find the node this vertex would be in.
+		const FNode* Node = FindNodeForPosition(RootNode.Get(), VertexPos);
+		if (Node)
+		{
+			// Try and find the closest vertex inside here.
+			const int32 NumVerts = Node->VertexIndices.Num();
+			for (int32 Index = 0; Index < NumVerts; ++Index)
+			{
+				const int32 MeshVertexIndex = Node->VertexIndices[Index];
+				const FVector NodeVertexPos = (FVector)MeshVertices[MeshVertexIndex].Position;
+
+				// If we have an exact match, look no further.
+				if (NodeVertexPos == VertexPos)
+				{
+					return MeshVertexIndex;
+				}
+
+				// Update the closest match if this one is closer.
+				const double DistSq = (NodeVertexPos - VertexPos).SquaredLength();
+				if (DistSq < ClosestVertexDist)
+				{
+					ClosestVertexDist = DistSq;
+					ClosestVertexIndex = MeshVertexIndex;
+				}
+			}
+		}
+
+		// Worst case fallback scenario, just check against all vertices in the mesh.
+		// We do this when there is no matching node, or when our distance is too large.
+		// Remember that the LOD levels can have different topology, so this can happen.			
+		if (ClosestVertexIndex == INDEX_NONE || ClosestVertexDist > FallBackDistanceThreshold)
+		{
+			ClosestVertexDist = DBL_MAX;
+			const int32 NumVerts = MeshVertices.Num();
+			for (int32 Index = 0; Index < NumVerts; ++Index)
+			{
+				const FVector NodeVertexPos = (FVector)MeshVertices[Index].Position;
+
+				// If we have an exact match, look no further.
+				if (NodeVertexPos == VertexPos)
+				{
+					return Index;
+				}
+
+				// Update the closest match if this one is closer.
+				const double DistSq = (NodeVertexPos - VertexPos).SquaredLength();
+				if (DistSq < ClosestVertexDist)
+				{
+					ClosestVertexDist = DistSq;
+					ClosestVertexIndex = Index;
+				}
+			}
+		}
+
+		return ClosestVertexIndex;
+	}
+
+	void FMLDeformerOctree::Split(FNode* Parent, int32 MaxNumVertsPerNode, int32 MaxDepth, int32 CurDepth)
+	{
+		// Create the child nodes.
+		Parent->ChildNodes.Reserve(8);
+		for (int32 Index = 0; Index < 8; Index++)
+		{
+			Parent->ChildNodes.Add(MakeUnique<FMLDeformerOctree::FNode>());
+			Parent->ChildNodes.Last()->Parent = Parent;
+		}
+
+		// Get the parent's box center and extents.
+		FVector Center;
+		FVector Extents;
+		Parent->AABB.GetCenterAndExtents(Center, Extents);
+		Extents *= 0.5; // Our splits are half the size.
+
+		// Split the parent AABB into 8 boxes.
+		// Top half.
+		Parent->ChildNodes[0]->AABB = FBox::BuildAABB(Center + FVector(-Extents.X, -Extents.Y,  Extents.Z), Extents);
+		Parent->ChildNodes[1]->AABB = FBox::BuildAABB(Center + FVector( Extents.X, -Extents.Y,  Extents.Z), Extents);
+		Parent->ChildNodes[2]->AABB = FBox::BuildAABB(Center + FVector(-Extents.X,  Extents.Y,  Extents.Z), Extents);
+		Parent->ChildNodes[3]->AABB = FBox::BuildAABB(Center + FVector( Extents.X,  Extents.Y,  Extents.Z), Extents);
+
+		// Bottom half.
+		Parent->ChildNodes[4]->AABB = FBox::BuildAABB(Center + FVector(-Extents.X, -Extents.Y, -Extents.Z), Extents);
+		Parent->ChildNodes[5]->AABB = FBox::BuildAABB(Center + FVector( Extents.X, -Extents.Y, -Extents.Z), Extents);
+		Parent->ChildNodes[6]->AABB = FBox::BuildAABB(Center + FVector(-Extents.X,  Extents.Y, -Extents.Z), Extents);
+		Parent->ChildNodes[7]->AABB = FBox::BuildAABB(Center + FVector( Extents.X,  Extents.Y, -Extents.Z), Extents);
+
+		// Pass the parent vertices to the child nodes.
+		for (int32 Index = 0; Index < Parent->ChildNodes.Num(); ++Index)
+		{
+			FNode* ChildNode = Parent->ChildNodes[Index].Get();
+			const FBox& ChildAABB = ChildNode->AABB;
+
+			const int32 NumParentVerts = Parent->VertexIndices.Num();
+			for (int32 VertexIndex = 0; VertexIndex < NumParentVerts; ++VertexIndex)
+			{
+				const int32 MeshVertexIndex = Parent->VertexIndices[VertexIndex];
+				const FVector Pos = (FVector)MeshVertices[MeshVertexIndex].Position;
+				if (ChildAABB.IsInsideOrOn(Pos))
+				{
+					ChildNode->VertexIndices.Add(MeshVertexIndex);	// TODO: reserve upfront?
+				}
+			}
+		}
+
+		// Get rid of the parent node vertices, as we transferred them into the children now.
+		Parent->VertexIndices.Empty();
+
+		// Remove empty children.
+		for (int32 Index = 0; Index < Parent->ChildNodes.Num();)
+		{
+			if (Parent->ChildNodes[Index]->VertexIndices.IsEmpty())
+			{
+				Parent->ChildNodes.RemoveAt(Index);
+			}
+			else
+			{
+				Index++;
+			}
+		}
+
+		// Split further if desired.
+		for (int32 Index = 0; Index < Parent->ChildNodes.Num(); ++Index)
+		{
+			FNode* ChildNode = Parent->ChildNodes[Index].Get();
+			if (ChildNode->VertexIndices.Num() > MaxNumVertsPerNode && CurDepth < MaxDepth)
+			{
+				Split(Parent->ChildNodes[Index].Get(), MaxNumVertsPerNode, MaxDepth, CurDepth + 1);
+			}
+		}
+	}
+
+	FMLDeformerOctree::FNode* FMLDeformerOctree::FindNodeForPosition(FNode* Node, const FVector Pos) const
+	{
+		if (Node->AABB.IsInsideOrOn(Pos))
+		{
+			for (int32 ChildIndex = 0; ChildIndex < Node->ChildNodes.Num(); ++ChildIndex)
+			{
+				FNode* ChildNode = Node->ChildNodes[ChildIndex].Get();
+				if (ChildNode->AABB.IsInside(Pos))
+				{
+					return FindNodeForPosition(ChildNode, Pos);
+				}
+			}
+
+			return Node;
+		}
+
+		return nullptr;
+	}
+
+	void FMLDeformerOctree::RecursiveRender(FNode* Node, FPrimitiveDrawInterface* PDI)
+	{
+		for (int32 ChildIndex = 0; ChildIndex < Node->ChildNodes.Num(); ++ChildIndex)
+		{
+			FNode* ChildNode = Node->ChildNodes[ChildIndex].Get();
+			RecursiveRender(ChildNode, PDI);
+		}
+
+		if (Node->VertexIndices.IsEmpty())
+		{
+			return;
+		}
+
+		FVector BoxVerts[8];
+		Node->AABB.GetVertices(BoxVerts);
+
+		const FLinearColor BoxColor = FLinearColor::Blue;
+
+		// Front side.
+		PDI->DrawLine(BoxVerts[1], BoxVerts[5], BoxColor, 0);
+		PDI->DrawLine(BoxVerts[0], BoxVerts[3], BoxColor, 0);
+		PDI->DrawLine(BoxVerts[0], BoxVerts[1], BoxColor, 0);
+		PDI->DrawLine(BoxVerts[3], BoxVerts[5], BoxColor, 0);
+
+		// Back side.
+		PDI->DrawLine(BoxVerts[6], BoxVerts[7], BoxColor, 0);
+		PDI->DrawLine(BoxVerts[2], BoxVerts[4], BoxColor, 0);
+		PDI->DrawLine(BoxVerts[6], BoxVerts[2], BoxColor, 0);
+		PDI->DrawLine(BoxVerts[7], BoxVerts[4], BoxColor, 0);
+
+		// Connections between front to back sides.
+		PDI->DrawLine(BoxVerts[2], BoxVerts[0], BoxColor, 0);
+		PDI->DrawLine(BoxVerts[4], BoxVerts[3], BoxColor, 0);
+		PDI->DrawLine(BoxVerts[7], BoxVerts[5], BoxColor, 0);
+		PDI->DrawLine(BoxVerts[6], BoxVerts[1], BoxColor, 0);
+	}
+}	// namespace UE::MLDeformer
diff --git a/Engine/Plugins/Animation/MLDeformer/MLDeformerFramework/Source/MLDeformerFrameworkEditor/Private/MLDeformerOctree.h b/Engine/Plugins/Animation/MLDeformer/MLDeformerFramework/Source/MLDeformerFrameworkEditor/Private/MLDeformerOctree.h
new file mode 100644
--- /dev/null
+++ b/Engine/Plugins/Animation/MLDeformer/MLDeformerFramework/Source/MLDeformerFrameworkEditor/Private/MLDeformerOctree.h
@@ -0,0 +1,50 @@
+// Copyright Epic Games, Inc. All Rights Reserved.
+
+#pragma once
+
+#include "CoreTypes.h"
+#include "Rendering/SkeletalMeshLODModel.h"
+#include "Templates/UniquePtr.h"
+
+class FPrimitiveDrawInterface;
+
+namespace UE::MLDeformer
+{
+	/**
+	 * An octree around a mesh consisting of FSoftSkinVertex vertices.
+	 * This class is used to accelerate finding a close matching vertex in a lower LOD, inside another LOD.
+	 * For example, for a given vertex position in LOD 1, we want to find the best matching vertex in LOD 0.
+	 * We do that by finding a close match. But instead of always iterating over all vertices in LOD 0, we would 
+	 * first only check a small subset of nearby vertices.
+	 */
+	class FMLDeformerOctree
+	{
+		struct FNode;
+
+	public:
+		void Build(const TArray<FSoftSkinVertex>& InVertices, int32 MaxNumVertsPerNode = 200, int32 MaxDepth = 5);
+		void Render(FPrimitiveDrawInterface* PDI);
+		int32 FindClosestMeshVertexIndex(FVector VertexPos, double FallBackDistanceThreshold = 0.1) const;
+
+	private:
+		void Split(FNode* Parent, int32 MaxNumVertsPerNode, int32 MaxDepth, int32 CurDepth);
+		void RecursiveRender(FNode* Node, FPrimitiveDrawInterface* PDI);
+		FNode* FindNodeForPosition(FNode* Node, const FVector Pos) const;
+
+	private:
+		/* The octree node. */
+		struct FNode
+		{
+			TArray<int32> VertexIndices;			// Points inside the MeshVertices member array.
+			TArray<TUniquePtr<FNode>> ChildNodes;	// Doesn't need to be 8, can be less, or empty.
+			FNode* Parent = nullptr;				// Pointer to parent node, or nullptr in case of the root node.
+			FBox AABB;								// The bounding box of this node.
+		};
+
+		/** The root node of the octree. */
+		TUniquePtr<FNode> RootNode;
+
+		/** The actual array of vertices. */
+		TArray<FSoftSkinVertex> MeshVertices;
+	};
+}	// namespace UE::MLDeformer
diff --git a/Engine/Plugins/Animation/MLDeformer/MLDeformerFramework/Source/MLDeformerFrameworkEditor/Private/MLDeformerSampler.cpp b/Engine/Plugins/Animation/MLDeformer/MLDeformerFramework/Source/MLDeformerFrameworkEditor/Private/MLDeformerSampler.cpp
--- a/Engine/Plugins/Animation/MLDeformer/MLDeformerFramework/Source/MLDeformerFrameworkEditor/Private/MLDeformerSampler.cpp
+++ b/Engine/Plugins/Animation/MLDeformer/MLDeformerFramework/Source/MLDeformerFrameworkEditor/Private/MLDeformerSampler.cpp
@@ -209,6 +209,11 @@ namespace UE::MLDeformer
 		}
 	}
 
+	void FMLDeformerSampler::ExtractSkinnedPositions(int32 LODIndex, TArray<FVector3f>& OutPositions)
+	{
+		ExtractSkinnedPositions(LODIndex, BoneMatrices, TempVertexPositions, OutPositions);
+	}
+
 	void FMLDeformerSampler::ExtractSkinnedPositions(int32 LODIndex, TArray<FMatrix44f>& InBoneMatrices, TArray<FVector3f>& TempPositions, TArray<FVector3f>& OutPositions) const
 	{
 		OutPositions.Reset();
diff --git a/Engine/Plugins/Animation/MLDeformer/MLDeformerFramework/Source/MLDeformerFrameworkEditor/Private/MLDeformerVizSettingsDetails.cpp b/Engine/Plugins/Animation/MLDeformer/MLDeformerFramework/Source/MLDeformerFrameworkEditor/Private/MLDeformerVizSettingsDetails.cpp
--- a/Engine/Plugins/Animation/MLDeformer/MLDeformerFramework/Source/MLDeformerFrameworkEditor/Private/MLDeformerVizSettingsDetails.cpp
+++ b/Engine/Plugins/Animation/MLDeformer/MLDeformerFramework/Source/MLDeformerFrameworkEditor/Private/MLDeformerVizSettingsDetails.cpp
@@ -176,8 +176,6 @@ namespace UE::MLDeformer
 		LiveSettingsCategory->AddProperty(UMLDeformerVizSettings::GetWeightPropertyName(), UMLDeformerVizSettings::StaticClass());
 		LiveSettingsCategory->AddProperty(UMLDeformerVizSettings::GetAnimPlaySpeedPropertyName(), UMLDeformerVizSettings::StaticClass());
 		LiveSettingsCategory->AddProperty(UMLDeformerVizSettings::GetTestingFrameNumberPropertyName(), UMLDeformerVizSettings::StaticClass());
-		LiveSettingsCategory->AddProperty(UMLDeformerVizSettings::GetQualityLevelPropertyName(), UMLDeformerVizSettings::StaticClass())
-			.Visibility(Model->DoesSupportQualityLevels() ? EVisibility::Visible : EVisibility::Collapsed);
 
 		IDetailGroup& HeatMapGroup = LiveSettingsCategory->AddGroup("HeatMap", LOCTEXT("HeatMap", "Heat Map"), false, true);
 		HeatMapGroup.AddPropertyRow(DetailBuilder.GetProperty(UMLDeformerVizSettings::GetShowHeatMapPropertyName(), UMLDeformerVizSettings::StaticClass()));
diff --git a/Engine/Plugins/Animation/MLDeformer/MLDeformerFramework/Source/MLDeformerFrameworkEditor/Private/MLDeformerVizSettingsTabSummoner.cpp b/Engine/Plugins/Animation/MLDeformer/MLDeformerFramework/Source/MLDeformerFrameworkEditor/Private/MLDeformerVizSettingsTabSummoner.cpp
--- a/Engine/Plugins/Animation/MLDeformer/MLDeformerFramework/Source/MLDeformerFrameworkEditor/Private/MLDeformerVizSettingsTabSummoner.cpp
+++ b/Engine/Plugins/Animation/MLDeformer/MLDeformerFramework/Source/MLDeformerFrameworkEditor/Private/MLDeformerVizSettingsTabSummoner.cpp
@@ -40,7 +40,7 @@ namespace UE::MLDeformer
 	{
 		return IDocumentation::Get()->CreateToolTip(
 			LOCTEXT("VizSettingsTooltip", "The visualization settings for the ML Deformer."), 
-			NULL, 
+			nullptr, 
 			TEXT("Shared/Editors/Persona"), 
 			TEXT("MLDeformerVizSettings_Window"));
 	}
diff --git a/Engine/Plugins/Animation/MLDeformer/MLDeformerFramework/Source/MLDeformerFrameworkEditor/Public/MLDeformerEditorModel.h b/Engine/Plugins/Animation/MLDeformer/MLDeformerFramework/Source/MLDeformerFrameworkEditor/Public/MLDeformerEditorModel.h
--- a/Engine/Plugins/Animation/MLDeformer/MLDeformerFramework/Source/MLDeformerFrameworkEditor/Public/MLDeformerEditorModel.h
+++ b/Engine/Plugins/Animation/MLDeformer/MLDeformerFramework/Source/MLDeformerFrameworkEditor/Public/MLDeformerEditorModel.h
@@ -549,6 +549,12 @@ namespace UE::MLDeformer
 		 */
 		virtual void UpdateDeformerGraph();
 
+		/**
+		 * Executed when the maximum number of LOD levels changed.
+		 * Other models could use this to reinitialize the morph targets for all LOD levels for example.
+		 */
+		virtual void OnMaxNumLODsChanged() {}
+
 		/**
 		 * Sample the vertex deltas between the training base model and target model.
 		 * This will initialize the sampler if needed and set the sampler space to post-skinning deltas, and then samples them using the sampler.
@@ -650,6 +656,12 @@ namespace UE::MLDeformer
 		 */
 		virtual void DrawPIEDebugActors();
 
+		/** 
+		 * Update LOD levels of the actors in editor world.
+		 * This can be used to for example sync the LOD levels of the compare actors with the main actor.
+		 */
+		virtual void UpdateActorLODs();
+
 		/** Apply the transforms of the debug actor to the actors in the asset editor world. This will internally call ApplyDebugActorTransforms(DebugActorComponentSpaceTransforms). */
 		void ApplyDebugActorTransforms();
 
@@ -949,7 +961,7 @@ namespace UE::MLDeformer
 		 *        recomputing normals can lead to lower quality results, in trade for faster performance.
 		 * @param MaskChannel The weight mask mode, which specifies what channel to get the weight data from. Such channel allows the user to define what areas the deformer should for example not be active in.
 		 * @param bInvertMaskChannel Specifies whether the weight mask should be inverted or not.
-		 * @param MaskBuffer An optional mask buffer that contains 'Model->GetNumBaseMeshVerts() * (MorphModel->GetNumMorphTargets() - 1)' number of floats. Deltas will be multiplied by this value. 
+		 * @param MaskBuffer An optional mask buffer that contains 'Model->GetNumBaseMeshVerts() * (MorphModel->GetNumMorphTargets()- 1)' number of floats. Deltas will be multiplied by this value. 
 		 *        When the mask buffer is an empty array, it will be ignored.
 		 */
 		UE_DEPRECATED(5.3, "Please call FMLDeformerMorphModelEditorModel::CreateMorphTargets instead.")
@@ -1018,7 +1030,24 @@ namespace UE::MLDeformer
 		int32 CalcNumValidCompareActorsPriorTo(int32 CompareActorIndex) const;
 		virtual bool IsAnimIndexValid(int32 AnimIndex) const;
 
+		void UpdateLODMappings();
+
 	protected:
+		struct FLODInfo
+		{
+			/** 
+			 * Map all vertices in this LOD into an imported vertex number in LOD 0. 
+			 * This basically tells us which vertex to get the vertex delta for, for each vertex in this LOD.
+			 */
+			TArray<int32> VtxMappingToLODZero;
+		};
+
+		/** The LOD information. The size of this array is the number of LODs on the skeletal mesh, unless less LOD levels are desired using MLD. */
+		TArray<FLODInfo> LODMappings;
+
+		/** A reusable array to store current LOD skinned vertex positions. */
+		TArray<FVector3f> SkinnedPositions;
+
 		/** The runtime model associated with this editor model. */
 		TObjectPtr<UMLDeformerModel> Model = nullptr;
 
diff --git a/Engine/Plugins/Animation/MLDeformer/MLDeformerFramework/Source/MLDeformerFrameworkEditor/Public/MLDeformerModelDetails.h b/Engine/Plugins/Animation/MLDeformer/MLDeformerFramework/Source/MLDeformerFrameworkEditor/Public/MLDeformerModelDetails.h
--- a/Engine/Plugins/Animation/MLDeformer/MLDeformerFramework/Source/MLDeformerFrameworkEditor/Public/MLDeformerModelDetails.h
+++ b/Engine/Plugins/Animation/MLDeformer/MLDeformerFramework/Source/MLDeformerFrameworkEditor/Public/MLDeformerModelDetails.h
@@ -122,5 +122,8 @@ namespace UE::MLDeformer
 
 		/** The training settings category. You most likely add your model properties to this. */
 		IDetailCategoryBuilder* TrainingSettingsCategoryBuilder = nullptr;
+
+		/** The LOD generation category. This category can be hidden, based on whether the model supports LOD or not. */
+		IDetailCategoryBuilder* LODSettingsCategoryBuilder = nullptr;
 	};
 }	// namespace UE::MLDeformer
diff --git a/Engine/Plugins/Animation/MLDeformer/MLDeformerFramework/Source/MLDeformerFrameworkEditor/Public/MLDeformerMorphModelEditorModel.h b/Engine/Plugins/Animation/MLDeformer/MLDeformerFramework/Source/MLDeformerFrameworkEditor/Public/MLDeformerMorphModelEditorModel.h
--- a/Engine/Plugins/Animation/MLDeformer/MLDeformerFramework/Source/MLDeformerFrameworkEditor/Public/MLDeformerMorphModelEditorModel.h
+++ b/Engine/Plugins/Animation/MLDeformer/MLDeformerFramework/Source/MLDeformerFrameworkEditor/Public/MLDeformerMorphModelEditorModel.h
@@ -33,6 +33,7 @@ namespace UE::MLDeformer
 		virtual FString GetHeatMapDeformerGraphPath() const override;
 		virtual void OnPreTraining() override;
 		virtual void OnPostTraining(ETrainingResult TrainingResult, bool bUsePartiallyTrainedWhenAborted) override;
+		virtual void OnMaxNumLODsChanged() override;
 		// ~END FMLDeformerEditorModel overrides.
 
 		/**
@@ -138,6 +139,8 @@ namespace UE::MLDeformer
 		UMLDeformerMorphModelVizSettings* GetMorphModelVizSettings() const;
 
 	protected:
+		void TransferMorphTargets(TArray<UMorphTarget*> MorphTargetsLODZero);
+
 		/**
 		 * Initialize a set of engine morph targets and compress them to GPU friendly buffers.
 		 * These morph targets are initialized from a set of deltas. Each morph target needs to have Model->GetNumBaseVerts() number of deltas.
diff --git a/Engine/Plugins/Animation/MLDeformer/MLDeformerFramework/Source/MLDeformerFrameworkEditor/Public/MLDeformerMorphModelQualityLevelCustomization.h b/Engine/Plugins/Animation/MLDeformer/MLDeformerFramework/Source/MLDeformerFrameworkEditor/Public/MLDeformerMorphModelQualityLevelCustomization.h
deleted file mode 100644
diff --git a/Engine/Plugins/Animation/MLDeformer/MLDeformerFramework/Source/MLDeformerFrameworkEditor/Public/MLDeformerSampler.h b/Engine/Plugins/Animation/MLDeformer/MLDeformerFramework/Source/MLDeformerFrameworkEditor/Public/MLDeformerSampler.h
--- a/Engine/Plugins/Animation/MLDeformer/MLDeformerFramework/Source/MLDeformerFrameworkEditor/Public/MLDeformerSampler.h
+++ b/Engine/Plugins/Animation/MLDeformer/MLDeformerFramework/Source/MLDeformerFrameworkEditor/Public/MLDeformerSampler.h
@@ -96,6 +96,9 @@ namespace UE::MLDeformer
 		/** Get the number of floats per curve. */
 		int32 GetNumFloatsPerCurve() const								{ return NumFloatsPerCurve; }
 
+		/** Get the skinned positions for a given LOD level. */
+		void ExtractSkinnedPositions(int32 LODIndex, TArray<FVector3f>& OutPositions);
+
 	protected:
 		/** Create the actors used for sampling. This creates two actors, one for the base skeletal mesh and one for the target mesh. */
 		virtual void CreateActors();
diff --git a/Engine/Plugins/Animation/MLDeformer/NearestNeighborModel/Source/NearestNeighborModel/Private/NearestNeighborModelInstance.cpp b/Engine/Plugins/Animation/MLDeformer/NearestNeighborModel/Source/NearestNeighborModel/Private/NearestNeighborModelInstance.cpp
--- a/Engine/Plugins/Animation/MLDeformer/NearestNeighborModel/Source/NearestNeighborModel/Private/NearestNeighborModelInstance.cpp
+++ b/Engine/Plugins/Animation/MLDeformer/NearestNeighborModel/Source/NearestNeighborModel/Private/NearestNeighborModelInstance.cpp
@@ -97,7 +97,7 @@ void UNearestNeighborModelInstance::Tick(float DeltaTime, float ModelWeight)
 		PostMLDeformerComponentInit();
 	}
 
-	bool bCalledExecute = false;
+	bool bExecuteCalled = false;
 	if (ModelWeight > 0.0001f && HasValidTransforms() && SetupInputs())
 	{
 		// Execute the model instance.
@@ -105,15 +105,15 @@ void UNearestNeighborModelInstance::Tick(float DeltaTime, float ModelWeight)
 		// calculate the network outputs and possibly use them, depending on how the model works.
 		Execute(ModelWeight);
 		RunNearestNeighborModel(DeltaTime, ModelWeight);
-		bCalledExecute = true;
+		bExecuteCalled = true;
 	}
 	else
 	{
 		HandleZeroModelWeight();
 	}
 
-	// Do some things afterards, such as copying over debug actor data.
-	PostTick(bCalledExecute);
+	// Do some things afterwards, such as copying over debug actor data.
+	PostTick(bExecuteCalled);
 }
 
 int64 UNearestNeighborModelInstance::SetBoneTransforms(float* OutputBuffer, int64 OutputBufferSize, int64 StartIndex)
@@ -355,7 +355,7 @@ void UNearestNeighborModelInstance::RunNearestNeighborModel(float DeltaTime, flo
 
 	// Grab the weight data for this morph set.
 	// This could potentially fail if we are applying this deformer to the wrong skeletal mesh component.
-	const int LOD = 0;	// For now we only support LOD 0, as we can't setup an ML Deformer per LOD yet.
+	const int32 LOD = SkeletalMeshComponent->GetPredictedLODLevel();
 	FExternalMorphSetWeights* const WeightData = FindWeightData(LOD);
 	if (WeightData == nullptr)
 	{
diff --git a/Engine/Plugins/Animation/MLDeformer/NearestNeighborModel/Source/NearestNeighborModelEditor/Private/NearestNeighborEditorModel.cpp b/Engine/Plugins/Animation/MLDeformer/NearestNeighborModel/Source/NearestNeighborModelEditor/Private/NearestNeighborEditorModel.cpp
--- a/Engine/Plugins/Animation/MLDeformer/NearestNeighborModel/Source/NearestNeighborModelEditor/Private/NearestNeighborEditorModel.cpp
+++ b/Engine/Plugins/Animation/MLDeformer/NearestNeighborModel/Source/NearestNeighborModelEditor/Private/NearestNeighborEditorModel.cpp
@@ -59,7 +59,7 @@ namespace UE::NearestNeighborModel
 
 	void FNearestNeighborEditorModel::Init(const InitSettings& InitSettings)
 	{
-		FMLDeformerGeomCacheEditorModel::Init(InitSettings);
+		FMLDeformerMorphModelEditorModel::Init(InitSettings);
 		InitInputInfo(Model->GetInputInfo());
 		VertexMapSelector = MakeUnique<FVertexMapSelector>();
 		VertexMapSelector->Update(Model->GetSkeletalMesh());
@@ -188,7 +188,7 @@ namespace UE::NearestNeighborModel
 			}
 			ResetMorphTargets();
 		}
-		FMLDeformerGeomCacheEditorModel::OnPostTraining(TrainingResult, bUsePartiallyTrainedWhenAborted);
+		FMLDeformerMorphModelEditorModel::OnPostTraining(TrainingResult, bUsePartiallyTrainedWhenAborted);
 	}
 
 	FMLDeformerTrainingInputAnim* FNearestNeighborEditorModel::GetTrainingInputAnim(int32 Index) const
@@ -295,7 +295,7 @@ namespace UE::NearestNeighborModel
 
 	void FNearestNeighborEditorModel::Render(const FSceneView* View, FViewport* Viewport, FPrimitiveDrawInterface* PDI)
 	{
-		FMLDeformerGeomCacheEditorModel::Render(View, Viewport, PDI);
+		FMLDeformerMorphModelEditorModel::Render(View, Viewport, PDI);
 		const FMLDeformerSampler* Sampler = GetSamplerForActiveAnim();
 		if (!Sampler || !Sampler->IsInitialized())
 		{
@@ -524,7 +524,8 @@ namespace UE::NearestNeighborModel
 			return UpdateResult;
 		}
 		InitEngineMorphTargets(NearestNeighborModel->GetMorphTargetDeltas());
-		const TSharedPtr<const FExternalMorphSet> MorphSet = NearestNeighborModel->GetMorphTargetSet();
+		const int32 LOD = 0;
+		const TSharedPtr<const FExternalMorphSet> MorphSet = NearestNeighborModel->GetMorphTargetSet(LOD);
 		if (MorphSet.IsValid() && MorphSet->MorphBuffers.IsMorphResourcesInitialized())
 		{
 			NearestNeighborModel->UpdateMorphTargetsLastWriteTime();
@@ -709,10 +710,14 @@ namespace UE::NearestNeighborModel
 
 	void FNearestNeighborEditorModel::ResetMorphTargets()
 	{
-		const TSharedPtr<FExternalMorphSet> MorphSet = GetMorphModel()->GetMorphTargetSet();
-		if (MorphSet.IsValid())
+		const int32 NumLODs = GetMorphModel()->GetNumLODs();
+		for (int32 LOD = 0; LOD < NumLODs; ++LOD)
 		{
-			MorphSet->MorphBuffers = FMorphTargetVertexInfoBuffers();
+			const TSharedPtr<FExternalMorphSet> MorphSet = GetMorphModel()->GetMorphTargetSet(LOD);
+			if (MorphSet.IsValid())
+			{
+				MorphSet->MorphBuffers = FMorphTargetVertexInfoBuffers();
+			}
 		}
 	}
 
diff --git a/Engine/Plugins/Animation/MLDeformer/NearestNeighborModel/Source/NearestNeighborModelEditor/Private/NearestNeighborModelDetails.cpp b/Engine/Plugins/Animation/MLDeformer/NearestNeighborModel/Source/NearestNeighborModelEditor/Private/NearestNeighborModelDetails.cpp
--- a/Engine/Plugins/Animation/MLDeformer/NearestNeighborModel/Source/NearestNeighborModelEditor/Private/NearestNeighborModelDetails.cpp
+++ b/Engine/Plugins/Animation/MLDeformer/NearestNeighborModel/Source/NearestNeighborModelEditor/Private/NearestNeighborModelDetails.cpp
@@ -281,9 +281,10 @@ namespace UE::NearestNeighborModel
 			LOCTEXT("MorphTargetsStateName", "Morph Targets State"),
 			Private::FGetTextDelegate::CreateLambda([this]
 			{
+				const int32 LOD = 0;
 				if (const UNearestNeighborModel* const NearestNeighborModel = GetCastModel())
 				{
-					const TSharedPtr<const FExternalMorphSet> MorphSet = NearestNeighborModel->GetMorphTargetSet();
+					const TSharedPtr<const FExternalMorphSet> MorphSet = NearestNeighborModel->GetMorphTargetSet(LOD);
 					if (MorphSet.IsValid() && MorphSet->MorphBuffers.IsMorphResourcesInitialized())
 					{
 						if (NearestNeighborModel->IsBeforeCustomVersionWasAdded())
diff --git a/Engine/Plugins/Animation/MLDeformer/NeuralMorphModel/Source/NeuralMorphModel/Private/NeuralMorphModel.cpp b/Engine/Plugins/Animation/MLDeformer/NeuralMorphModel/Source/NeuralMorphModel/Private/NeuralMorphModel.cpp
--- a/Engine/Plugins/Animation/MLDeformer/NeuralMorphModel/Source/NeuralMorphModel/Private/NeuralMorphModel.cpp
+++ b/Engine/Plugins/Animation/MLDeformer/NeuralMorphModel/Source/NeuralMorphModel/Private/NeuralMorphModel.cpp
@@ -6,6 +6,7 @@
 #include "NeuralMorphNetwork.h"
 #include "MLDeformerAsset.h"
 #include "MLDeformerComponent.h"
+#include "MLDeformerObjectVersion.h"
 #include "UObject/UObjectGlobals.h"
 
 #include UE_INLINE_GENERATED_CPP_BY_NAME(NeuralMorphModel)
@@ -50,6 +51,8 @@ UMLDeformerInputInfo* UNeuralMorphModel::CreateInputInfo()
 
 void UNeuralMorphModel::Serialize(FArchive& Archive)
 {
+	Archive.UsingCustomVersion(UE::MLDeformer::FMLDeformerObjectVersion::GUID);
+
 	if (Archive.IsSaving())
 	{
 		UpdateMissingGroupNames();
diff --git a/Engine/Plugins/Animation/MLDeformer/NeuralMorphModel/Source/NeuralMorphModel/Private/NeuralMorphModelInstance.cpp b/Engine/Plugins/Animation/MLDeformer/NeuralMorphModel/Source/NeuralMorphModel/Private/NeuralMorphModelInstance.cpp
--- a/Engine/Plugins/Animation/MLDeformer/NeuralMorphModel/Source/NeuralMorphModel/Private/NeuralMorphModelInstance.cpp
+++ b/Engine/Plugins/Animation/MLDeformer/NeuralMorphModel/Source/NeuralMorphModel/Private/NeuralMorphModelInstance.cpp
@@ -229,7 +229,7 @@ void UNeuralMorphModelInstance::Execute(float ModelWeight)
 
 	// Grab the weight data for this morph set.
 	// This could potentially fail if we are applying this deformer to the wrong skeletal mesh component.
-	const int LOD = 0;	// For now we only support LOD 0, as we can't setup an ML Deformer per LOD yet.
+	const int32 LOD = SkeletalMeshComponent->GetPredictedLODLevel();
 	FExternalMorphSetWeights* WeightData = FindWeightData(LOD);
 	if (WeightData == nullptr)
 	{
@@ -264,32 +264,9 @@ void UNeuralMorphModelInstance::Execute(float ModelWeight)
 	WeightData->Weights[0] = ModelWeight;
 
 	// Update all generated morph target weights with the values calculated by our neural network.
-	const TArrayView<const float> ErrorValues = MorphModel->GetMorphTargetErrorValues();
-	const TArrayView<const int32> ErrorOrder = MorphModel->GetMorphTargetErrorOrder();
-	if (!ErrorValues.IsEmpty())
+	for (int32 MorphIndex = 0; MorphIndex < NumNetworkWeights; ++MorphIndex)
 	{
-		const int32 QualityLevel = GetMLDeformerComponent()->GetQualityLevel();
-		const int32 NumActiveMorphs = MorphModel->GetNumActiveMorphs(QualityLevel);
-		for (int32 Index = 0; Index < NumActiveMorphs; ++Index)
-		{
-			const int32 MorphIndex = ErrorOrder[Index];
-			const float TargetWeight = NetworkOutputs[MorphIndex] * ModelWeight;
-			WeightData->Weights[MorphIndex + 1] = FMath::Lerp<float, float>(StartMorphWeights[MorphIndex + 1], TargetWeight, MorphLerpAlpha);
-		}
-
-		// Disable all inactive morphs.
-		for (int32 Index = NumActiveMorphs; Index < NumNetworkWeights; ++Index)
-		{
-			const int32 MorphIndex = ErrorOrder[Index];
-			WeightData->Weights[MorphIndex + 1] = FMath::Lerp<float, float>(StartMorphWeights[MorphIndex + 1], 0.0f, MorphLerpAlpha);
-		}
-	}
-	else // Old models might not have the error values yet, and thus not support the quality levels yet.
-	{
-		for (int32 MorphIndex = 0; MorphIndex < NumNetworkWeights; ++MorphIndex)
-		{
-			WeightData->Weights[MorphIndex + 1] = NetworkOutputs[MorphIndex] * ModelWeight;
-		}
+		WeightData->Weights[MorphIndex + 1] = NetworkOutputs[MorphIndex] * ModelWeight;
 	}
 
 	// Clamp morph target weights to be within the bounds we seen on the training data set.
diff --git a/Engine/Source/Runtime/Engine/Classes/Components/SkinnedMeshComponent.h b/Engine/Source/Runtime/Engine/Classes/Components/SkinnedMeshComponent.h
--- a/Engine/Source/Runtime/Engine/Classes/Components/SkinnedMeshComponent.h
+++ b/Engine/Source/Runtime/Engine/Classes/Components/SkinnedMeshComponent.h
@@ -429,6 +429,7 @@ public:
 
 	/** Get the external morph sets for a given LOD. */
 	const FExternalMorphSets& GetExternalMorphSets(int32 LOD) const { return ExternalMorphSets[LOD]; }
+	FExternalMorphSets& GetExternalMorphSets(int32 LOD) { return ExternalMorphSets[LOD]; }
 
 	/** Get the array of external morph target sets. It is an array, one entry for each LOD. */
 	const TArray<FExternalMorphSets>& GetExternalMorphSetsArray() const { return ExternalMorphSets; }
diff --git a/Engine/Source/Runtime/Engine/Private/Animation/AnimTrace.cpp b/Engine/Source/Runtime/Engine/Private/Animation/AnimTrace.cpp
--- a/Engine/Source/Runtime/Engine/Private/Animation/AnimTrace.cpp
+++ b/Engine/Source/Runtime/Engine/Private/Animation/AnimTrace.cpp
@@ -296,7 +296,7 @@ UE_TRACE_EVENT_BEGIN(Animation, Inertialization)
 UE_TRACE_EVENT_END()
 
 FAutoConsoleVariable CVarRecordExternalMorphTargets(
-	TEXT("RecordExternalMorphTargets"),
+	TEXT("a.RecordExternalMorphTargets"),
 	false,
 	TEXT("Record the external morph target weights inside animation insights. On default this is disabled, because it can slow down recording."),
 	ECVF_Default
diff --git a/Engine/Source/Runtime/Engine/Private/Components/SkinnedMeshComponent.cpp b/Engine/Source/Runtime/Engine/Private/Components/SkinnedMeshComponent.cpp
--- a/Engine/Source/Runtime/Engine/Private/Components/SkinnedMeshComponent.cpp
+++ b/Engine/Source/Runtime/Engine/Private/Components/SkinnedMeshComponent.cpp
@@ -735,7 +735,7 @@ FExternalMorphWeightData& USkinnedMeshComponent::GetExternalMorphWeights(int32 L
 
 bool USkinnedMeshComponent::IsValidExternalMorphSetLODIndex(int32 LOD) const
 {
-	return ExternalMorphSets.IsValidIndex(LOD);
+	return ExternalMorphWeightData.IsValidIndex(LOD) && ExternalMorphSets.IsValidIndex(LOD);
 }
 
 void USkinnedMeshComponent::AddExternalMorphSet(int32 LOD, int32 ID, TSharedPtr<FExternalMorphSet> MorphSet)
