Skeletal Mesh: Change the import storage to use FMeshDescription instead of FSkeletalMeshImportData. This includes deprecating the USkeletalMesh::LoadLODImportedData and friends, in lieu of an API more similar to the one on UStaticMesh (GetMeshDescription/CommitMeshDescription/etc).

diff --git a/Engine/Plugins/Animation/IKRig/Source/IKRigEditor/Private/RetargetEditor/IKRetargeterPoseGenerator.cpp b/Engine/Plugins/Animation/IKRig/Source/IKRigEditor/Private/RetargetEditor/IKRetargeterPoseGenerator.cpp
--- a/Engine/Plugins/Animation/IKRig/Source/IKRigEditor/Private/RetargetEditor/IKRetargeterPoseGenerator.cpp
+++ b/Engine/Plugins/Animation/IKRig/Source/IKRigEditor/Private/RetargetEditor/IKRetargeterPoseGenerator.cpp
@@ -454,16 +454,15 @@ void FRetargetAutoPoseGenerator::GetRefPoseVerticesAndWeights(
 	const FRetargetSkeleton& Skeleton = Processor->GetSkeleton(SourceOrTarget);
 	USkeletalMesh* Mesh = Skeleton.SkeletalMesh;
 	constexpr int32 LODIndex = 0;
-	FMeshDescription MeshDescription;
-	Mesh->GetMeshDescription(LODIndex, MeshDescription);
-	if (MeshDescription.IsEmpty())
+	const FMeshDescription *MeshDescription = Mesh->GetMeshDescription(LODIndex);
+	if (!MeshDescription || MeshDescription->IsEmpty())
 	{
 		return;
 	}
 	
-	const FSkeletalMeshConstAttributes MeshAttribs(MeshDescription);
+	const FSkeletalMeshConstAttributes MeshAttribs(*MeshDescription);
 	const FSkinWeightsVertexAttributesConstRef VertexSkinWeights = MeshAttribs.GetVertexSkinWeights();
-	const int32 NumVertices = MeshDescription.Vertices().Num();
+	const int32 NumVertices = MeshDescription->Vertices().Num();
 	OutWeights.SetNum(NumVertices);
 	OutPositions.SetNum(NumVertices);
 	const FReferenceSkeleton& RefSkeleton = Mesh->GetRefSkeleton();
@@ -483,7 +482,7 @@ void FRetargetAutoPoseGenerator::GetRefPoseVerticesAndWeights(
 		}
 
 		// get position of vertex
-		const FVector Position = static_cast<FVector>(MeshDescription.GetVertexPosition(VertexID));
+		const FVector Position = static_cast<FVector>(MeshDescription->GetVertexPosition(VertexID));
 		OutPositions[VertexIndex] = Position;
 	});
 }
diff --git a/Engine/Plugins/Animation/MLDeformer/MLDeformerFramework/Source/MLDeformerFramework/MLDeformerFramework.Build.cs b/Engine/Plugins/Animation/MLDeformer/MLDeformerFramework/Source/MLDeformerFramework/MLDeformerFramework.Build.cs
--- a/Engine/Plugins/Animation/MLDeformer/MLDeformerFramework/Source/MLDeformerFramework/MLDeformerFramework.Build.cs
+++ b/Engine/Plugins/Animation/MLDeformer/MLDeformerFramework/Source/MLDeformerFramework/MLDeformerFramework.Build.cs
@@ -42,7 +42,9 @@ public class MLDeformerFramework : ModuleRules
 				"OptimusCore",
 				"RenderCore",
 				"ComputeFramework",
-				"Projects"
+				"Projects",
+				"MeshDescription",
+				"SkeletalMeshDescription"
 			}
 			);
 		
diff --git a/Engine/Plugins/Animation/MLDeformer/MLDeformerFramework/Source/MLDeformerFramework/Private/MLDeformerGeomCacheHelpers.cpp b/Engine/Plugins/Animation/MLDeformer/MLDeformerFramework/Source/MLDeformerFramework/Private/MLDeformerGeomCacheHelpers.cpp
--- a/Engine/Plugins/Animation/MLDeformer/MLDeformerFramework/Source/MLDeformerFramework/Private/MLDeformerGeomCacheHelpers.cpp
+++ b/Engine/Plugins/Animation/MLDeformer/MLDeformerFramework/Source/MLDeformerFramework/Private/MLDeformerGeomCacheHelpers.cpp
@@ -6,6 +6,7 @@
 #include "GeometryCache.h"
 #include "GeometryCacheMeshData.h"
 #include "GeometryCacheTrack.h"
+#include "SkeletalMeshAttributes.h"
 #include "Animation/AnimSequence.h"
 #include "Engine/SkeletalMesh.h"
 #include "Rendering/SkeletalMeshModel.h"
@@ -45,10 +46,10 @@ namespace UE::MLDeformer
 				int32 NumSkelMeshes = 0;
 				if (InSkeletalMesh)
 				{
-					FSkeletalMeshModel* ImportedModel = InSkeletalMesh->GetImportedModel();
-					if (ImportedModel)
+					if (const FMeshDescription* MeshDescription = InSkeletalMesh->GetMeshDescription(0))
 					{
-						NumSkelMeshes = ImportedModel->LODModels[0].ImportedMeshInfos.Num();		
+						const FSkeletalMeshConstAttributes Attributes(*MeshDescription);
+						NumSkelMeshes = Attributes.GetNumSourceGeometryParts();
 					}
 				}
 
@@ -137,12 +138,24 @@ namespace UE::MLDeformer
 			return;
 		}
 
+		if (!SkelMesh->GetImportedModel()->LODModels.IsValidIndex(0))
+		{
+			return;
+		}
+
+		const FSkeletalMeshLODModel& LODModel = SkelMesh->GetImportedModel()->LODModels[0];
+
 		// If we haven't got any imported mesh infos then the asset needs to be reimported first.
 		// We show an error for this in the editor UI already.
-		FSkeletalMeshModel* ImportedModel = SkelMesh->GetImportedModel();
-		check(ImportedModel);
-		const TArray<FSkelMeshImportedMeshInfo>& SkelMeshInfos = ImportedModel->LODModels[0].ImportedMeshInfos;
-		if (SkelMeshInfos.IsEmpty())
+		const FMeshDescription* MeshDescription = SkelMesh->GetMeshDescription(0);
+		if (!MeshDescription || MeshDescription->IsEmpty())
+		{
+			return;
+		}
+
+		FSkeletalMeshConstAttributes MeshAttributes(*MeshDescription);
+
+		if (!MeshAttributes.HasSourceGeometryParts())
 		{
 			return;
 		}
@@ -156,19 +169,23 @@ namespace UE::MLDeformer
 
 		// For all meshes in the skeletal mesh.
 		const float SampleTime = 0.0f;
-		FString SkelMeshName;
 
-		const bool bIsSoloMesh = (GeomCache->Tracks.Num() == 1 && SkelMeshInfos.Num() == 1);	// Do we just have one mesh and one track?
+		const int32 NumSourceGeoParts = MeshAttributes.GetNumSourceGeometryParts();
+		const FSkeletalMeshAttributes::FSourceGeometryPartNameConstRef GeoPartNames = MeshAttributes.GetSourceGeometryPartNames();
+		const FSkeletalMeshAttributes::FSourceGeometryPartVertexOffsetAndCountConstRef GeoPartOffsetAndCounts = MeshAttributes.GetSourceGeometryPartVertexOffsetAndCounts();
+		
+
+		const bool bIsSoloMesh = (GeomCache->Tracks.Num() == 1 && NumSourceGeoParts == 1);	// Do we just have one mesh and one track?
 		for (int32 TrackIndex = 0; TrackIndex < GeomCache->Tracks.Num(); ++TrackIndex)
 		{
 			// Check if this is a candidate based on the mesh and track name.
 			UGeometryCacheTrack* Track = GeomCache->Tracks[TrackIndex];
 
 			bool bFoundMatch = false;
-			for (int32 SkelMeshIndex = 0; SkelMeshIndex < SkelMeshInfos.Num(); ++SkelMeshIndex)
+			for (int32 GeoPartIndex = 0; GeoPartIndex < NumSourceGeoParts; ++GeoPartIndex)
 			{
-				const FSkelMeshImportedMeshInfo& MeshInfo = SkelMeshInfos[SkelMeshIndex];
-				SkelMeshName = MeshInfo.Name.ToString();
+				const FSourceGeometryPartID GeoPartID(GeoPartIndex);
+				FString SkelMeshName = GeoPartNames[GeoPartID].ToString();
 
 				if (Track &&
 					(IsPotentialMatch(Track->GetName(), SkelMeshName) || bIsSoloMesh))
@@ -195,7 +212,9 @@ namespace UE::MLDeformer
 					NumGeomMeshVerts += 1;	// +1 Because we use indices, so a cube's max index is 7, while there are 8 vertices.
 
 					// Make sure the vertex counts match.
-					const int32 NumSkelMeshVerts = MeshInfo.NumVertices;
+					TArrayView<const int32> GeoPartInfo = GeoPartOffsetAndCounts.Get(GeoPartIndex);
+					const int32 StartImportedVertex = GeoPartInfo[0]; 
+					const int32 NumSkelMeshVerts = GeoPartInfo[1];
 					if (NumSkelMeshVerts != NumGeomMeshVerts)
 					{
 						continue;
@@ -207,14 +226,14 @@ namespace UE::MLDeformer
 					// Create a new mesh mapping entry.
 					OutMeshMappings.AddDefaulted();
 					UE::MLDeformer::FMLDeformerGeomCacheMeshMapping& Mapping = OutMeshMappings.Last();
-					Mapping.MeshIndex = SkelMeshIndex;
+					Mapping.MeshIndex = GeoPartIndex;
 					Mapping.TrackIndex = TrackIndex;
 					Mapping.SkelMeshToTrackVertexMap.AddUninitialized(NumSkelMeshVerts);
 					Mapping.ImportedVertexToRenderVertexMap.AddUninitialized(NumSkelMeshVerts);
 
 					// For all vertices (both skel mesh and geom cache mesh have the same number of verts here).
 					const int32 BatchSize = 100;
-					const int32 NumBatches = (MeshInfo.NumVertices / BatchSize) + 1;
+					const int32 NumBatches = (NumSkelMeshVerts / BatchSize) + 1;
 					ParallelFor(NumBatches, [&](int32 BatchIndex)
 					{
 						const int32 StartVertex = BatchIndex * BatchSize;
@@ -223,7 +242,7 @@ namespace UE::MLDeformer
 							return;
 						}
 
-						const int32 NumVertsInBatch = (StartVertex + BatchSize) < MeshInfo.NumVertices ? BatchSize : FMath::Max(NumSkelMeshVerts - StartVertex, 0);
+						const int32 NumVertsInBatch = (StartVertex + BatchSize) < NumSkelMeshVerts ? BatchSize : FMath::Max(NumSkelMeshVerts - StartVertex, 0);
 						for (int32 VertexIndex = StartVertex; VertexIndex < StartVertex + NumVertsInBatch; ++VertexIndex)
 						{
 							// Find the first vertex with the same dcc vertex in the geom cache mesh.
@@ -233,7 +252,7 @@ namespace UE::MLDeformer
 							Mapping.SkelMeshToTrackVertexMap[VertexIndex] = GeomCacheVertexIndex;
 
 							// Map the source asset vertex number to a render vertex. This is the first duplicate of that vertex.
-							const int32 RenderVertexIndex = ImportedModel->LODModels[0].MeshToImportVertexMap.Find(MeshInfo.StartImportedVertex + VertexIndex);
+							const int32 RenderVertexIndex = LODModel.MeshToImportVertexMap.Find(StartImportedVertex + VertexIndex);
 							Mapping.ImportedVertexToRenderVertexMap[VertexIndex] = RenderVertexIndex;
 						}
 					});
@@ -278,16 +297,16 @@ namespace UE::MLDeformer
 			return;
 		}
 
-		const FSkeletalMeshModel* ImportedModel = InSkelMesh->GetImportedModel();
-		if (!ensure(ImportedModel != nullptr))
+		const FMeshDescription* MeshDescription = InSkelMesh->GetMeshDescription(InLODIndex);
+		if (!ensure(MeshDescription))
 		{
 			return;
 		}
-	
-		const FSkeletalMeshLODModel& LODModel = ImportedModel->LODModels[InLODIndex];
-		const TArray<FSkelMeshImportedMeshInfo>& SkelMeshInfos = LODModel.ImportedMeshInfos;
 
-		const uint32 NumVertices = LODModel.MaxImportVertex + 1;
+		const FSkeletalMeshConstAttributes MeshAttributes(*MeshDescription);
+		const FSkeletalMeshAttributes::FSourceGeometryPartVertexOffsetAndCountConstRef GeoPartOffsetAndCounts = MeshAttributes.GetSourceGeometryPartVertexOffsetAndCounts();
+		
+		const uint32 NumVertices = MeshDescription->Vertices().Num();
 		OutPositions.Reset(NumVertices);
 		OutPositions.AddZeroed(NumVertices);
 
@@ -295,7 +314,10 @@ namespace UE::MLDeformer
 		for (int32 MeshMappingIndex = 0; MeshMappingIndex < InMeshMappings.Num(); ++MeshMappingIndex)
 		{
 			const UE::MLDeformer::FMLDeformerGeomCacheMeshMapping& MeshMapping = InMeshMappings[MeshMappingIndex];
-			const FSkelMeshImportedMeshInfo& MeshInfo = SkelMeshInfos[MeshMapping.MeshIndex];
+			TArrayView<const int32> GeoPartInfo = GeoPartOffsetAndCounts.Get(MeshMapping.MeshIndex);
+			const int32 StartImportedVertex = GeoPartInfo[0];
+			const int32 NumImportedVertices = GeoPartInfo[1];
+			
 			UGeometryCacheTrack* Track = InGeometryCache->Tracks[MeshMapping.TrackIndex];
 
 			FGeometryCacheMeshData GeomCacheMeshData;
@@ -304,9 +326,9 @@ namespace UE::MLDeformer
 				continue;
 			}
 
-			for (int32 VertexIndex = 0; VertexIndex < MeshInfo.NumVertices; ++VertexIndex)
+			for (int32 VertexIndex = 0; VertexIndex < NumImportedVertices; ++VertexIndex)
 			{
-				const int32 SkinnedVertexIndex = MeshInfo.StartImportedVertex + VertexIndex;
+				const int32 SkinnedVertexIndex = StartImportedVertex + VertexIndex;
 				const int32 GeomCacheVertexIndex = MeshMapping.SkelMeshToTrackVertexMap[VertexIndex];
 				if (GeomCacheVertexIndex != INDEX_NONE && GeomCacheMeshData.Positions.IsValidIndex(GeomCacheVertexIndex))
 				{
diff --git a/Engine/Plugins/Animation/MLDeformer/MLDeformerFramework/Source/MLDeformerFrameworkEditor/MLDeformerFrameworkEditor.Build.cs b/Engine/Plugins/Animation/MLDeformer/MLDeformerFramework/Source/MLDeformerFrameworkEditor/MLDeformerFrameworkEditor.Build.cs
--- a/Engine/Plugins/Animation/MLDeformer/MLDeformerFramework/Source/MLDeformerFrameworkEditor/MLDeformerFrameworkEditor.Build.cs
+++ b/Engine/Plugins/Animation/MLDeformer/MLDeformerFramework/Source/MLDeformerFrameworkEditor/MLDeformerFrameworkEditor.Build.cs
@@ -52,7 +52,8 @@ namespace UnrealBuildTool.Rules
 					"TimeManagement",
 					"RenderCore",
 					"RHI",
-					"AssetDefinition"
+					"AssetDefinition",
+					"SkeletalMeshDescription",
 				}
 			);
 
diff --git a/Engine/Plugins/Animation/MLDeformer/MLDeformerFramework/Source/MLDeformerFrameworkEditor/Private/MLDeformerEditorModel.cpp b/Engine/Plugins/Animation/MLDeformer/MLDeformerFramework/Source/MLDeformerFrameworkEditor/Private/MLDeformerEditorModel.cpp
--- a/Engine/Plugins/Animation/MLDeformer/MLDeformerFramework/Source/MLDeformerFrameworkEditor/Private/MLDeformerEditorModel.cpp
+++ b/Engine/Plugins/Animation/MLDeformer/MLDeformerFramework/Source/MLDeformerFrameworkEditor/Private/MLDeformerEditorModel.cpp
@@ -47,6 +47,7 @@
 #include "SMLDeformerDebugSelectionWidget.h"
 #include "EditorViewportClient.h"
 #include "DrawDebugHelpers.h"
+#include "SkeletalMeshAttributes.h"
 #include "Slate/SceneViewport.h"
 #include "Rendering/SkeletalMeshLODModel.h"
 #include "Async/ParallelFor.h"
@@ -1366,21 +1367,22 @@ namespace UE::MLDeformer
 
 	FText FMLDeformerEditorModel::GetSkeletalMeshNeedsReimportErrorText() const
 	{
-		FText Result;
-
 		if (Model->GetSkeletalMesh())
 		{
-			FSkeletalMeshModel* ImportedModel = Model->GetSkeletalMesh()->GetImportedModel();
-			check(ImportedModel);
-
-			const TArray<FSkelMeshImportedMeshInfo>& SkelMeshInfos = ImportedModel->LODModels[0].ImportedMeshInfos;
-			if (SkelMeshInfos.IsEmpty())
+			constexpr int32 LODIndex = 0;
+			const FMeshDescription* MeshDescription = Model->GetSkeletalMesh()->GetMeshDescription(LODIndex);
+			if (!ensure(MeshDescription))
 			{
-				Result = LOCTEXT("SkelMeshNeedsReimport", "Skeletal Mesh asset needs to be reimported.");
+				return LOCTEXT("SkelMeshIsEmpty", "Skeletal Mesh asset is empty.");
+			}
+			const FSkeletalMeshConstAttributes MeshAttributes(*MeshDescription);
+			if (MeshAttributes.GetNumSourceGeometryParts() == 0)
+			{
+				return LOCTEXT("SkelMeshNeedsReimport", "Skeletal Mesh asset needs to be reimported.");
 			}
 		}
 
-		return Result;
+		return {};
 	}
 
 	FText FMLDeformerEditorModel::GetInputsErrorText() const
diff --git a/Engine/Plugins/Animation/MLDeformer/MLDeformerFramework/Source/MLDeformerFrameworkEditor/Private/MLDeformerGeomCacheSampler.cpp b/Engine/Plugins/Animation/MLDeformer/MLDeformerFramework/Source/MLDeformerFrameworkEditor/Private/MLDeformerGeomCacheSampler.cpp
--- a/Engine/Plugins/Animation/MLDeformer/MLDeformerFramework/Source/MLDeformerFrameworkEditor/Private/MLDeformerGeomCacheSampler.cpp
+++ b/Engine/Plugins/Animation/MLDeformer/MLDeformerFramework/Source/MLDeformerFrameworkEditor/Private/MLDeformerGeomCacheSampler.cpp
@@ -20,6 +20,7 @@
 #include "GeometryCache.h"
 #include "GeometryCacheMeshData.h"
 #include "GeometryCacheTrack.h"
+#include "SkeletalMeshAttributes.h"
 #include "Async/ParallelFor.h"
 
 namespace UE::MLDeformer
@@ -69,23 +70,26 @@ namespace UE::MLDeformer
 		// Call this first to update bone and curve values.
 		// This will also calculate the skinned positions if the delta space is set to PostSkinning.
 		FMLDeformerSampler::Sample(InAnimFrameIndex);
-
 		USkeletalMesh* SkeletalMesh = SkeletalMeshComponent.Get() ? SkeletalMeshComponent->GetSkeletalMeshAsset() : nullptr;
 		UGeometryCache* GeometryCache = GeometryCacheComponent->GetGeometryCache();
 		if (SkeletalMeshComponent && SkeletalMesh && GeometryCacheComponent && GeometryCache)
 		{
 			const float DeltaCutoffLength = Model->GetDeltaCutoffLength();
 			const FTransform& AlignmentTransform = Model->GetAlignmentTransform();
-			FSkeletalMeshModel* ImportedModel = SkeletalMesh->GetImportedModel();
 
 			// For all mesh mappings we found.
-			const int32 LODIndex = 0;
-			const FSkeletalMeshLODModel& LODModel = ImportedModel->LODModels[LODIndex];
-			const TArray<FSkelMeshImportedMeshInfo>& SkelMeshInfos = LODModel.ImportedMeshInfos;
+			constexpr int32 LODIndex = 0;
+			const FMeshDescription* MeshDescription = SkeletalMesh->GetMeshDescription(LODIndex);
+			const FSkeletalMeshConstAttributes MeshAttributes(*MeshDescription);
+			const FSkeletalMeshAttributes::FSourceGeometryPartVertexOffsetAndCountConstRef GeoPartOffsetAndCounts = MeshAttributes.GetSourceGeometryPartVertexOffsetAndCounts();
+			
 			for (int32 MeshMappingIndex = 0; MeshMappingIndex < MeshMappings.Num(); ++MeshMappingIndex)
 			{
 				const UE::MLDeformer::FMLDeformerGeomCacheMeshMapping& MeshMapping = MeshMappings[MeshMappingIndex];
-				const FSkelMeshImportedMeshInfo& MeshInfo = SkelMeshInfos[MeshMapping.MeshIndex];
+				TArrayView<const int32> GeoPartInfo = GeoPartOffsetAndCounts.Get(MeshMapping.MeshIndex);
+				const int32 StartImportedVertex = GeoPartInfo[0];
+				const int32 NumVertices = GeoPartInfo[1];
+				
 				UGeometryCacheTrack* Track = GeometryCache->Tracks[MeshMapping.TrackIndex];
 
 				// Sample the mesh data of the geom cache.
@@ -100,19 +104,19 @@ namespace UE::MLDeformer
 				const FSkinWeightVertexBuffer& SkinWeightBuffer = *SkeletalMeshComponent->GetSkinWeightBuffer(LODIndex);
 
 				const int32 BatchSize = 500;
-				const int32 NumBatches = (MeshInfo.NumVertices / BatchSize) + 1;
+				const int32 NumBatches = (NumVertices / BatchSize) + 1;
 				ParallelFor(NumBatches, [&](int32 BatchIndex)
 				{
 					const int32 StartVertex = BatchIndex * BatchSize;
-					if (StartVertex >= MeshInfo.NumVertices || VertexDeltas.IsEmpty())
+					if (StartVertex >= NumVertices || VertexDeltas.IsEmpty())
 					{
 						return;
 					}
 
-					const int32 NumVertsInBatch = (StartVertex + BatchSize) < MeshInfo.NumVertices ? BatchSize : FMath::Max(MeshInfo.NumVertices - StartVertex, 0);
+					const int32 NumVertsInBatch = (StartVertex + BatchSize) < NumVertices ? BatchSize : FMath::Max(NumVertices - StartVertex, 0);
 					for (int32 VertexIndex = StartVertex; VertexIndex < StartVertex + NumVertsInBatch; ++VertexIndex)
 					{
-						const int32 SkinnedVertexIndex = MeshInfo.StartImportedVertex + VertexIndex;
+						const int32 SkinnedVertexIndex = StartImportedVertex + VertexIndex;
 						const int32 GeomCacheVertexIndex = MeshMapping.SkelMeshToTrackVertexMap[VertexIndex];
 						if (GeomCacheVertexIndex != INDEX_NONE && GeomCacheMeshData.Positions.IsValidIndex(GeomCacheVertexIndex))
 						{
diff --git a/Engine/Plugins/Animation/MLDeformer/NearestNeighborModel/Source/NearestNeighborModel/NearestNeighborModel.Build.cs b/Engine/Plugins/Animation/MLDeformer/NearestNeighborModel/Source/NearestNeighborModel/NearestNeighborModel.Build.cs
--- a/Engine/Plugins/Animation/MLDeformer/NearestNeighborModel/Source/NearestNeighborModel/NearestNeighborModel.Build.cs
+++ b/Engine/Plugins/Animation/MLDeformer/NearestNeighborModel/Source/NearestNeighborModel/NearestNeighborModel.Build.cs
@@ -38,14 +38,15 @@ public class NearestNeighborModel : ModuleRules
 				"CoreUObject",
 				"Engine",
 				"GeometryCache",
-				"OptimusCore",
-				"Projects",
-				"RenderCore",
-				"RHI",
 				"MLDeformerFramework",
 				"MeshDescription",
 				"NNE",
 				"NNERuntimeBasicCpu",
+				"OptimusCore",
+				"Projects",
+				"RHI",
+				"RenderCore",
+				"SkeletalMeshDescription",
 			}
 			);
 		
diff --git a/Engine/Plugins/Animation/MLDeformer/NearestNeighborModel/Source/NearestNeighborModel/Private/NearestNeighborModel.cpp b/Engine/Plugins/Animation/MLDeformer/NearestNeighborModel/Source/NearestNeighborModel/Private/NearestNeighborModel.cpp
--- a/Engine/Plugins/Animation/MLDeformer/NearestNeighborModel/Source/NearestNeighborModel/Private/NearestNeighborModel.cpp
+++ b/Engine/Plugins/Animation/MLDeformer/NearestNeighborModel/Source/NearestNeighborModel/Private/NearestNeighborModel.cpp
@@ -22,6 +22,7 @@
 #include "Rendering/MorphTargetVertexInfoBuffers.h"
 #include "Rendering/SkeletalMeshModel.h"
 #include "ShaderCore.h"
+#include "SkeletalMeshAttributes.h"
 #include "UObject/AssetRegistryTagsContext.h"
 
 #include UE_INLINE_GENERATED_CPP_BY_NAME(NearestNeighborModel)
@@ -262,40 +263,49 @@ namespace UE::NearestNeighborModel
 #if WITH_EDITORONLY_DATA
 		TArray<FInt32Range> GetMeshVertRanges(const USkeletalMesh& SkelMesh)
 		{
-			const FSkeletalMeshModel* ImportedModel = SkelMesh.GetImportedModel();
-			TArray<FInt32Range> Empty;
-			if (!ImportedModel)
+			constexpr int32 LODIndex = 0;
+
+			if (!SkelMesh.HasMeshDescription(LODIndex))
 			{
-				return Empty;
+				return {};
 			}
-			constexpr int32 LODIndex = 0;
-			if (!ImportedModel->LODModels.IsValidIndex(LODIndex))
+
+			const FMeshDescription* MeshDescription = SkelMesh.GetMeshDescription(LODIndex);
+			const FSkeletalMeshConstAttributes MeshAttributes(*MeshDescription);
+
+			if (!MeshAttributes.HasSourceGeometryParts())
 			{
-				return Empty;
+				return {};
 			}
-			const FSkeletalMeshLODModel& LODModel = ImportedModel->LODModels[LODIndex];
-			const TArray<FSkelMeshImportedMeshInfo>& SkelMeshInfos = LODModel.ImportedMeshInfos;
+
 			TArray<FInt32Range> MeshVertRanges;
-			MeshVertRanges.SetNum(SkelMeshInfos.Num());
-			for (int32 Index = 0; Index < SkelMeshInfos.Num(); ++Index)
+			MeshVertRanges.Reserve(MeshAttributes.GetNumSourceGeometryParts());
+			const FSkeletalMeshAttributesShared::FSourceGeometryPartVertexOffsetAndCountConstRef PartOffsetAndCountRef = MeshAttributes.GetSourceGeometryPartVertexOffsetAndCounts();
+			for (const FSourceGeometryPartID GeometryPartID: MeshAttributes.SourceGeometryParts().GetElementIDs())
 			{
-				const FSkelMeshImportedMeshInfo& Info = SkelMeshInfos[Index];
-				MeshVertRanges[Index] = FInt32Range(Info.StartImportedVertex, Info.StartImportedVertex + Info.NumVertices);
+				TArrayView<const int32> OffsetAndCount = PartOffsetAndCountRef.Get(GeometryPartID);
+				MeshVertRanges.Emplace(OffsetAndCount[0], OffsetAndCount[0] + OffsetAndCount[1]);
 			}
 			return MeshVertRanges;
 		}
 
 		TArray<FName> GetVertexFloatAttributeNames(const USkeletalMesh& SkelMesh)
 		{
-			FMeshDescription MeshDescription;
 			constexpr int32 LODIndex = 0;
-			SkelMesh.GetMeshDescription(LODIndex, MeshDescription);
-			const TAttributesSet<FVertexID> AttributesSet = MeshDescription.VertexAttributes();
+
+			if (!SkelMesh.HasMeshDescription(LODIndex))
+			{
+				return {};
+			}
+
+			const FMeshDescription* MeshDescription = SkelMesh.GetMeshDescription(LODIndex);
+			
+			const TAttributesSet<FVertexID> AttributesSet = MeshDescription->VertexAttributes();
 			TArray<FName> AttributeNames;
 			AttributesSet.ForEach([&AttributeNames, &AttributesSet](const FName AttributeName, const auto AttributesRef)
 			{
 				const bool bIsAutoGenerated = (AttributesRef.GetFlags() & EMeshAttributeFlags::AutoGenerated) != EMeshAttributeFlags::None;
-				if (!bIsAutoGenerated && AttributesSet.template HasAttributeOfType<float>(AttributeName))
+				if (!bIsAutoGenerated && AttributesSet.HasAttributeOfType<float>(AttributeName))
 				{
 					AttributeNames.Add(AttributeName);
 				}
@@ -907,14 +917,18 @@ UE::NearestNeighborModel::EOpFlag UNearestNeighborModelSection::UpdateVertexWeig
 	constexpr int32 LODIndex = 0;
 	if (!SkeletalMesh->GetImportedModel()->LODModels.IsValidIndex(LODIndex))
 	{
-		UE_LOG(LogNearestNeighborModel, Error, TEXT("SkeletalMesh has no LODModel."));
+		UE_LOG(LogNearestNeighborModel, Error, TEXT("SkeletalMesh has no LODModel at LOD 0."));
+		return EOpFlag::Error;
+	}
+	if (!SkeletalMesh->HasMeshDescription(LODIndex))
+	{
+		UE_LOG(LogNearestNeighborModel, Error, TEXT("SkeletalMesh has no Mesh Description at LOD 0."));
 		return EOpFlag::Error;
 	}
 	const FSkeletalMeshLODModel& LODModel = SkeletalMesh->GetImportedModel()->LODModels[LODIndex];
 
-	FMeshDescription MeshDescription;
-	SkeletalMesh->GetMeshDescription(LODIndex, MeshDescription);
-	const TAttributesSet<FVertexID> AttributesSet = MeshDescription.VertexAttributes();
+	const FMeshDescription* MeshDescription = SkeletalMesh->GetMeshDescription(LODIndex);
+	const TAttributesSet<FVertexID> AttributesSet = MeshDescription->VertexAttributes();
 	if (!AttributesSet.HasAttributeOfType<float>(AttributeName))
 	{
 		UE_LOG(LogNearestNeighborModel, Error, TEXT("There is no float attribute %s in SkeletalMesh."), *AttributeName.ToString());
@@ -927,7 +941,7 @@ UE::NearestNeighborModel::EOpFlag UNearestNeighborModelSection::UpdateVertexWeig
 	VertexMap.Reserve(NumImportedVertices);
 	VertexWeights.Reset();
 	VertexWeights.Reserve(NumImportedVertices);
-	for (const FVertexID VertexID : MeshDescription.Vertices().GetElementIDs())
+	for (const FVertexID VertexID : MeshDescription->Vertices().GetElementIDs())
 	{
 		const float Value = AttributesRef[VertexID];
 		if (Value > 0)
diff --git a/Engine/Plugins/Animation/MLDeformer/NearestNeighborModel/Source/NearestNeighborModelEditor/NearestNeighborModelEditor.Build.cs b/Engine/Plugins/Animation/MLDeformer/NearestNeighborModel/Source/NearestNeighborModelEditor/NearestNeighborModelEditor.Build.cs
--- a/Engine/Plugins/Animation/MLDeformer/NearestNeighborModel/Source/NearestNeighborModelEditor/NearestNeighborModelEditor.Build.cs
+++ b/Engine/Plugins/Animation/MLDeformer/NearestNeighborModel/Source/NearestNeighborModelEditor/NearestNeighborModelEditor.Build.cs
@@ -37,7 +37,8 @@ namespace UnrealBuildTool.Rules
 					"ToolWidgets",
 					"ComputeFramework",
 					"RenderCore",
-					"RHI"
+					"RHI",
+					"SkeletalMeshDescription",
 				}
 			);
 
diff --git a/Engine/Plugins/Animation/MLDeformer/NearestNeighborModel/Source/NearestNeighborModelEditor/Private/NearestNeighborEditorModel.cpp b/Engine/Plugins/Animation/MLDeformer/NearestNeighborModel/Source/NearestNeighborModelEditor/Private/NearestNeighborEditorModel.cpp
--- a/Engine/Plugins/Animation/MLDeformer/NearestNeighborModel/Source/NearestNeighborModelEditor/Private/NearestNeighborEditorModel.cpp
+++ b/Engine/Plugins/Animation/MLDeformer/NearestNeighborModel/Source/NearestNeighborModelEditor/Private/NearestNeighborEditorModel.cpp
@@ -22,6 +22,7 @@
 #include "NearestNeighborTrainingModel.h"
 #include "Rendering/SkeletalMeshModel.h"
 #include "SceneManagement.h"
+#include "SkeletalMeshAttributes.h"
 
 #define LOCTEXT_NAMESPACE "NearestNeighborEditorModel"
 
@@ -772,36 +773,35 @@ namespace UE::NearestNeighborModel
 
 	void FVertexMapSelector::Update(const USkeletalMesh* SkelMesh)
 	{
-		if (!SkelMesh)
-		{
-			Reset();
-			return;
-		}
-		const FSkeletalMeshModel* ImportedModel = SkelMesh->GetImportedModel();
-		if (!ImportedModel)
-		{
-			Reset();
-			return;
-		}
 		constexpr int32 LODIndex = 0;
-		if (!ImportedModel->LODModels.IsValidIndex(LODIndex))
+		if (!SkelMesh || !SkelMesh->HasMeshDescription(LODIndex))
 		{
 			Reset();
 			return;
 		}
-		const FSkeletalMeshLODModel& LODModel = ImportedModel->LODModels[LODIndex];
-		const TArray<FSkelMeshImportedMeshInfo>& SkelMeshInfos = LODModel.ImportedMeshInfos;
-		Options.Reserve(SkelMeshInfos.Num() + 1);
-		VertexMapStrings.Reserve(SkelMeshInfos.Num() + 1);
-		for (const FSkelMeshImportedMeshInfo& Info: SkelMeshInfos)
+
+		const FMeshDescription* MeshDescription = SkelMesh->GetMeshDescription(LODIndex);
+		const FSkeletalMeshConstAttributes MeshAttributes(*MeshDescription);
+		
+		const FSkeletalMeshAttributesShared::FSourceGeometryPartNameConstRef NameRef = MeshAttributes.GetSourceGeometryPartNames();
+		const FSkeletalMeshAttributesShared::FSourceGeometryPartVertexOffsetAndCountConstRef PartOffsetAndCountRef = MeshAttributes.GetSourceGeometryPartVertexOffsetAndCounts();
+		
+		Options.Reserve(MeshAttributes.GetNumSourceGeometryParts() + 1);
+		VertexMapStrings.Reserve(MeshAttributes.GetNumSourceGeometryParts() + 1);
+
+		for (const FSourceGeometryPartID GeometryPartID: MeshAttributes.SourceGeometryParts().GetElementIDs())
 		{
-			TSharedPtr<FString> Option = MakeShared<FString>(Info.Name.ToString());
+			FName Name = NameRef.Get(GeometryPartID);
+			TArrayView<const int32> OffsetAndCount = PartOffsetAndCountRef.Get(GeometryPartID);
+
+			TSharedPtr<FString> Option = MakeShared<FString>(Name.ToString());
+			FString VertexMapString = FString::Printf(TEXT("%d-%d"), OffsetAndCount[0], OffsetAndCount[0] + OffsetAndCount[1] - 1);
 			Options.Add(Option);
-			FString VertexMapString = FString::Printf(TEXT("%d-%d"), Info.StartImportedVertex, Info.StartImportedVertex + Info.NumVertices - 1);
 			VertexMapStrings.Add(Option, MoveTemp(VertexMapString));
 		}
+		
 		Options.Add(CustomString);
-		VertexMapStrings.Add(CustomString, FString::Printf(TEXT("0-%d"), LODModel.MaxImportVertex));
+		VertexMapStrings.Add(CustomString, FString::Printf(TEXT("0-%d"), MeshDescription->Vertices().Num() - 1));
 	}
 
 	TArray<TSharedPtr<FString>>* FVertexMapSelector::GetOptions()
diff --git a/Engine/Plugins/Animation/MLDeformer/NearestNeighborModel/Source/NearestNeighborModelEditor/Private/NearestNeighborGeomCacheSampler.cpp b/Engine/Plugins/Animation/MLDeformer/NearestNeighborModel/Source/NearestNeighborModelEditor/Private/NearestNeighborGeomCacheSampler.cpp
--- a/Engine/Plugins/Animation/MLDeformer/NearestNeighborModel/Source/NearestNeighborModelEditor/Private/NearestNeighborGeomCacheSampler.cpp
+++ b/Engine/Plugins/Animation/MLDeformer/NearestNeighborModel/Source/NearestNeighborModelEditor/Private/NearestNeighborGeomCacheSampler.cpp
@@ -12,6 +12,7 @@
 #include "GeometryCacheMeshData.h"
 #include "GeometryCacheTrack.h"
 #include "NearestNeighborModel.h"
+#include "SkeletalMeshAttributes.h"
 #include "Rendering/SkeletalMeshModel.h"
 #include "Rendering/SkeletalMeshRenderData.h"
 
@@ -121,14 +122,12 @@ namespace UE::NearestNeighborModel
 			return;
 		}
 		const FTransform& AlignmentTransform = Model->GetAlignmentTransform();
-		const FSkeletalMeshModel* ImportedModel = SkeletalMesh->GetImportedModel();
 		constexpr int32 LODIndex = 0;
-		if (!ImportedModel || !ImportedModel->LODModels.IsValidIndex(LODIndex))
+		if (!SkeletalMesh->HasMeshDescription(LODIndex))
 		{
 			VertexDeltas.Reset(0);
 			return;
 		}
-		const FSkeletalMeshLODModel& LODModel = ImportedModel->LODModels[LODIndex];
 		const FSkeletalMeshRenderData* RenderData = SkeletalMesh->GetResourceForRendering();
 		const FSkinWeightVertexBuffer* SkinWeightBuffer = SkeletalMeshComponent->GetSkinWeightBuffer(LODIndex);
 		if (!RenderData || !RenderData->LODRenderData.IsValidIndex(LODIndex) || !SkinWeightBuffer)
@@ -137,7 +136,10 @@ namespace UE::NearestNeighborModel
 			return;
 		}
 		const FSkeletalMeshLODRenderData& LODRenderData = RenderData->LODRenderData[LODIndex];
-		const TArray<FSkelMeshImportedMeshInfo>& SkelMeshInfos = LODModel.ImportedMeshInfos;
+		const FMeshDescription* MeshDescription = SkeletalMesh->GetMeshDescription(LODIndex);
+		const FSkeletalMeshConstAttributes MeshAttributes(*MeshDescription);
+		const FSkeletalMeshAttributesShared::FSourceGeometryPartVertexOffsetAndCountConstRef PartOffsetAndCountRef = MeshAttributes.GetSourceGeometryPartVertexOffsetAndCounts();
+		
 		if (GeomCacheMeshDatas.Num() != MeshMappings.Num())
 		{
 			GeomCacheMeshDatas.SetNum(MeshMappings.Num());
@@ -147,7 +149,6 @@ namespace UE::NearestNeighborModel
 		for (int32 MeshMappingIndex = 0; MeshMappingIndex < MeshMappings.Num(); ++MeshMappingIndex)
 		{
 			const UE::MLDeformer::FMLDeformerGeomCacheMeshMapping& MeshMapping = MeshMappings[MeshMappingIndex];
-			const FSkelMeshImportedMeshInfo& MeshInfo = SkelMeshInfos[MeshMapping.MeshIndex];
 			UGeometryCacheTrack* const Track = GeometryCache->Tracks[MeshMapping.TrackIndex];
 		
 			// Sample the mesh data of the geom cache.
@@ -156,11 +157,15 @@ namespace UE::NearestNeighborModel
 			{
 				continue;
 			}
-		
+
+			TArrayView<const int32> OffsetAndCount = PartOffsetAndCountRef.Get(MeshMapping.MeshIndex);
+			const int32 VertexOffset = OffsetAndCount[0];
+			const int32 VertexCount = OffsetAndCount[1];
+			
 			// Calculate the vertex deltas.
-			for (int32 VertexIndex = 0; VertexIndex < MeshInfo.NumVertices; ++VertexIndex)
+			for (int32 VertexIndex = 0; VertexIndex < VertexCount; ++VertexIndex)
 			{
-				const int32 SkinnedVertexIndex = MeshInfo.StartImportedVertex + VertexIndex;
+				const int32 SkinnedVertexIndex = VertexOffset + VertexIndex;
 				const int32 GeomCacheVertexIndex = MeshMapping.SkelMeshToTrackVertexMap[VertexIndex];
 				if (GeomCacheMeshData.Positions.IsValidIndex(GeomCacheVertexIndex))
 				{
diff --git a/Engine/Plugins/Animation/MLDeformer/NearestNeighborModel/Source/NearestNeighborModelEditor/Private/NearestNeighborModelSectionCustomization.cpp b/Engine/Plugins/Animation/MLDeformer/NearestNeighborModel/Source/NearestNeighborModelEditor/Private/NearestNeighborModelSectionCustomization.cpp
--- a/Engine/Plugins/Animation/MLDeformer/NearestNeighborModel/Source/NearestNeighborModelEditor/Private/NearestNeighborModelSectionCustomization.cpp
+++ b/Engine/Plugins/Animation/MLDeformer/NearestNeighborModel/Source/NearestNeighborModelEditor/Private/NearestNeighborModelSectionCustomization.cpp
@@ -9,6 +9,7 @@
 #include "MLDeformerEditorStyle.h"
 #include "NearestNeighborModel.h"
 #include "NearestNeighborEditorModel.h"
+#include "SkeletalMeshAttributes.h"
 #include "Rendering/SkeletalMeshLODImporterData.h" 
 #include "Rendering/SkeletalMeshModel.h"
 #include "SMLDeformerBonePickerDialog.h"
@@ -49,15 +50,28 @@ namespace UE::NearestNeighborModel
 
 		void CreateVertexAttributes(USkeletalMesh& SkeletalMesh, const FString& AttributeName, const TArray<int32>& VertexMap, const TArray<float>& VertexWeights)
 		{
-			FSkeletalMeshImportData ImportData; 
 			constexpr int32 LODIndex = 0;
-			SkeletalMesh.LoadLODImportedData(LODIndex, ImportData);
-			using SkeletalMeshImportData::FVertexAttribute;
-			constexpr int32 NumComponents = 1;
-			FVertexAttribute VertexAttribute(GetVertexAttributeValues(VertexMap, VertexWeights, SkeletalMesh.GetNumImportedVertices()), NumComponents);
-			ImportData.VertexAttributeNames.Add(AttributeName);
-			ImportData.VertexAttributes.Add(VertexAttribute);
-			SkeletalMesh.SaveLODImportedData(LODIndex, ImportData);
+
+			FMeshDescription* MeshDescription = SkeletalMesh.GetMeshDescription(LODIndex);
+			if (!MeshDescription)
+			{
+				return;
+			}
+
+			FSkeletalMeshAttributes MeshAttributes(*MeshDescription);
+
+			if (!MeshDescription->VertexAttributes().HasAttribute(*AttributeName))
+			{
+				MeshDescription->VertexAttributes().RegisterAttribute<float>(*AttributeName, 1, 0.0f);
+			}
+			TVertexAttributesRef<float> AttributeRef = MeshDescription->VertexAttributes().GetAttributesRef<float>(*AttributeName);
+
+			for (const int32 VertexIndex: VertexMap)
+			{
+				AttributeRef.Set(VertexIndex, VertexWeights[VertexIndex]);
+			}
+			
+			SkeletalMesh.CommitMeshDescription(LODIndex);
 		}
 
 
diff --git a/Engine/Plugins/Animation/RigLogic/Source/RigLogicModule/Private/SkelMeshDNAUtils.cpp b/Engine/Plugins/Animation/RigLogic/Source/RigLogicModule/Private/SkelMeshDNAUtils.cpp
--- a/Engine/Plugins/Animation/RigLogic/Source/RigLogicModule/Private/SkelMeshDNAUtils.cpp
+++ b/Engine/Plugins/Animation/RigLogic/Source/RigLogicModule/Private/SkelMeshDNAUtils.cpp
@@ -458,161 +458,6 @@ void USkelMeshDNAUtils::UpdateJointBehavior(USkeletalMeshComponent* InSkelMeshCo
 	InSkelMeshComponent->InitAnim(true);
 }
 
-// BEGIN: MHC MODIFICATION
-void USkelMeshDNAUtils::RecreateSourceData(USkeletalMesh* InSkelMesh)
-// END: MHC MODIFICATION
-{
-	IMeshUtilities& MeshUtilities = FModuleManager::Get().LoadModuleChecked<IMeshUtilities>("MeshUtilities");
-	// Source data must be updated during cooking.
-	InSkelMesh->EmptyAllImportData();
-	MeshUtilities.CreateImportDataFromLODModel(InSkelMesh);
-}
-
-// BEGIN: MHC MODIFICATION
-void USkelMeshDNAUtils::UpdateSourceData(USkeletalMesh* InSkelMesh)
-// END: MHC MODIFICATION
-{
-	FSkeletalMeshModel* ImportedModel = InSkelMesh->GetImportedModel();
-	const int32 LODCount = ImportedModel->LODModels.Num();
-	const TArray<FTransform>& RawBonePose = InSkelMesh->GetRefSkeleton().GetRawRefBonePose();
-	for (int32 LODIndex = 0; LODIndex < LODCount; LODIndex++)
-	{
-		// Update vertices.	
-		const FSkeletalMeshLODModel& LODModel = ImportedModel->LODModels[LODIndex];
-
-		FSkeletalMeshImportData ImportData;
-		InSkelMesh->LoadLODImportedData(LODIndex, ImportData);
-
-		const int32 LODMeshVtxCount = LODModel.MeshToImportVertexMap.Num();
-		TArray<FSoftSkinVertex> LODVertices;
-		LODModel.GetVertices(LODVertices);
-
-// BEGIN: MHC MODIFICATION
-		ImportData.Influences.Empty();
-		for (int32 LODMeshVtxIndex = 0; LODMeshVtxIndex < LODMeshVtxCount; LODMeshVtxIndex++)
-		{
-			// Update points but keep vertex order in the original MeshToImportVertex map to preserve DNA vertex ordering in the original FBX.
-			int32 FbxVertexIndex = LODModel.MeshToImportVertexMap[LODMeshVtxIndex];
-
-			if (FbxVertexIndex <= LODModel.MaxImportVertex)
-			{
-				ImportData.Points[FbxVertexIndex] = LODVertices[LODMeshVtxIndex].Position;
-			}
-
-			// Update influences.
-			int32 SectionIdx;
-			int32 VertexIdx;
-			LODModel.GetSectionFromVertexIndex(LODMeshVtxIndex, SectionIdx, VertexIdx);
-			const FSoftSkinVertex& Vertex = LODModel.Sections[SectionIdx].SoftVertices[VertexIdx];
-			for (int32 InfluenceIndex = 0; InfluenceIndex < MAX_TOTAL_INFLUENCES; ++InfluenceIndex)
-			{
-				// Have we reached the end of valid weights?
-				if (Vertex.InfluenceWeights[InfluenceIndex] == 0)
-				{
-					break;
-				}
-				SkeletalMeshImportData::FRawBoneInfluence& Influence = ImportData.Influences.AddZeroed_GetRef();
-				Influence.VertexIndex = FbxVertexIndex;
-				Influence.BoneIndex = LODModel.Sections[SectionIdx].BoneMap[Vertex.InfluenceBones[InfluenceIndex]];
-				Influence.Weight = static_cast<float>(Vertex.InfluenceWeights[InfluenceIndex]) / 65535.0f;
-			}
-		}
-
-		// Update tangents.
-		ImportData.NumTexCoords = FMath::Min(LODModel.NumTexCoords, static_cast<uint32>(MAX_TEXCOORDS));
-
-		const int32 VertexCount = LODVertices.Num();
-		const int32 TriangleCount = LODModel.IndexBuffer.Num() / 3;
-		ImportData.Wedges.Empty(TriangleCount);
-		ImportData.Faces.Empty(VertexCount);
-		for (int32 TriangleIndex = 0; TriangleIndex < TriangleCount; ++TriangleIndex)
-		{
-			SkeletalMeshImportData::FTriangle& Face = ImportData.Faces.AddZeroed_GetRef();
-			int32 IndexBufferIndex = TriangleIndex * 3;
-			int32 SectionIndex;
-			int32 VerticeIndex[3];
-			int32 SectionVertexIndex[3];
-			for (int32 Corner = 0; Corner < 3; ++Corner)
-			{
-				const int32 WedgeIndex = ImportData.Wedges.Num();
-				VerticeIndex[Corner] = LODModel.IndexBuffer[IndexBufferIndex + Corner];
-				LODModel.GetSectionFromVertexIndex(VerticeIndex[Corner], SectionIndex, SectionVertexIndex[Corner]);
-				const FSoftSkinVertex& Vertex = LODVertices[VerticeIndex[Corner]];
-				SkeletalMeshImportData::FVertex& Wedge = ImportData.Wedges.AddZeroed_GetRef();
-				Wedge.Color = Vertex.Color;
-				for (int32 UVIndex = 0; UVIndex < static_cast<int32>(ImportData.NumTexCoords); ++UVIndex)
-				{
-					Wedge.UVs[UVIndex] = Vertex.UVs[UVIndex];
-				}
-				Wedge.VertexIndex = LODModel.MeshToImportVertexMap[VerticeIndex[Corner]];
-				Wedge.MatIndex = LODModel.Sections[SectionIndex].MaterialIndex;
-
-				Face.WedgeIndex[Corner] = WedgeIndex;
-				Face.TangentX[Corner] = Vertex.TangentX;
-				Face.TangentY[Corner] = Vertex.TangentY;
-				Face.TangentZ[Corner] = Vertex.TangentZ;
-			}
-			Face.MatIndex = LODModel.Sections[SectionIndex].MaterialIndex;
-			//Faceted by default
-			Face.SmoothingGroups = 0x0;
-		}
-
-		//Recreate smooth group
-		ImportData.ComputeSmoothGroupFromNormals();
-// END: MHC MODIFICATION
-
-		// Update reference pose.
-		const int32 JointCount = LODModel.RequiredBones.Num();
-		if (ImportData.RefBonesBinary.Num() == JointCount)
-		{
-			for (int32 JointIndex = 0; JointIndex < JointCount; JointIndex++)
-			{
-				const int32 OriginalBoneIndex = LODModel.RequiredBones[JointIndex];
-				const FTransform& UpdatedTransform = FTransform(RawBonePose[OriginalBoneIndex]);
-				ImportData.RefBonesBinary[OriginalBoneIndex].BonePos.Transform = FTransform3f(UpdatedTransform);
-			}
-		}
-
-		// Update morph targets. 
-		const int32 MorphTargetCount = InSkelMesh->GetMorphTargets().Num();
-		ImportData.MorphTargetModifiedPoints.Empty(MorphTargetCount);
-		ImportData.MorphTargetNames.Empty(MorphTargetCount);
-		ImportData.MorphTargets.Empty(MorphTargetCount);
-		if (LODIndex == 0)
-		{
-			// Blend shapes are used only in LOD0.
-			for (int32 MorphIndex = 0; MorphIndex < MorphTargetCount; MorphIndex++)
-			{
-				UMorphTarget* MorphTarget = InSkelMesh->GetMorphTargets()[MorphIndex];
-				// Add Morph target name.
-				ImportData.MorphTargetNames.Add(MorphTarget->GetName());
-				FSkeletalMeshImportData MorphTargetImportDeltas;
-
-				FMorphTargetLODModel& MorphLODModel = MorphTarget->GetMorphLODModels()[LODIndex];
-
-				// Init deltas and vertices for the current morph target.
-				int32 NumDeltas = MorphLODModel.Vertices.Num();
-				MorphTargetImportDeltas.Points.Reserve(NumDeltas);
-				TSet<uint32> MorphTargetImportVertices;
-				MorphTargetImportVertices.Reserve(NumDeltas);
-
-				FMorphTargetDelta* Deltas = MorphLODModel.Vertices.GetData();
-				for (int32 DeltaIndex = 0; DeltaIndex < NumDeltas; DeltaIndex++)
-				{
-					const uint32 SourceIndex = LODModel.MeshToImportVertexMap[Deltas[DeltaIndex].SourceIdx];
-					MorphTargetImportDeltas.Points.Add(ImportData.Points[SourceIndex] + Deltas[DeltaIndex].PositionDelta);
-					MorphTargetImportVertices.Add(SourceIndex);
-				}
-				ImportData.MorphTargetModifiedPoints.Add(MorphTargetImportVertices);
-				ImportData.MorphTargets.Add(MorphTargetImportDeltas);
-			}
-
-		}
-
-		InSkelMesh->SaveLODImportedData(LODIndex, ImportData);
-	}
-}
-
 UDNAAsset* USkelMeshDNAUtils::GetMeshDNA(USkeletalMesh* InSkelMesh)
 {
 	if (InSkelMesh && InSkelMesh->GetAssetUserDataArray())
diff --git a/Engine/Plugins/Animation/RigLogic/Source/RigLogicModule/Public/SkelMeshDNAUtils.h b/Engine/Plugins/Animation/RigLogic/Source/RigLogicModule/Public/SkelMeshDNAUtils.h
--- a/Engine/Plugins/Animation/RigLogic/Source/RigLogicModule/Public/SkelMeshDNAUtils.h
+++ b/Engine/Plugins/Animation/RigLogic/Source/RigLogicModule/Public/SkelMeshDNAUtils.h
@@ -56,12 +56,6 @@ public:
 	/** Gets the DNA asset embedded in the mesh */
 	static UDNAAsset* GetMeshDNA(USkeletalMesh* InSkelMesh);
 
-// BEGIN: MHC MODIFICATION
-	/** Recreates source skeleton data for the purpose of character cooking and export.  */
-	static void RecreateSourceData(USkeletalMesh* InSkelMesh);
-	/** Updates ImportData structs for SkeletalMesh in non destructive way to preserve the original mesh */
-	static void UpdateSourceData(USkeletalMesh* InSkelMesh);
-// END: MHC MODIFICATION
 #endif // WITH_EDITORONLY_DATA
 	/** Converts DNA vertex coordinates to UE4 coordinate system **/
 	inline static FVector ConvertDNAVertexToUE4CoordSystem(FVector InVertexPositionInDNA)
diff --git a/Engine/Plugins/Experimental/ChaosClothAssetEditor/Source/ChaosClothAssetDataflowNodes/Private/ChaosClothAsset/TransferSkinWeightsNode.cpp b/Engine/Plugins/Experimental/ChaosClothAssetEditor/Source/ChaosClothAssetDataflowNodes/Private/ChaosClothAsset/TransferSkinWeightsNode.cpp
--- a/Engine/Plugins/Experimental/ChaosClothAssetEditor/Source/ChaosClothAssetDataflowNodes/Private/ChaosClothAsset/TransferSkinWeightsNode.cpp
+++ b/Engine/Plugins/Experimental/ChaosClothAssetEditor/Source/ChaosClothAssetDataflowNodes/Private/ChaosClothAsset/TransferSkinWeightsNode.cpp
@@ -30,11 +30,10 @@ namespace UE::Chaos::ClothAsset::Private
 {
 	/** Convert the USkeletalMesh to DynamicMesh. */
 	static void SkeletalMeshToDynamicMesh(USkeletalMesh* FromSkeletalMeshAsset, int32 LodIndex, FDynamicMesh3& ToDynamicMesh)
-	{	
-		FMeshDescription SourceMesh;
-		FromSkeletalMeshAsset->GetMeshDescription(LodIndex, SourceMesh);
+	{
+		const FMeshDescription* SourceMesh = FromSkeletalMeshAsset->GetMeshDescription(LodIndex);
 		FMeshDescriptionToDynamicMesh Converter;
-		Converter.Convert(&SourceMesh, ToDynamicMesh);
+		Converter.Convert(SourceMesh, ToDynamicMesh);
 	}
 
 	/** Convert the ClothCollection to DynamicMesh. */
diff --git a/Engine/Plugins/Experimental/ChaosClothAssetEditor/Source/ChaosClothAssetTools/Private/ChaosClothAsset/ClothAssetBuilderEditor.cpp b/Engine/Plugins/Experimental/ChaosClothAssetEditor/Source/ChaosClothAssetTools/Private/ChaosClothAsset/ClothAssetBuilderEditor.cpp
--- a/Engine/Plugins/Experimental/ChaosClothAssetEditor/Source/ChaosClothAssetTools/Private/ChaosClothAsset/ClothAssetBuilderEditor.cpp
+++ b/Engine/Plugins/Experimental/ChaosClothAssetEditor/Source/ChaosClothAssetTools/Private/ChaosClothAsset/ClothAssetBuilderEditor.cpp
@@ -27,7 +27,6 @@ void UClothAssetBuilderEditor::BuildLod(FSkeletalMeshLODModel& LODModel, const U
 	LODModel.Empty();
 
 	// Clear the mesh infos, none are stored on this asset
-	LODModel.ImportedMeshInfos.Empty();
 	LODModel.MaxImportVertex = 0;
 
 	// Set 1 texture coordinate
diff --git a/Engine/Plugins/Experimental/ChaosFlesh/Source/ChaosFleshNodes/Private/Dataflow/ChaosFleshTetrahedralNodes.cpp b/Engine/Plugins/Experimental/ChaosFlesh/Source/ChaosFleshNodes/Private/Dataflow/ChaosFleshTetrahedralNodes.cpp
--- a/Engine/Plugins/Experimental/ChaosFlesh/Source/ChaosFleshNodes/Private/Dataflow/ChaosFleshTetrahedralNodes.cpp
+++ b/Engine/Plugins/Experimental/ChaosFlesh/Source/ChaosFleshNodes/Private/Dataflow/ChaosFleshTetrahedralNodes.cpp
@@ -231,24 +231,11 @@ void FGenerateTetrahedralCollectionDataflowNodes::Evaluate(Dataflow::FContext& C
 			else if (InSkeletalMesh)
 			{
 				// Check first if we have bulk data available and non-empty.
-				const int32 LODIndex = 0;
+				constexpr int32 LODIndex = 0;
 				FMeshDescription SourceMesh;
-				if (InSkeletalMesh->IsLODImportedDataBuildAvailable(LODIndex) && !InSkeletalMesh->IsLODImportedDataEmpty(LODIndex))
+				if (InSkeletalMesh->HasMeshDescription(LODIndex))
 				{
-					FSkeletalMeshImportData SkeletalMeshImportData;
-					InSkeletalMesh->LoadLODImportedData(LODIndex, SkeletalMeshImportData);
-					SkeletalMeshImportData.GetMeshDescription(SourceMesh);
-				}
-				else
-				{
-					// Fall back on the LOD model directly if no bulk data exists. When we commit
-					// the mesh description, we override using the bulk data. This can happen for older
-					// skeletal meshes, from UE 4.24 and earlier.
-					const FSkeletalMeshModel* SkeletalMeshModel = InSkeletalMesh->GetImportedModel();
-					if (SkeletalMeshModel && SkeletalMeshModel->LODModels.IsValidIndex(LODIndex))
-					{
-						SkeletalMeshModel->LODModels[LODIndex].GetMeshDescription(SourceMesh, InSkeletalMesh);
-					}
+					InSkeletalMesh->CloneMeshDescription(LODIndex, SourceMesh);
 				}
 				FMeshDescriptionToDynamicMesh Converter;
 				Converter.Convert(&SourceMesh, DynamicMesh);
diff --git a/Engine/Plugins/Experimental/GeometryScripting/Source/GeometryScriptingCore/Private/MeshAssetFunctions.cpp b/Engine/Plugins/Experimental/GeometryScripting/Source/GeometryScriptingCore/Private/MeshAssetFunctions.cpp
--- a/Engine/Plugins/Experimental/GeometryScripting/Source/GeometryScriptingCore/Private/MeshAssetFunctions.cpp
+++ b/Engine/Plugins/Experimental/GeometryScripting/Source/GeometryScriptingCore/Private/MeshAssetFunctions.cpp
@@ -627,30 +627,22 @@ UDynamicMesh* UGeometryScriptLibrary_StaticMeshFunctions::CopyMeshFromSkeletalMe
 #if WITH_EDITOR
 	const int32 UseLODIndex = FMath::Clamp(RequestedLOD.LODIndex, 0, FromSkeletalMeshAsset->GetLODNum() - 1);;
 	
-	FMeshDescription SourceMesh;
+	const FMeshDescription* SourceMesh = nullptr;
 
 	// Check first if we have bulk data available and non-empty.
-	if (FromSkeletalMeshAsset->IsLODImportedDataBuildAvailable(UseLODIndex) && !FromSkeletalMeshAsset->IsLODImportedDataEmpty(UseLODIndex))
+	if (FromSkeletalMeshAsset->HasMeshDescription(UseLODIndex))
 	{
-		FSkeletalMeshImportData SkeletalMeshImportData;
-		FromSkeletalMeshAsset->LoadLODImportedData(UseLODIndex, SkeletalMeshImportData);
-		SkeletalMeshImportData.GetMeshDescription(SourceMesh);
+		SourceMesh = FromSkeletalMeshAsset->GetMeshDescription(UseLODIndex); 
 	}
-	else
+	if (SourceMesh == nullptr)
 	{
-		// Fall back on the LOD model directly if no bulk data exists. When we commit
-		// the mesh description, we override using the bulk data. This can happen for older
-		// skeletal meshes, from UE 4.24 and earlier.
-		const FSkeletalMeshModel* SkeletalMeshModel = FromSkeletalMeshAsset->GetImportedModel();
-		if (SkeletalMeshModel && SkeletalMeshModel->LODModels.IsValidIndex(UseLODIndex))
-		{
-			SkeletalMeshModel->LODModels[UseLODIndex].GetMeshDescription(SourceMesh, FromSkeletalMeshAsset);
-		}			
+		UE::Geometry::AppendError(Debug, EGeometryScriptErrorType::InvalidInputs, LOCTEXT("CopyMeshFromSkeletalMesh_LODNotAvailable", "CopyMeshFromSkeletalMesh: Requested LOD is not available"));
+		return ToDynamicMesh;
 	}
 
 	FDynamicMesh3 NewMesh;
 	FMeshDescriptionToDynamicMesh Converter;
-	Converter.Convert(&SourceMesh, NewMesh, AssetOptions.bRequestTangents);
+	Converter.Convert(SourceMesh, NewMesh, AssetOptions.bRequestTangents);
 	
 	ToDynamicMesh->SetMesh(MoveTemp(NewMesh));
 	
@@ -761,15 +753,22 @@ UDynamicMesh* UGeometryScriptLibrary_StaticMeshFunctions::CopyMeshToSkeletalMesh
 
 	verify(ToSkeletalMeshAsset->Modify());
 
-	FMeshDescription MeshDescription;
-	FSkeletalMeshAttributes MeshAttributes(MeshDescription);
+	FMeshDescription* MeshDescription = ToSkeletalMeshAsset->GetMeshDescription(TargetLOD.LODIndex);
+	if (MeshDescription == nullptr)
+	{
+		MeshDescription = ToSkeletalMeshAsset->CreateMeshDescription(TargetLOD.LODIndex);
+	}
+	
+	FSkeletalMeshAttributes MeshAttributes(*MeshDescription);
 	MeshAttributes.Register();
+
+	ToSkeletalMeshAsset->ModifyMeshDescription(TargetLOD.LODIndex);
 	
 	FConversionToMeshDescriptionOptions ConversionOptions;
 	FDynamicMeshToMeshDescription Converter(ConversionOptions);
 	FromDynamicMesh->ProcessMesh([&](const FDynamicMesh3& ReadMesh)
 	{
-		Converter.Convert(&ReadMesh, MeshDescription, !Options.bEnableRecomputeTangents);
+		Converter.Convert(&ReadMesh, *MeshDescription, !Options.bEnableRecomputeTangents);
 	});
 
 	// Ensure we have enough LODInfos to cover up to the requested LOD.
@@ -792,15 +791,12 @@ UDynamicMesh* UGeometryScriptLibrary_StaticMeshFunctions::CopyMeshToSkeletalMesh
 	SkeletalLODInfo->ReductionSettings.MaxNumOfVerts = MAX_int32; 
 	SkeletalLODInfo->ReductionSettings.BaseLOD = TargetLOD.LODIndex;
 
-	FSkeletalMeshImportData SkeletalMeshImportData = 
-		FSkeletalMeshImportData::CreateFromMeshDescription(MeshDescription);
-
-	// if we are replacing materials, construct a FSkeletalMaterial list with unique slot names
-	TArray<FSkeletalMaterial> NewMaterials;
+	// update materials on the Asset
 	if (Options.bReplaceMaterials)
 	{
 		bool bHaveSlotNames = (Options.NewMaterialSlotNames.Num() == Options.NewMaterials.Num());
 
+		TArray<FSkeletalMaterial> NewMaterials;
 		for (int32 k = 0; k < Options.NewMaterials.Num(); ++k)
 		{
 			FSkeletalMaterial NewMaterial;
@@ -813,36 +809,29 @@ UDynamicMesh* UGeometryScriptLibrary_StaticMeshFunctions::CopyMeshToSkeletalMesh
 			NewMaterial.UVChannelData = FMeshUVChannelInfo(1.f);		// this avoids an ensure in  UStaticMesh::GetUVChannelData
 			NewMaterials.Add(NewMaterial);
 		}
+		
+		ToSkeletalMeshAsset->SetMaterials(NewMaterials);
+	}
 
-		// copy the material set to the SkeletalMeshImportData
-		SkeletalMeshImportData.Materials.SetNum(NewMaterials.Num());
-		for (int32 k = 0; k < NewMaterials.Num(); ++k)
+	ToSkeletalMeshAsset->CommitMeshDescription(TargetLOD.LODIndex);
+
+	bool bHasVertexColors = false;
+	TVertexInstanceAttributesConstRef<FVector4f> VertexColors = MeshAttributes.GetVertexInstanceColors();
+	for (const FVertexInstanceID VertexInstanceID: MeshDescription->VertexInstances().GetElementIDs())
+	{
+		if (!VertexColors.Get(VertexInstanceID).Equals(FVector4f::One()))
 		{
-			SkeletalMeshImportData.Materials[k].Material = NewMaterials[k].MaterialInterface;
-			SkeletalMeshImportData.Materials[k].MaterialImportName = NewMaterials[k].MaterialSlotName.ToString();
+			bHasVertexColors = true;
+			break;
 		}
 	}
-
-
-	ToSkeletalMeshAsset->SaveLODImportedData(TargetLOD.LODIndex, SkeletalMeshImportData);
-
+		
 	// configure vertex color setup in the Asset
-    ToSkeletalMeshAsset->SetHasVertexColors(SkeletalMeshImportData.bHasVertexColors);
+	ToSkeletalMeshAsset->SetHasVertexColors(bHasVertexColors);
 #if WITH_EDITORONLY_DATA
-    ToSkeletalMeshAsset->SetVertexColorGuid(SkeletalMeshImportData.bHasVertexColors ? FGuid::NewGuid() : FGuid() );
+	ToSkeletalMeshAsset->SetVertexColorGuid(bHasVertexColors ? FGuid::NewGuid() : FGuid() );
 #endif
-
-	// update materials on the Asset
-	if (Options.bReplaceMaterials)
-	{
-		ToSkeletalMeshAsset->SetMaterials(NewMaterials);
-	}
-
-	// Make sure the mesh builder knows it's the latest variety, so that the render data gets
-	// properly rebuilt.
-	ToSkeletalMeshAsset->SetLODImportedDataVersions(TargetLOD.LODIndex, ESkeletalMeshGeoImportVersions::LatestVersion, ESkeletalMeshSkinningImportVersions::LatestVersion);
-	ToSkeletalMeshAsset->SetUseLegacyMeshDerivedDataKey(false);
-
+	
 	if (Options.bDeferMeshPostEditChange == false)
 	{
 		ToSkeletalMeshAsset->PostEditChange();
diff --git a/Engine/Plugins/Experimental/MeshModelingToolsetExp/Source/SkeletalMeshModifiers/Private/SkeletonModifier.cpp b/Engine/Plugins/Experimental/MeshModelingToolsetExp/Source/SkeletalMeshModifiers/Private/SkeletonModifier.cpp
--- a/Engine/Plugins/Experimental/MeshModelingToolsetExp/Source/SkeletalMeshModifiers/Private/SkeletonModifier.cpp
+++ b/Engine/Plugins/Experimental/MeshModelingToolsetExp/Source/SkeletalMeshModifiers/Private/SkeletonModifier.cpp
@@ -214,7 +214,8 @@ bool USkeletonModifier::SetSkeletalMesh(USkeletalMesh* InSkeletalMesh)
 
 	// store mesh description to edit
 	MeshDescription = MakeUnique<FMeshDescription>();
-	SkeletalMesh->GetMeshDescription(USkeletonModifierLocals::LODIndex, *MeshDescription);
+	SkeletalMesh->CloneMeshDescription(USkeletonModifierLocals::LODIndex, *MeshDescription);
+	
 	if (MeshDescription->IsEmpty())
 	{
 		UE_LOG(LogAnimation, Error, TEXT("Skeleton Modifier: mesh description is emtpy."));
@@ -374,7 +375,9 @@ bool USkeletonModifier::CommitSkeletonToSkeletalMesh()
 	SkeletalMesh->CalculateInvRefMatrices();
 	
 	// update skeletal mesh LOD (cf. USkeletalMesh::CommitMeshDescription)
-	SkeletalMesh->CommitMeshDescription(USkeletonModifierLocals::LODIndex, *MeshDescription);
+	SkeletalMesh->ModifyMeshDescription(USkeletonModifierLocals::LODIndex);
+	SkeletalMesh->CreateMeshDescription(USkeletonModifierLocals::LODIndex, MoveTemp(*MeshDescription));
+	SkeletalMesh->CommitMeshDescription(USkeletonModifierLocals::LODIndex);
 
 	// update skeleton
 	if (Skeleton->RecreateBoneTree(SkeletalMesh.Get()))
diff --git a/Engine/Plugins/Experimental/MeshModelingToolsetExp/Source/SkeletalMeshModifiers/Private/SkinWeightModifier.cpp b/Engine/Plugins/Experimental/MeshModelingToolsetExp/Source/SkeletalMeshModifiers/Private/SkinWeightModifier.cpp
--- a/Engine/Plugins/Experimental/MeshModelingToolsetExp/Source/SkeletalMeshModifiers/Private/SkinWeightModifier.cpp
+++ b/Engine/Plugins/Experimental/MeshModelingToolsetExp/Source/SkeletalMeshModifiers/Private/SkinWeightModifier.cpp
@@ -35,7 +35,7 @@ bool USkinWeightModifier::SetSkeletalMesh(USkeletalMesh* InMesh)
 	
 	// store mesh description to edit
 	MeshDescription = MakeUnique<FMeshDescription>();
-	InMesh->GetMeshDescription(LODIndex, *MeshDescription);
+	InMesh->CloneMeshDescription(LODIndex, *MeshDescription);
 	if (MeshDescription->IsEmpty())
 	{
 		UE_LOG(LogAnimation, Error, TEXT("Skin Weight Modifier: mesh description is empty."));
@@ -128,8 +128,11 @@ bool USkinWeightModifier::CommitWeightsToSkeletalMesh()
 	{
 		// flush any pending rendering commands, which might touch a component while we are rebuilding it's mesh
 		FlushRenderingCommands();
+		
 		// update skeletal mesh LOD (cf. USkeletalMesh::CommitMeshDescription)
-		Mesh->CommitMeshDescription(LODIndex, *MeshDescription);
+		Mesh->ModifyMeshDescription(LODIndex);
+		Mesh->CreateMeshDescription(LODIndex, MoveTemp(*MeshDescription));
+		Mesh->CommitMeshDescription(LODIndex);
 		// rebuilds mesh
 		Mesh->PostEditChange();
 	}
diff --git a/Engine/Plugins/Experimental/SkeletalReduction/Source/Private/SkeletalMeshReductionPlugin.cpp b/Engine/Plugins/Experimental/SkeletalReduction/Source/Private/SkeletalMeshReductionPlugin.cpp
--- a/Engine/Plugins/Experimental/SkeletalReduction/Source/Private/SkeletalMeshReductionPlugin.cpp
+++ b/Engine/Plugins/Experimental/SkeletalReduction/Source/Private/SkeletalMeshReductionPlugin.cpp
@@ -2022,7 +2022,7 @@ void FQuadricSkeletalMeshReduction::ReduceSkeletalMesh(USkeletalMesh& SkeletalMe
 	if (Settings.BaseLOD > 0)
 	{
 		
-		if (Settings.BaseLOD == LODIndex && !SkeletalMesh.IsLODImportedDataBuildAvailable(LODIndex))
+		if (Settings.BaseLOD == LODIndex && !SkeletalMesh.HasMeshDescription(LODIndex))
 		{
 			//Cannot reduce ourself if we are not imported
 			UE_ASSET_LOG(LogSkeletalMeshReduction, Warning, &SkeletalMesh, TEXT("Building LOD %d - Cannot generate LOD with itself if the LOD do not have imported Data. Using Base LOD 0 instead"), LODIndex);
@@ -2138,7 +2138,7 @@ void FQuadricSkeletalMeshReduction::ReduceSkeletalMesh(USkeletalMesh& SkeletalMe
 
 	bool bReducingSourceModel = false;
 	//Reducing source LOD, we need to use the LOD import data so it can be iterative
-	if (BaseLOD == LODIndex && SkeletalMesh.IsLODImportedDataBuildAvailable(LODIndex))
+	if (BaseLOD == LODIndex && SkeletalMesh.HasMeshDescription(LODIndex))
 	{
 		bReducingSourceModel = true;
 	}
@@ -2264,8 +2264,6 @@ void FQuadricSkeletalMeshReduction::ReduceSkeletalMesh(USkeletalMesh& SkeletalMe
 
 	// Swap out the old model.  
 	bool bPutBackRawMesh = false;
-	ESkeletalMeshGeoImportVersions GeoImportVersion = ESkeletalMeshGeoImportVersions::Before_Versionning;
-	ESkeletalMeshSkinningImportVersions SkinningImportVersion = ESkeletalMeshSkinningImportVersions::Before_Versionning;
 	{
 		FSkeletalMeshLODModel* Old = LODModels[LODIndex];
 
@@ -2277,7 +2275,6 @@ void FQuadricSkeletalMeshReduction::ReduceSkeletalMesh(USkeletalMesh& SkeletalMe
 			//We need to backup the original RawSkeletalMeshBulkData in case it was an imported LOD
 			if (!bLODModelAdded && !bIsOldRawSkelMeshEmpty)
 			{
-				SkeletalMesh.GetLODImportedDataVersions(LODIndex, GeoImportVersion, SkinningImportVersion);
 				bPutBackRawMesh = true;
 			}
 
@@ -2301,7 +2298,6 @@ void FQuadricSkeletalMeshReduction::ReduceSkeletalMesh(USkeletalMesh& SkeletalMe
 		}
 		else if(bReducingSourceModel)
 		{
-			SkeletalMesh.GetLODImportedDataVersions(BaseLOD, GeoImportVersion, SkinningImportVersion);
 			bPutBackRawMesh = true;
 		}
 	}
@@ -2396,19 +2392,15 @@ void FQuadricSkeletalMeshReduction::ReduceSkeletalMesh(USkeletalMesh& SkeletalMe
 			FLODUtilities::RestoreClothingFromBackup(&SkeletalMesh, ClothingBindings, LODIndex);
 		}
 	}
-	if (bPutBackRawMesh)
-	{
-		check((bReducingSourceModel || !bLODModelAdded));
-		//Put back the original import data, we need it to allow inline reduction and skeletal mesh split workflow
-		//It also warranty that we do not change the ddc key
-		SkeletalMesh.SetLODImportedDataVersions(LODIndex, GeoImportVersion, SkinningImportVersion);
-		if (!bLODModelAdded)
-		{
-			//The SaveLODImportdData can change the cache RawSkeletalMeshBulkDataID, so we have to put back the backup after the save
-			FSkeletalMeshLODModel& ImportedModelLOD = SkeletalMesh.GetImportedModel()->LODModels[LODIndex];
-			ImportedModelLOD.BuildStringID = BackupLodModelBuildStringID;
-			ImportedModelLOD.RawSkeletalMeshBulkDataID = BackupRawSkeletalMeshBulkDataID;
-		}
+	
+	//Put back the original import data, we need it to allow inline reduction and skeletal mesh split workflow
+	//It also warranty that we do not change the ddc key
+	if (bPutBackRawMesh && !bLODModelAdded)
+	{
+		//The SaveLODImportdData can change the cache RawSkeletalMeshBulkDataID, so we have to put back the backup after the save
+		FSkeletalMeshLODModel& ImportedModelLOD = SkeletalMesh.GetImportedModel()->LODModels[LODIndex];
+		ImportedModelLOD.BuildStringID = BackupLodModelBuildStringID;
+		ImportedModelLOD.RawSkeletalMeshBulkDataID = BackupRawSkeletalMeshBulkDataID;
 	}
 
 	SkeletalMesh.CalculateRequiredBones(SkeletalMeshResource.LODModels[LODIndex], SkeletalMesh.GetRefSkeleton(), &BonesToRemove);
diff --git a/Engine/Plugins/Interchange/Runtime/Source/Import/Private/Mesh/InterchangeSkeletalMeshFactory.cpp b/Engine/Plugins/Interchange/Runtime/Source/Import/Private/Mesh/InterchangeSkeletalMeshFactory.cpp
--- a/Engine/Plugins/Interchange/Runtime/Source/Import/Private/Mesh/InterchangeSkeletalMeshFactory.cpp
+++ b/Engine/Plugins/Interchange/Runtime/Source/Import/Private/Mesh/InterchangeSkeletalMeshFactory.cpp
@@ -995,8 +995,6 @@ UInterchangeFactoryBase::FImportAssetResult UInterchangeSkeletalMeshFactory::Beg
 	}
 	
 	SkeletalMesh->PreEditChange(nullptr);
-	//Allocate the LODImport data in the main thread
-	SkeletalMesh->ReserveLODImportData(SkeletalMeshFactoryNode->GetLodDataCount());
 
 	//Lock the skeletalmesh properties if the skeletal mesh already exist (re-import)
 	if (ExistingAsset)
@@ -1304,12 +1302,6 @@ UInterchangeFactoryBase::FImportAssetResult UInterchangeSkeletalMeshFactory::Imp
 
 
 		TRACE_CPUPROFILER_EVENT_SCOPE(UInterchangeSkeletalMeshFactory::CreateAsset_LOD)
-		ESkeletalMeshGeoImportVersions GeoImportVersion = ESkeletalMeshGeoImportVersions::LatestVersion;
-		ESkeletalMeshSkinningImportVersions SkinningImportVersion = ESkeletalMeshSkinningImportVersions::LatestVersion;
-		if (ImportAssetObjectData.bIsReImport && SkeletalMesh->GetImportedModel() && SkeletalMesh->GetImportedModel()->LODModels.IsValidIndex(CurrentLodIndex))
-		{
-			SkeletalMesh->GetLODImportedDataVersions(CurrentLodIndex, GeoImportVersion, SkinningImportVersion);
-		}
 
 		FString LodUniqueId = LodDataUniqueIds[LodIndex];
 		const UInterchangeSkeletalMeshLodDataNode* LodDataNode = Cast<UInterchangeSkeletalMeshLodDataNode>(Arguments.NodeContainer->GetNode(LodUniqueId));
@@ -1519,6 +1511,28 @@ UInterchangeFactoryBase::FImportAssetResult UInterchangeSkeletalMeshFactory::Imp
 			ensure(ImportedResource->LODModels.Add(new FSkeletalMeshLODModel()) == CurrentLodIndex);
 		}
 
+		auto AddLodInfo = [&SkeletalMesh]()
+		{
+			FSkeletalMeshLODInfo& NewLODInfo = SkeletalMesh->AddLODInfo();
+			NewLODInfo.ReductionSettings.NumOfTrianglesPercentage = 1.0f;
+			NewLODInfo.ReductionSettings.NumOfVertPercentage = 1.0f;
+			NewLODInfo.ReductionSettings.MaxDeviationPercentage = 0.0f;
+			NewLODInfo.LODHysteresis = 0.02f;
+			NewLODInfo.bImportWithBaseMesh = true;
+		};
+
+		if (ImportAssetObjectData.bIsReImport)
+		{
+			while (SkeletalMesh->GetLODNum() <= CurrentLodIndex)
+			{
+				AddLodInfo();
+			}
+		}
+		else
+		{
+			AddLodInfo();
+		}
+		
 		FSkeletalMeshLODModel& LODModel = ImportedResource->LODModels[CurrentLodIndex];
 
 		UE::Interchange::Private::ProcessImportMeshInfluences(SkeletalMeshImportData.Wedges.Num(), SkeletalMeshImportData.Influences);
@@ -1538,10 +1552,12 @@ UInterchangeFactoryBase::FImportAssetResult UInterchangeSkeletalMeshFactory::Imp
 		{
 			if (CurrentLodIndex != 0)
 			{
-				if (!SkeletalMesh->IsLODImportedDataEmpty(CurrentLodIndex))
+				if (SkeletalMesh->HasMeshDescription(CurrentLodIndex))
 				{
 					FSkeletalMeshImportData LODImportData;
+					PRAGMA_DISABLE_DEPRECATION_WARNINGS
 					SkeletalMesh->LoadLODImportedData(CurrentLodIndex, LODImportData);
+					PRAGMA_ENABLE_DEPRECATION_WARNINGS
 					for (int32 SectionIndex = 0; SectionIndex < LODImportData.Materials.Num(); ++SectionIndex)
 					{
 						ExistingOriginalPerSectionMaterialImportName.Add(FName(*LODImportData.Materials[SectionIndex].MaterialImportName));
@@ -1558,43 +1574,9 @@ UInterchangeFactoryBase::FImportAssetResult UInterchangeSkeletalMeshFactory::Imp
 			}
 		}
 		//Store the original fbx import data the SkelMeshImportDataPtr should not be modified after this
+		PRAGMA_DISABLE_DEPRECATION_WARNINGS
 		SkeletalMesh->SaveLODImportedData(CurrentLodIndex, SkeletalMeshImportData);
-
-		if (ContentInfo.bApplySkinningOnly)
-		{
-			SkeletalMesh->SetLODImportedDataVersions(CurrentLodIndex, GeoImportVersion, ESkeletalMeshSkinningImportVersions::LatestVersion);
-		}
-		else if (ContentInfo.bApplyGeometryOnly)
-		{
-			SkeletalMesh->SetLODImportedDataVersions(CurrentLodIndex, ESkeletalMeshGeoImportVersions::LatestVersion, SkinningImportVersion);
-		}
-		else
-		{
-			//We reimport both
-			SkeletalMesh->SetLODImportedDataVersions(CurrentLodIndex, ESkeletalMeshGeoImportVersions::LatestVersion, ESkeletalMeshSkinningImportVersions::LatestVersion);
-		}
-
-		auto AddLodInfo = [&SkeletalMesh]()
-		{
-			FSkeletalMeshLODInfo& NewLODInfo = SkeletalMesh->AddLODInfo();
-			NewLODInfo.ReductionSettings.NumOfTrianglesPercentage = 1.0f;
-			NewLODInfo.ReductionSettings.NumOfVertPercentage = 1.0f;
-			NewLODInfo.ReductionSettings.MaxDeviationPercentage = 0.0f;
-			NewLODInfo.LODHysteresis = 0.02f;
-			NewLODInfo.bImportWithBaseMesh = true;
-		};
-
-		if (ImportAssetObjectData.bIsReImport)
-		{
-			while (SkeletalMesh->GetLODNum() <= CurrentLodIndex)
-			{
-				AddLodInfo();
-			}
-		}
-		else
-		{
-			AddLodInfo();
-		}
+		PRAGMA_ENABLE_DEPRECATION_WARNINGS
 
 		//Update the bounding box if we are importing the LOD 0
 		if(CurrentLodIndex == 0)
diff --git a/Engine/Plugins/MeshPainting/Source/MeshPaintingToolset/Private/MeshPaintSkeletalMeshAdapter.cpp b/Engine/Plugins/MeshPainting/Source/MeshPaintingToolset/Private/MeshPaintSkeletalMeshAdapter.cpp
--- a/Engine/Plugins/MeshPainting/Source/MeshPaintingToolset/Private/MeshPaintSkeletalMeshAdapter.cpp
+++ b/Engine/Plugins/MeshPainting/Source/MeshPaintingToolset/Private/MeshPaintSkeletalMeshAdapter.cpp
@@ -10,7 +10,7 @@
 #include "Factories/FbxSkeletalMeshImportData.h"
 #include "InterchangeAssetImportData.h"
 #include "InterchangeGenericAssetsPipeline.h"
-#include "InterchangePythonPipelineBase.h"
+#include "StaticMeshAttributes.h"
 #include "Rendering/RenderCommandPipes.h"
 
 DEFINE_LOG_CATEGORY_STATIC(LogMeshPaintSkeletalMeshAdapter, Log, All);
@@ -24,23 +24,23 @@ void PropagateVertexPaintToSkeletalMesh(USkeletalMesh* SkeletalMesh, int32 LODIn
 {
 	struct FMatchFaceData
 	{
-		int32 SoftVerticeIndexes[3];
+		int32 SoftVertexIndexes[3];
 	};
 
-	if (!SkeletalMesh || SkeletalMesh->IsLODImportedDataEmpty(LODIndex) || !SkeletalMesh->IsLODImportedDataBuildAvailable(LODIndex))
+	if (!SkeletalMesh || SkeletalMesh->HasMeshDescription(LODIndex))
 	{
 		//We do not propagate vertex color for old asset
 		return;
 	}
 
-	auto GetMatchKey = [](const FVector& PositionA, const FVector& PositionB, const FVector& PositionC)->FSHAHash
+	auto GetMatchKey = [](const FVector3f& PositionA, const FVector3f& PositionB, const FVector3f& PositionC)->FSHAHash
 	{
 		FSHA1 SHA;
 		FSHAHash SHAHash;
 
-		SHA.Update((const uint8*)&PositionA, sizeof(FVector));
-		SHA.Update((const uint8*)&PositionB, sizeof(FVector));
-		SHA.Update((const uint8*)&PositionC, sizeof(FVector));
+		SHA.Update(reinterpret_cast<const uint8*>(&PositionA), sizeof(FVector3f));
+		SHA.Update(reinterpret_cast<const uint8*>(&PositionB), sizeof(FVector3f));
+		SHA.Update(reinterpret_cast<const uint8*>(&PositionC), sizeof(FVector3f));
 		SHA.Final();
 		SHA.GetHash(&SHAHash.Hash[0]);
 
@@ -53,43 +53,47 @@ void PropagateVertexPaintToSkeletalMesh(USkeletalMesh* SkeletalMesh, int32 LODIn
 	TArray<FSoftSkinVertex> SrcVertices;
 	LODModel.GetVertices(SrcVertices);
 
-	FSkeletalMeshImportData ImportData;
-	SkeletalMesh->LoadLODImportedData(LODIndex, ImportData);
+	FMeshDescription* MeshDescription = SkeletalMesh->GetMeshDescription(LODIndex);
 
 	TMap<FSHAHash, FMatchFaceData> MatchTriangles;
-	MatchTriangles.Reserve(ImportData.Wedges.Num());
+	MatchTriangles.Reserve(MeshDescription->Triangles().Num());
 
 	for (int32 IndexBufferIndex = 0, SrcIndexBufferNum = SrcIndexBuffer.Num(); IndexBufferIndex < SrcIndexBufferNum; IndexBufferIndex += 3)
 	{
-		FVector PositionA = (FVector)SrcVertices[SrcIndexBuffer[IndexBufferIndex]].Position;
-		FVector PositionB = (FVector)SrcVertices[SrcIndexBuffer[IndexBufferIndex + 1]].Position;
-		FVector PositionC = (FVector)SrcVertices[SrcIndexBuffer[IndexBufferIndex + 2]].Position;
+		const FVector3f& PositionA = SrcVertices[SrcIndexBuffer[IndexBufferIndex]].Position;
+		const FVector3f& PositionB = SrcVertices[SrcIndexBuffer[IndexBufferIndex + 1]].Position;
+		const FVector3f& PositionC = SrcVertices[SrcIndexBuffer[IndexBufferIndex + 2]].Position;
 
 		FSHAHash Key = GetMatchKey(PositionA, PositionB, PositionC);
 		FMatchFaceData MatchFaceData;
-		MatchFaceData.SoftVerticeIndexes[0] = SrcIndexBuffer[IndexBufferIndex];
-		MatchFaceData.SoftVerticeIndexes[1] = SrcIndexBuffer[IndexBufferIndex + 1];
-		MatchFaceData.SoftVerticeIndexes[2] = SrcIndexBuffer[IndexBufferIndex + 2];
+		MatchFaceData.SoftVertexIndexes[0] = SrcIndexBuffer[IndexBufferIndex];
+		MatchFaceData.SoftVertexIndexes[1] = SrcIndexBuffer[IndexBufferIndex + 1];
+		MatchFaceData.SoftVertexIndexes[2] = SrcIndexBuffer[IndexBufferIndex + 2];
 		MatchTriangles.Add(Key, MatchFaceData);
 	}
 
+	FStaticMeshAttributes MeshAttributes(*MeshDescription);
+	TVertexInstanceAttributesRef<FVector4f> ColorAttribute = MeshAttributes.GetVertexInstanceColors();
+
 	bool bAllVertexFound = true;
-	for (int32 FaceIndex = 0, FaceNum = ImportData.Faces.Num(); FaceIndex < FaceNum; ++FaceIndex)
+	for (const FTriangleID TriangleID: MeshDescription->Triangles().GetElementIDs())
 	{
-		const SkeletalMeshImportData::FTriangle& Triangle = ImportData.Faces[FaceIndex];
-		SkeletalMeshImportData::FVertex& WedgeA = ImportData.Wedges[Triangle.WedgeIndex[0]];
-		SkeletalMeshImportData::FVertex& WedgeB = ImportData.Wedges[Triangle.WedgeIndex[1]];
-		SkeletalMeshImportData::FVertex& WedgeC = ImportData.Wedges[Triangle.WedgeIndex[2]];
-		FVector PositionA = (FVector)ImportData.Points[WedgeA.VertexIndex];
-		FVector PositionB = (FVector)ImportData.Points[WedgeB.VertexIndex];
-		FVector PositionC = (FVector)ImportData.Points[WedgeC.VertexIndex];
+		TArrayView<const FVertexID> TriangleVertexIDs = MeshDescription->GetTriangleVertices(TriangleID);
+
+		FVector3f PositionA = MeshDescription->GetVertexPosition(TriangleVertexIDs[0]);
+		FVector3f PositionB = MeshDescription->GetVertexPosition(TriangleVertexIDs[1]);
+		FVector3f PositionC = MeshDescription->GetVertexPosition(TriangleVertexIDs[2]);
 
 		const FSHAHash Key = GetMatchKey(PositionA, PositionB, PositionC);
-		if (FMatchFaceData* MatchFaceData = MatchTriangles.Find(Key))
+		if (const FMatchFaceData* MatchFaceData = MatchTriangles.Find(Key))
 		{
-			WedgeA.Color = SrcVertices[MatchFaceData->SoftVerticeIndexes[0]].Color;
-			WedgeB.Color = SrcVertices[MatchFaceData->SoftVerticeIndexes[1]].Color;
-			WedgeC.Color = SrcVertices[MatchFaceData->SoftVerticeIndexes[2]].Color;
+			TArrayView<const FVertexInstanceID> TriangleVertexIndexIDs = MeshDescription->GetTriangleVertexInstances(TriangleID);
+
+			for (int32 Index = 0; Index < 3; Index++)
+			{
+				FLinearColor Color(SrcVertices[MatchFaceData->SoftVertexIndexes[Index]].Color.ReinterpretAsLinear());
+				ColorAttribute.Set(TriangleVertexIndexIDs[Index], Color);
+			}
 		}
 		else if (bAllVertexFound)
 		{
@@ -98,8 +102,8 @@ void PropagateVertexPaintToSkeletalMesh(USkeletalMesh* SkeletalMesh, int32 LODIn
 			UE_LOG(LogMeshPaintSkeletalMeshAdapter, Warning, TEXT("Some vertex color data could not be applied to the %s SkeletalMesh asset."), *SkeletalMeshName);
 		}
 	}
-
-	SkeletalMesh->SaveLODImportedData(LODIndex, ImportData);
+	
+	SkeletalMesh->CommitMeshDescription(LODIndex);
 }
 
 
diff --git a/Engine/Plugins/Runtime/MeshModelingToolset/Source/ModelingComponentsEditorOnly/Private/ToolTargets/SkeletalMeshComponentToolTarget.cpp b/Engine/Plugins/Runtime/MeshModelingToolset/Source/ModelingComponentsEditorOnly/Private/ToolTargets/SkeletalMeshComponentToolTarget.cpp
--- a/Engine/Plugins/Runtime/MeshModelingToolset/Source/ModelingComponentsEditorOnly/Private/ToolTargets/SkeletalMeshComponentToolTarget.cpp
+++ b/Engine/Plugins/Runtime/MeshModelingToolset/Source/ModelingComponentsEditorOnly/Private/ToolTargets/SkeletalMeshComponentToolTarget.cpp
@@ -9,16 +9,28 @@
 #include "Rendering/SkeletalMeshModel.h"
 #include "ToolTargets/SkeletalMeshToolTarget.h"
 #include "SkeletalMeshAttributes.h"
+#include "Framework/Notifications/NotificationManager.h"
 #include "UObject/Package.h"
+#include "Widgets/Notifications/SNotificationList.h"
 
 #include UE_INLINE_GENERATED_CPP_BY_NAME(SkeletalMeshComponentToolTarget)
 
 using namespace UE::Geometry;
 
-namespace USkeletalMeshComponentToolTargetLocals
+#define LOCTEXT_NAMESPACE "SkeletalMeshComponentToolTarget"
+
+namespace SkeletalMeshComponentToolTargetLocals
+{
+static void DisplayCriticalWarningMessage(const FText& InMessage)
 {
-	int32 LODIndex = 0;
+	FNotificationInfo Info(InMessage);
+	Info.ExpireDuration = 5.0f;
+	FSlateNotificationManager::Get().AddNotification(Info);
+
+	UE_LOG(LogGeometry, Warning, TEXT("%s"), *InMessage.ToString());
 }
+}
+
 
 
 //
@@ -109,20 +121,15 @@ bool USkeletalMeshComponentReadOnlyToolTarget::CommitMaterialSetUpdate(const FCo
 
 const FMeshDescription* USkeletalMeshComponentReadOnlyToolTarget::GetMeshDescription(const FGetMeshParameters& GetMeshParams)
 {
-	if (!ensure(IsValid()))
+	if (ensure(IsValid()))
 	{
-		return nullptr;
-	}
-	ensure(GetMeshParams.bHaveRequestLOD == false);	// not supported yet, just returning default LOD
-
-	if (!CachedMeshDescription.IsValid())
-	{
-		CachedMeshDescription = MakeUnique<FMeshDescription>();
-		const USkeletalMesh* SkeletalMesh = Cast<USkeletalMeshComponent>(Component)->GetSkeletalMeshAsset();
-		USkeletalMeshToolTarget::GetMeshDescription(SkeletalMesh, *CachedMeshDescription);
+		const USkeletalMesh* SkeletalMesh = GetSkeletalMesh(); 
+		const int32 LODIndex = USkeletalMeshReadOnlyToolTarget::GetValidEditingLOD(SkeletalMesh, DefaultEditingLOD, GetMeshParams.bHaveRequestLOD, GetMeshParams.RequestLOD);
+	
+		return SkeletalMesh->GetMeshDescription(LODIndex);
 	}
 	
-	return CachedMeshDescription.Get();
+	return nullptr;
 }
 
 FMeshDescription USkeletalMeshComponentReadOnlyToolTarget::GetEmptyMeshDescription()
@@ -134,6 +141,11 @@ FMeshDescription USkeletalMeshComponentReadOnlyToolTarget::GetEmptyMeshDescripti
 	return EmptyMeshDescription;
 }
 
+TArray<EMeshLODIdentifier> USkeletalMeshComponentReadOnlyToolTarget::GetAvailableLODs(bool bSkipAutoGenerated) const
+{
+	return USkeletalMeshReadOnlyToolTarget::GetAvailableLODs(GetSkeletalMesh(), bSkipAutoGenerated);
+}
+
 FDynamicMesh3 USkeletalMeshComponentReadOnlyToolTarget::GetDynamicMesh()
 {
 	return GetDynamicMeshViaMeshDescription(*this);
@@ -158,24 +170,20 @@ USkeletalMesh* USkeletalMeshComponentReadOnlyToolTarget::GetSkeletalMesh() const
 
 void USkeletalMeshComponentToolTarget::CommitMeshDescription(const FCommitter& Committer, const FCommitMeshParameters& CommitMeshParams)
 {
-	if (ensure(IsValid()) == false) return;
-
-	USkeletalMesh* SkeletalMesh = Cast<USkeletalMeshComponent>(Component)->GetSkeletalMeshAsset();
-	if (!CachedMeshDescription.IsValid())
+	if (ensure(IsValid()))
 	{
-		CachedMeshDescription = MakeUnique<FMeshDescription>();
-		USkeletalMeshToolTarget::GetMeshDescription(SkeletalMesh, *CachedMeshDescription);
-	}
+		USkeletalMesh* SkeletalMesh = GetSkeletalMesh();
+		
+		const int32 LODIndex = USkeletalMeshReadOnlyToolTarget::GetValidEditingLOD(SkeletalMesh, DefaultEditingLOD, CommitMeshParams.bHaveTargetLOD, CommitMeshParams.TargetLOD);
 
-	// unregister the component while we update its skeletal mesh
-	FComponentReregisterContext ComponentReregisterContext(Component.Get());
+		// Unregister the component while we update its skeletal mesh
+		FComponentReregisterContext ComponentReregisterContext(Component.Get());
 
-	USkeletalMeshToolTarget::CommitMeshDescription(SkeletalMesh, CachedMeshDescription.Get(), Committer);
+		USkeletalMeshToolTarget::CommitMeshDescription(SkeletalMesh, Committer, LODIndex);
 
-	// this rebuilds physics, but it doesn't undo!
-	Component->RecreatePhysicsState();
-
-	CachedMeshDescription.Reset();
+		// this rebuilds physics, but it doesn't undo!
+		Component->RecreatePhysicsState();
+	}
 }
 
 
@@ -240,3 +248,4 @@ UToolTarget* USkeletalMeshComponentToolTargetFactory::BuildTarget(UObject* Sourc
 	return Target;
 }
 
+#undef LOCTEXT_NAMESPACE
diff --git a/Engine/Plugins/Runtime/MeshModelingToolset/Source/ModelingComponentsEditorOnly/Private/ToolTargets/SkeletalMeshToolTarget.cpp b/Engine/Plugins/Runtime/MeshModelingToolset/Source/ModelingComponentsEditorOnly/Private/ToolTargets/SkeletalMeshToolTarget.cpp
--- a/Engine/Plugins/Runtime/MeshModelingToolset/Source/ModelingComponentsEditorOnly/Private/ToolTargets/SkeletalMeshToolTarget.cpp
+++ b/Engine/Plugins/Runtime/MeshModelingToolset/Source/ModelingComponentsEditorOnly/Private/ToolTargets/SkeletalMeshToolTarget.cpp
@@ -9,6 +9,7 @@
 #include "Widgets/Notifications/SNotificationList.h"
 #include "Framework/Notifications/NotificationManager.h"
 #include "SkeletalMeshAttributes.h"
+#include "Misc/SlowTask.h"
 #include "UObject/Package.h"
 
 #include UE_INLINE_GENERATED_CPP_BY_NAME(SkeletalMeshToolTarget)
@@ -17,9 +18,17 @@ using namespace UE::Geometry;
 
 #define LOCTEXT_NAMESPACE "SkeletalMeshToolTarget"
 
-namespace USkeletalMeshToolTargetLocals
+
+namespace SkeletalMeshToolTargetLocals
+{
+static void DisplayCriticalWarningMessage(const FText& InMessage)
 {
-	int32 LODIndex = 0;
+	FNotificationInfo Info(InMessage);
+	Info.ExpireDuration = 5.0f;
+	FSlateNotificationManager::Get().AddNotification(Info);
+
+	UE_LOG(LogGeometry, Warning, TEXT("%s"), *InMessage.ToString());
+}
 }
 
 
@@ -43,6 +52,30 @@ bool USkeletalMeshReadOnlyToolTarget::IsValid(const USkeletalMesh* SkeletalMeshI
 	return true;
 }
 
+TArray<EMeshLODIdentifier> USkeletalMeshReadOnlyToolTarget::GetAvailableLODs(const USkeletalMesh* SkeletalMesh, bool bSkipAutoGenerated)
+{
+	TArray<EMeshLODIdentifier> AvailableLODs;
+
+	if (!SkeletalMesh)
+	{
+		return AvailableLODs;
+	}
+
+	const int32 MaxLODs = FMath::Min(8, SkeletalMesh->GetNumSourceModels());
+
+	for (int32 LODIndex = 0; LODIndex < MaxLODs; LODIndex++)
+	{
+		if (bSkipAutoGenerated && !SkeletalMesh->HasMeshDescription(LODIndex))
+		{
+			continue;
+		}
+
+		AvailableLODs.Add(static_cast<EMeshLODIdentifier>(LODIndex));
+	}
+
+	return AvailableLODs;
+}
+
 int32 USkeletalMeshReadOnlyToolTarget::GetNumMaterials() const
 {
 	return ensure(IsValid()) ? SkeletalMesh->GetMaterials().Num() : 0;
@@ -122,21 +155,48 @@ bool USkeletalMeshReadOnlyToolTarget::CommitMaterialSetUpdate(USkeletalMesh* Ske
 	return true;
 }
 
-const FMeshDescription* USkeletalMeshReadOnlyToolTarget::GetMeshDescription(const FGetMeshParameters& GetMeshParams)
+int32 USkeletalMeshReadOnlyToolTarget::GetValidEditingLOD(const USkeletalMesh* SkeletalMesh, int32 DefaultEditingLOD, bool bHaveRequestLOD, EMeshLODIdentifier RequestedEditingLOD)
 {
-	if (!ensure(IsValid()))
+	using namespace SkeletalMeshToolTargetLocals;
+
+	int32 ValidEditingLOD = DefaultEditingLOD;
+
+	if (bHaveRequestLOD)
 	{
-		return nullptr;
+		if (RequestedEditingLOD == EMeshLODIdentifier::MaxQuality)
+		{
+			ValidEditingLOD = 0;
+		}
+		else if (RequestedEditingLOD == EMeshLODIdentifier::HiResSource)
+		{
+			DisplayCriticalWarningMessage(LOCTEXT("SkelmeshHiresLODUnavailable", "HiRes LOD not available on skeletal meshes. Falling back to using LOD 0."));
+		}
+		else if (RequestedEditingLOD != EMeshLODIdentifier::Default)
+		{
+			ValidEditingLOD = static_cast<int32>(RequestedEditingLOD);
+		}
+	}
+	
+	const int32 MaxLOD = SkeletalMesh->GetLODNum() - 1;
+	if (ValidEditingLOD > MaxLOD)
+	{
+		DisplayCriticalWarningMessage(FText::Format(LOCTEXT("SkelmeshLODUnavailable", "LOD {0} requested but is not available. Falling back to LOD {1}."), FText::AsNumber(ValidEditingLOD), FText::AsNumber(MaxLOD)));
+		ValidEditingLOD = MaxLOD;
 	}
-	ensure(GetMeshParams.bHaveRequestLOD == false);	// not supported yet, just returning default LOD
 
-	if (!CachedMeshDescription.IsValid())
+	return ValidEditingLOD;
+}
+
+const FMeshDescription* USkeletalMeshReadOnlyToolTarget::GetMeshDescription(const FGetMeshParameters& GetMeshParams)
+{
+	if (ensure(IsValid()))
 	{
-		CachedMeshDescription = MakeUnique<FMeshDescription>();
-		GetMeshDescription(SkeletalMesh.Get(), *CachedMeshDescription);
+		const int32 LODIndex = GetValidEditingLOD(SkeletalMesh.Get(), DefaultEditingLOD, GetMeshParams.bHaveRequestLOD, GetMeshParams.RequestLOD);
+	
+		return SkeletalMesh->GetMeshDescription(LODIndex);
 	}
 
-	return CachedMeshDescription.Get();
+	return nullptr;
 }
 
 
@@ -149,10 +209,9 @@ FMeshDescription USkeletalMeshReadOnlyToolTarget::GetEmptyMeshDescription()
 	return EmptyMeshDescription;
 }
 
-void USkeletalMeshReadOnlyToolTarget::GetMeshDescription(const USkeletalMesh* SkeletalMeshIn, FMeshDescription& MeshDescriptionOut)
+TArray<EMeshLODIdentifier> USkeletalMeshReadOnlyToolTarget::GetAvailableLODs(bool bSkipAutoGenerated) const
 {
-	using namespace USkeletalMeshToolTargetLocals;
-	SkeletalMeshIn->GetMeshDescription(LODIndex, MeshDescriptionOut);
+	return GetAvailableLODs(SkeletalMesh.Get(), bSkipAutoGenerated);
 }
 
 FDynamicMesh3 USkeletalMeshReadOnlyToolTarget::GetDynamicMesh()
@@ -176,42 +235,43 @@ USkeletalMesh* USkeletalMeshReadOnlyToolTarget::GetSkeletalMesh() const
 
 void USkeletalMeshToolTarget::CommitMeshDescription(const FCommitter& Committer, const FCommitMeshParameters& CommitMeshParams)
 {
-	if (ensure(IsValid()) == false) return;
-
-	if (!CachedMeshDescription.IsValid())
+	if (ensure(IsValid()))
 	{
-		CachedMeshDescription = MakeUnique<FMeshDescription>();
-		GetMeshDescription(SkeletalMesh.Get(), *CachedMeshDescription);
+		const int32 LODIndex = GetValidEditingLOD(SkeletalMesh.Get(), DefaultEditingLOD, CommitMeshParams.bHaveTargetLOD, CommitMeshParams.TargetLOD);
+		
+		CommitMeshDescription(SkeletalMesh.Get(), Committer, LODIndex);
 	}
-	CommitMeshDescription(SkeletalMesh.Get(), CachedMeshDescription.Get(), Committer);
 }
 
-void USkeletalMeshToolTarget::CommitMeshDescription(USkeletalMesh* SkeletalMeshIn,
-	FMeshDescription* MeshDescription, const FCommitter& Committer)
+void USkeletalMeshToolTarget::CommitMeshDescription(USkeletalMesh* SkeletalMesh, const FCommitter& Committer, int32 LODIndex)
 {
-	using namespace USkeletalMeshToolTargetLocals;
+	using namespace SkeletalMeshToolTargetLocals;
 
-	if (SkeletalMeshIn->GetPathName().StartsWith(TEXT("/Engine/")))
+	if (SkeletalMesh->GetPathName().StartsWith(TEXT("/Engine/")))
 	{
-		FText Error = FText::Format(LOCTEXT("CannotModifyBuiltInAssetError", "Cannot modify built-in engine asset: {0}"), FText::FromString(*SkeletalMeshIn->GetPathName()));
-		FNotificationInfo Info(Error);
-		Info.ExpireDuration = 5.0f;
-		FSlateNotificationManager::Get().AddNotification(Info);
-
-		UE_LOG(LogTemp, Warning, TEXT("%s"), *Error.ToString());
+		DisplayCriticalWarningMessage(FText::Format(LOCTEXT("CannotModifyBuiltInAssetError", "Cannot modify built-in engine asset: {0}"), FText::FromString(*SkeletalMesh->GetPathName())));
 		return;
 	}
 
 	// flush any pending rendering commands, which might touch a component while we are rebuilding it's mesh
 	FlushRenderingCommands();
 
+	// Ensure the asset is set to transactional, so that undo works proper.
+	SkeletalMesh->SetFlags(RF_Transactional);
+	SkeletalMesh->Modify();
+
+	// Mark the mesh at this LOD as modified as well, so that we can undo any geometry changes.
+	SkeletalMesh->ModifyMeshDescription(LODIndex);
+	
 	FCommitterParams CommitterParams;
-	CommitterParams.MeshDescriptionOut = MeshDescription;
+	CommitterParams.MeshDescriptionOut = SkeletalMesh->GetMeshDescription(LODIndex);
 	Committer(CommitterParams);
-	
-	SkeletalMeshIn->CommitMeshDescription(LODIndex, *CommitterParams.MeshDescriptionOut);
 
-	SkeletalMeshIn->PostEditChange();
+	// Commit the edited mesh description to bulk data.
+	SkeletalMesh->CommitMeshDescription(LODIndex);
+
+	// This call will invoke the skeletal mesh render data rebuild machinery.
+	SkeletalMesh->PostEditChange();
 }
 
 void USkeletalMeshToolTarget::CommitDynamicMesh(const FDynamicMesh3& Mesh, const FDynamicMeshCommitInfo& CommitInfo)
diff --git a/Engine/Plugins/Runtime/MeshModelingToolset/Source/ModelingComponentsEditorOnly/Public/ToolTargets/SkeletalMeshComponentToolTarget.h b/Engine/Plugins/Runtime/MeshModelingToolset/Source/ModelingComponentsEditorOnly/Public/ToolTargets/SkeletalMeshComponentToolTarget.h
--- a/Engine/Plugins/Runtime/MeshModelingToolset/Source/ModelingComponentsEditorOnly/Public/ToolTargets/SkeletalMeshComponentToolTarget.h
+++ b/Engine/Plugins/Runtime/MeshModelingToolset/Source/ModelingComponentsEditorOnly/Public/ToolTargets/SkeletalMeshComponentToolTarget.h
@@ -40,7 +40,11 @@ public:
 	// IMeshDescriptionProvider implementation
 	virtual const FMeshDescription* GetMeshDescription(const FGetMeshParameters& GetMeshParams = FGetMeshParameters()) override;
 	virtual FMeshDescription GetEmptyMeshDescription() override;
+	virtual bool SupportsLODs() const override { return true; }
+	virtual TArray<EMeshLODIdentifier> GetAvailableLODs(bool bSkipAutoGenerated = true) const override;
+	virtual EMeshLODIdentifier GetMeshDescriptionLOD() const override { return static_cast<EMeshLODIdentifier>(DefaultEditingLOD); }
 
+	
 	// IMaterialProvider implementation
 	virtual int32 GetNumMaterials() const override;
 	virtual UMaterialInterface* GetMaterial(int32 MaterialIndex) const override;
@@ -55,12 +59,11 @@ public:
 	virtual USkeletalMesh* GetSkeletalMesh() const override;
 
 protected:
+	// The LOD to edit if requested to get/commit at EMeshLODIdentifier::Default or if no specific LOD is requested.
+	int32 DefaultEditingLOD = 0;	
+	
 	// So that the tool target factory can poke into Component.
 	friend class USkeletalMeshComponentReadOnlyToolTargetFactory;
-
-	// Until USkeletalMesh stores its internal representation as FMeshDescription, we need to
-	// retain the storage here to cover the lifetime of the pointer returned by GetMeshDescription(). 
-	TUniquePtr<FMeshDescription> CachedMeshDescription;
 };
 
 
diff --git a/Engine/Plugins/Runtime/MeshModelingToolset/Source/ModelingComponentsEditorOnly/Public/ToolTargets/SkeletalMeshToolTarget.h b/Engine/Plugins/Runtime/MeshModelingToolset/Source/ModelingComponentsEditorOnly/Public/ToolTargets/SkeletalMeshToolTarget.h
--- a/Engine/Plugins/Runtime/MeshModelingToolset/Source/ModelingComponentsEditorOnly/Public/ToolTargets/SkeletalMeshToolTarget.h
+++ b/Engine/Plugins/Runtime/MeshModelingToolset/Source/ModelingComponentsEditorOnly/Public/ToolTargets/SkeletalMeshToolTarget.h
@@ -40,6 +40,9 @@ public:
 	// IMeshDescriptionProvider implementation
 	virtual const FMeshDescription* GetMeshDescription(const FGetMeshParameters& GetMeshParams = FGetMeshParameters()) override;
 	virtual FMeshDescription GetEmptyMeshDescription() override;
+	virtual bool SupportsLODs() const override { return true; }
+	virtual TArray<EMeshLODIdentifier> GetAvailableLODs(bool bSkipAutoGenerated = true) const override;
+	virtual EMeshLODIdentifier GetMeshDescriptionLOD() const override { return static_cast<EMeshLODIdentifier>(DefaultEditingLOD); }
 
 	// IMaterialProvider implementation
 	virtual int32 GetNumMaterials() const override;
@@ -57,6 +60,9 @@ public:
 protected:
 	TWeakObjectPtr<USkeletalMesh> SkeletalMesh = nullptr;
 
+	// The LOD to edit if requested to get/commit at EMeshLODIdentifier::Default or if no specific LOD is requested.
+	int32 DefaultEditingLOD = 0;	
+	
 	// So that the tool target factory can poke into Component.
 	friend class USkeletalMeshReadOnlyToolTargetFactory;
 	friend class USkeletalMeshComponentReadOnlyToolTarget;
@@ -65,15 +71,12 @@ protected:
 	UE_DEPRECATED(5.3, "This function doesn't check anything skeletal mesh specific, only (badly) checks to see if the object is alive")
 	static bool IsValid(const USkeletalMesh* USkeletalMesh);
 
-	static void GetMeshDescription(const USkeletalMesh* SkeletalMesh, FMeshDescription& MeshDescriptionOut);
+	static TArray<EMeshLODIdentifier> GetAvailableLODs(const USkeletalMesh* SkeletalMesh, bool bSkipAutoGenerated);	
 	static void GetMaterialSet(const USkeletalMesh* SkeletalMesh, FComponentMaterialSet& MaterialSetOut,
 		bool bPreferAssetMaterials);
 	static bool CommitMaterialSetUpdate(USkeletalMesh* SkeletalMesh,
 		const FComponentMaterialSet& MaterialSet, bool bApplyToAsset);
-
-	// Until USkeletalMesh stores its internal representation as FMeshDescription, we need to
-	// retain the storage here to cover the lifetime of the pointer returned by GetMeshDescription(). 
-	TUniquePtr<FMeshDescription> CachedMeshDescription;	
+	static int32 GetValidEditingLOD(const USkeletalMesh* SkeletalMesh, int32 DefaultEditingLOD, bool bHaveRequestLOD, EMeshLODIdentifier RequestedEditingLOD);	
 };
 
 
@@ -103,8 +106,7 @@ protected:
 	friend class USkeletalMeshToolTargetFactory;
 	friend class USkeletalMeshComponentToolTarget;
 
-	static void CommitMeshDescription(USkeletalMesh* SkeletalMesh, 
-		FMeshDescription* MeshDescription, const FCommitter& Committer);
+	static void CommitMeshDescription(USkeletalMesh* SkeletalMesh, const FCommitter& Committer, int32 LODIndex);
 };
 
 
diff --git a/Engine/Source/Developer/MeshBuilder/Private/SkeletalMeshBuilder.cpp b/Engine/Source/Developer/MeshBuilder/Private/SkeletalMeshBuilder.cpp
--- a/Engine/Source/Developer/MeshBuilder/Private/SkeletalMeshBuilder.cpp
+++ b/Engine/Source/Developer/MeshBuilder/Private/SkeletalMeshBuilder.cpp
@@ -92,7 +92,9 @@ bool FSkeletalMeshBuilder::Build(const FSkeletalMeshBuildParameters& SkeletalMes
 		FSkeletalMeshLODModel& BuildLODModel = SkeletalMesh->GetImportedModel()->LODModels[LODIndex];
 
 		//Load the imported data
+		PRAGMA_DISABLE_DEPRECATION_WARNINGS
 		SkeletalMesh->LoadLODImportedData(LODIndex, SkeletalMeshImportData);
+		PRAGMA_ENABLE_DEPRECATION_WARNINGS
 
 		TArray<FVector3f> LODPoints;
 		TArray<SkeletalMeshImportData::FMeshWedge> LODWedges;
@@ -173,7 +175,7 @@ bool FSkeletalMeshBuilder::Build(const FSkeletalMeshBuildParameters& SkeletalMes
 		{
 			SlowTask.EnterProgressFrame(1.0f, NSLOCTEXT("SkeltalMeshBuilder", "RegenerateLOD", "Regenerate LOD..."));
 			//Update the original reduction data since we just build a new LODModel.
-			if (LODInfo->ReductionSettings.BaseLOD == LODIndex && SkeletalMesh->IsLODImportedDataBuildAvailable(LODIndex))
+			if (LODInfo->ReductionSettings.BaseLOD == LODIndex && SkeletalMesh->HasMeshDescription(LODIndex))
 			{
 				if (LODIndex == 0)
 				{
diff --git a/Engine/Source/Developer/MeshUtilities/Private/MeshUtilities.cpp b/Engine/Source/Developer/MeshUtilities/Private/MeshUtilities.cpp
--- a/Engine/Source/Developer/MeshUtilities/Private/MeshUtilities.cpp
+++ b/Engine/Source/Developer/MeshUtilities/Private/MeshUtilities.cpp
@@ -5931,22 +5931,23 @@ void FMeshUtilities::GenerateRuntimeSkinWeightData(
 	}
 }
 
-void FMeshUtilities::CreateImportDataFromLODModel(USkeletalMesh* SkeletalMesh) const
+void FMeshUtilities::CreateImportDataFromLODModel(USkeletalMesh* SkeletalMesh, bool bInResetReductionAsNeeded) const
 {
 	check(IsInGameThread());
 
 	for (int32 LodIndex = 0; LodIndex < SkeletalMesh->GetLODNum(); LodIndex++)
 	{
-		const FSkeletalMeshLODInfo* ThisLODInfo = SkeletalMesh->GetLODInfo(LodIndex);
+		FSkeletalMeshLODInfo* ThisLODInfo = SkeletalMesh->GetLODInfo(LodIndex);
 
 		check(ThisLODInfo);
-		const bool bRawDataEmpty = SkeletalMesh->IsLODImportedDataEmpty(LodIndex);
-		const bool bRawBuildDataAvailable = SkeletalMesh->IsLODImportedDataBuildAvailable(LodIndex);
-		if (!bRawDataEmpty && bRawBuildDataAvailable)
+		if (SkeletalMesh->HasMeshDescription(LodIndex))
 		{
 			//No need to create import data if we already have some
 			continue;
 		}
+
+		// If the mesh was not pulled out of the reduction data, we need to reset the LOD settings
+		// so that the mesh doesn't get reduced again if it gets regenerated.
 		const bool bReductionActive = SkeletalMesh->IsReductionActive(LodIndex);
 		const bool bInlineReduction = (ThisLODInfo->ReductionSettings.BaseLOD == LodIndex);
 		if (bReductionActive && !bInlineReduction)
@@ -6004,6 +6005,10 @@ void FMeshUtilities::CreateImportDataFromLODModel(USkeletalMesh* SkeletalMesh) c
 					ExistingMorphTarget->PopulateDeltas(MorphTargetDeltas, LodIndex, ToUpdateLODModel.Sections, ThisLODInfo->BuildSettings.bRecomputeNormals, false, ThisLODInfo->BuildSettings.MorphThresholdPosition);
 				}
 			}
+			
+			// The model got pulled out of the reduction storage, don't reset the LOD reduction settings,
+			// since we may want to regenerate the mesh to those settings.
+			bInResetReductionAsNeeded = false;
 		}
 
 		FSkeletalMeshLODModel* LODModel = &(SkeletalMesh->GetImportedModel()->LODModels[LodIndex]);
@@ -6222,9 +6227,24 @@ void FMeshUtilities::CreateImportDataFromLODModel(USkeletalMesh* SkeletalMesh) c
 			}
 		}
 
+		PRAGMA_DISABLE_DEPRECATION_WARNINGS
 		SkeletalMesh->SaveLODImportedData(LodIndex, ImportData);
-		SkeletalMesh->SetLODImportedDataVersions(LodIndex, ESkeletalMeshGeoImportVersions::LatestVersion, ESkeletalMeshSkinningImportVersions::LatestVersion);
-		UE_ASSET_LOG(LogSkeletalMesh, Display, SkeletalMesh, TEXT("FMeshUtilities::CreateImportDataFromLODModel: ImportData created for LOD %d"), LodIndex);
+		PRAGMA_ENABLE_DEPRECATION_WARNINGS
+		
+		if (bInResetReductionAsNeeded)
+		{
+			FSkeletalMeshOptimizationSettings& ReductionSettings = ThisLODInfo->ReductionSettings;
+			
+			//Remove the reduction settings
+			ReductionSettings.NumOfTrianglesPercentage = 1.0f;
+			ReductionSettings.NumOfVertPercentage = 1.0f;
+			ReductionSettings.MaxNumOfTrianglesPercentage = MAX_uint32;
+			ReductionSettings.MaxNumOfVertsPercentage = MAX_uint32;
+			ReductionSettings.TerminationCriterion = SMTC_NumOfTriangles;
+			ThisLODInfo->bHasBeenSimplified = false;
+		}
+		
+		UE_ASSET_LOG(LogSkeletalMesh, Display, SkeletalMesh, TEXT("Import data converted for LOD %d"), LodIndex);
 	}
 }
 
diff --git a/Engine/Source/Developer/MeshUtilities/Private/MeshUtilitiesPrivate.h b/Engine/Source/Developer/MeshUtilities/Private/MeshUtilitiesPrivate.h
--- a/Engine/Source/Developer/MeshUtilities/Private/MeshUtilitiesPrivate.h
+++ b/Engine/Source/Developer/MeshUtilities/Private/MeshUtilitiesPrivate.h
@@ -156,7 +156,7 @@ private:
 
 	virtual void GenerateRuntimeSkinWeightData(const FSkeletalMeshLODModel* ImportedModel, const TArray<FRawSkinWeight>& InRawSkinWeights, bool bInUseHighPrecisionWeights, FRuntimeSkinWeightProfileData& InOutSkinWeightOverrideData) const override;
 
-	virtual void CreateImportDataFromLODModel(USkeletalMesh* SkeletalMesh) const override;
+	virtual void CreateImportDataFromLODModel(USkeletalMesh* SkeletalMesh, bool bInResetReductionAsNeeded) const override;
 
 	void RegisterMenus();
 
diff --git a/Engine/Source/Developer/MeshUtilities/Public/MeshUtilities.h b/Engine/Source/Developer/MeshUtilities/Public/MeshUtilities.h
--- a/Engine/Source/Developer/MeshUtilities/Public/MeshUtilities.h
+++ b/Engine/Source/Developer/MeshUtilities/Public/MeshUtilities.h
@@ -202,6 +202,7 @@ public:
 		TArray<FText> * OutWarningMessages = NULL,
 		TArray<FName> * OutWarningNames = NULL
 		) = 0;
+
 	
 	/** Cache optimize the index buffer. */
 	virtual void CacheOptimizeIndexBuffer(TArray<uint16>& Indices) = 0;
@@ -341,6 +342,8 @@ public:
 	/*
 	 * This function create the import data using the LODModel. You can call this function if you load an asset that was not re-import since the build refactor and the chunking is more agressive than the bake data in the LODModel.
 	 * You can also need this function if you create a skeletalmesh with LODModel instead of import data, so your newly created skeletalmesh can be build properly.
+	 * If the LODModel is not being pulled out of the old reduction storage, and bInResetReductionAsNeeded is true, then the reduction settings for that
+	 * LOD will be reset to avoid regenerating the mesh even more reduced.
 	 */
-	virtual void CreateImportDataFromLODModel(USkeletalMesh* SkeletalMesh) const = 0;
+	virtual void CreateImportDataFromLODModel(USkeletalMesh* InSkeletalMesh, bool bInResetReductionAsNeeded = false) const = 0;
 };
