[MLDeformer] The neural morph model's bone mask is generated incorrectly. It seemed like one hierarchy level too deep. This change fixes that. Also the way the mask is setup now has changed, so that you can't just add bones to the mask, but also remove them.

diff --git a/Engine/Plugins/Animation/MLDeformer/MLDeformerFramework/Source/MLDeformerFrameworkEditor/Private/MLDeformerMorphModelEditorModel.cpp b/Engine/Plugins/Animation/MLDeformer/MLDeformerFramework/Source/MLDeformerFrameworkEditor/Private/MLDeformerMorphModelEditorModel.cpp
--- a/Engine/Plugins/Animation/MLDeformer/MLDeformerFramework/Source/MLDeformerFrameworkEditor/Private/MLDeformerMorphModelEditorModel.cpp
+++ b/Engine/Plugins/Animation/MLDeformer/MLDeformerFramework/Source/MLDeformerFrameworkEditor/Private/MLDeformerMorphModelEditorModel.cpp
@@ -876,7 +876,7 @@ namespace UE::MLDeformer
 		}
 	}
 
-	void FMLDeformerMorphModelEditorModel::RecursiveAddBoneToMaskUpwards(const FReferenceSkeleton& RefSkel, int32 SkeletonBoneIndex, int32 MaxHierarchyDepth, const TArray<int32>& VirtualParentTable, TArray<int32>& OutBonesAdded, int32 CurHierarchyDepth)
+	void FMLDeformerMorphModelEditorModel::RecursiveAddBoneToMaskUpwards(const FReferenceSkeleton& RefSkel, int32 SkeletonBoneIndex, int32 MaxHierarchyDepth, TArray<int32>& OutBonesAdded, int32 CurHierarchyDepth)
 	{
 		if (CurHierarchyDepth > MaxHierarchyDepth)
 		{
@@ -893,11 +893,16 @@ namespace UE::MLDeformer
 		const int32 ParentSkeletonBoneIndex = RefSkel.GetParentIndex(SkeletonBoneIndex);
 		if (ParentSkeletonBoneIndex != INDEX_NONE)
 		{
-			RecursiveAddBoneToMaskUpwards(RefSkel, ParentSkeletonBoneIndex, MaxHierarchyDepth, VirtualParentTable, OutBonesAdded, CurHierarchyDepth + 1);
+			RecursiveAddBoneToMaskUpwards(RefSkel, ParentSkeletonBoneIndex, MaxHierarchyDepth, OutBonesAdded, CurHierarchyDepth + 1);
 		}
 	}
 
-	void FMLDeformerMorphModelEditorModel::RecursiveAddBoneToMaskDownwards(const FReferenceSkeleton& RefSkel, int32 SkeletonBoneIndex, int32 MaxHierarchyDepth, const TArray<int32>& VirtualParentTable, TArray<int32>& OutBonesAdded, int32 CurHierarchyDepth)
+	void FMLDeformerMorphModelEditorModel::RecursiveAddBoneToMaskUpwards(const FReferenceSkeleton& RefSkel, int32 SkeletonBoneIndex, int32 MaxHierarchyDepth, const TArray<int32>& VirtualParentTable, TArray<int32>& OutBonesAdded, int32 CurHierarchyDepth)
+	{
+		RecursiveAddBoneToMaskUpwards(RefSkel, SkeletonBoneIndex, MaxHierarchyDepth, OutBonesAdded, CurHierarchyDepth);
+	}
+
+	void FMLDeformerMorphModelEditorModel::RecursiveAddBoneToMaskDownwards(const FReferenceSkeleton& RefSkel, int32 SkeletonBoneIndex, int32 MaxHierarchyDepth, TArray<int32>& OutBonesAdded, int32 CurHierarchyDepth)
 	{
 		if (CurHierarchyDepth > MaxHierarchyDepth)
 		{
@@ -918,10 +923,15 @@ namespace UE::MLDeformer
 		// Now recursively add the child bones.
 		for (const int32 ChildIndex : ChildBones)
 		{
-			RecursiveAddBoneToMaskDownwards(RefSkel, ChildIndex, MaxHierarchyDepth, VirtualParentTable, OutBonesAdded, CurHierarchyDepth + 1);
+			RecursiveAddBoneToMaskDownwards(RefSkel, ChildIndex, MaxHierarchyDepth, OutBonesAdded, CurHierarchyDepth + 1);
 		}
 	}
 
+	void FMLDeformerMorphModelEditorModel::RecursiveAddBoneToMaskDownwards(const FReferenceSkeleton& RefSkel, int32 SkeletonBoneIndex, int32 MaxHierarchyDepth, const TArray<int32>& VirtualParentTable, TArray<int32>& OutBonesAdded, int32 CurHierarchyDepth)
+	{
+		RecursiveAddBoneToMaskDownwards(RefSkel, SkeletonBoneIndex, MaxHierarchyDepth, OutBonesAdded, CurHierarchyDepth);
+	}
+
 	int32 FMLDeformerMorphModelEditorModel::FindVirtualParentIndex(const FReferenceSkeleton& RefSkel, int32 BoneIndex, const TArray<FName>& IncludedBoneNames) const
 	{
 		int32 CurBoneIndex = BoneIndex;
@@ -947,15 +957,17 @@ namespace UE::MLDeformer
 
 	TArray<int32> FMLDeformerMorphModelEditorModel::BuildVirtualParentTable(const FReferenceSkeleton& RefSkel, const TArray<FName>& IncludedBoneNames) const
 	{
+		PRAGMA_DISABLE_DEPRECATION_WARNINGS
+
 		TArray<int32> VirtualParentTable;
 		VirtualParentTable.SetNumUninitialized(RefSkel.GetNum());
-
 		for (int32 BoneIndex = 0; BoneIndex < RefSkel.GetNum(); ++BoneIndex)
 		{
 			VirtualParentTable[BoneIndex] = FindVirtualParentIndex(RefSkel, BoneIndex, IncludedBoneNames);
 		}
-
 		return MoveTemp(VirtualParentTable);
+
+		PRAGMA_ENABLE_DEPRECATION_WARNINGS
 	}
 
 }	// namespace UE::MLDeformer
diff --git a/Engine/Plugins/Animation/MLDeformer/MLDeformerFramework/Source/MLDeformerFrameworkEditor/Private/SMLDeformerBonePickerDialog.cpp b/Engine/Plugins/Animation/MLDeformer/MLDeformerFramework/Source/MLDeformerFrameworkEditor/Private/SMLDeformerBonePickerDialog.cpp
--- a/Engine/Plugins/Animation/MLDeformer/MLDeformerFramework/Source/MLDeformerFrameworkEditor/Private/SMLDeformerBonePickerDialog.cpp
+++ b/Engine/Plugins/Animation/MLDeformer/MLDeformerFramework/Source/MLDeformerFrameworkEditor/Private/SMLDeformerBonePickerDialog.cpp
@@ -24,7 +24,7 @@ namespace UE::MLDeformer
 		SuperArgs.OnGenerateRow(this, &SMLDeformerBonePickerTreeWidget::MakeTableRowWidget);
 		SuperArgs.OnGetChildren(this, &SMLDeformerBonePickerTreeWidget::HandleGetChildrenForTree);
 		SuperArgs.OnMouseButtonDoubleClick(this, &SMLDeformerBonePickerTreeWidget::OnMouseDoubleClicked);
-		SuperArgs.HighlightParentNodesForSelection(true);
+		SuperArgs.HighlightParentNodesForSelection(false);
 		SuperArgs.AllowInvisibleItemSelection(true);  // Without this we deselect everything when we filter or we collapse.
 
 		STreeView<TSharedPtr<FMLDeformerBonePickerTreeElement>>::Construct(SuperArgs);
@@ -141,6 +141,30 @@ namespace UE::MLDeformer
 		}
 	}
 
+	void SMLDeformerBonePickerDialog::SelectInitialItemsRecursive(const TSharedPtr<FMLDeformerBonePickerTreeElement>& Item)
+	{
+		if (!Item.IsValid())
+		{
+			return;
+		}
+
+		const bool bSelected = InitialSelectedBoneNames.Contains(Item->Name);
+		BoneTreeWidget->SetItemSelection(Item, bSelected, ESelectInfo::Type::Direct);
+		
+		for (const TSharedPtr<FMLDeformerBonePickerTreeElement>& ChildItem : Item->Children)
+		{
+			SelectInitialItemsRecursive(ChildItem);
+		}
+	}
+
+	void SMLDeformerBonePickerDialog::SelectInitialItems()
+	{
+		for (const TSharedPtr<FMLDeformerBonePickerTreeElement>& Item : BoneTreeWidget->GetRootItems())
+		{
+			SelectInitialItemsRecursive(Item);
+		}
+	}
+
 	TSharedRef<SMLDeformerBonePickerTreeWidget> SMLDeformerBonePickerDialog::CreateBoneTree()
 	{
 		TSharedRef<SMLDeformerBonePickerTreeWidget> TreeWidget = 
@@ -149,6 +173,7 @@ namespace UE::MLDeformer
 			.PickerWidget(SharedThis(this));
 
 		RefreshBoneTreeElements();
+		SelectInitialItems();
 		return TreeWidget;
 	}
 
@@ -201,6 +226,7 @@ namespace UE::MLDeformer
 		bAllowMultiSelect = InArgs._AllowMultiSelect;
 		IncludeList = InArgs._IncludeList;
 		ExtraWidget = InArgs._ExtraWidget;
+		InitialSelectedBoneNames = InArgs._InitialSelectedBoneNames;
 
 		FText DialogTitle;
 		if (InArgs._Title.ToString().IsEmpty())
diff --git a/Engine/Plugins/Animation/MLDeformer/MLDeformerFramework/Source/MLDeformerFrameworkEditor/Private/SMLDeformerInputBonesWidget.cpp b/Engine/Plugins/Animation/MLDeformer/MLDeformerFramework/Source/MLDeformerFrameworkEditor/Private/SMLDeformerInputBonesWidget.cpp
--- a/Engine/Plugins/Animation/MLDeformer/MLDeformerFramework/Source/MLDeformerFrameworkEditor/Private/SMLDeformerInputBonesWidget.cpp
+++ b/Engine/Plugins/Animation/MLDeformer/MLDeformerFramework/Source/MLDeformerFrameworkEditor/Private/SMLDeformerInputBonesWidget.cpp
@@ -58,11 +58,6 @@ namespace UE::MLDeformer
 			{
 				Menu.AddMenuEntry(Actions.DeleteInputBones);
 			}
-
-			if (!GetRootItems().IsEmpty())
-			{
-				Menu.AddMenuEntry(Actions.ClearInputBones);
-			}
 		}
 		Menu.EndSection();
 
diff --git a/Engine/Plugins/Animation/MLDeformer/MLDeformerFramework/Source/MLDeformerFrameworkEditor/Private/SMLDeformerInputCurvesWidget.cpp b/Engine/Plugins/Animation/MLDeformer/MLDeformerFramework/Source/MLDeformerFrameworkEditor/Private/SMLDeformerInputCurvesWidget.cpp
--- a/Engine/Plugins/Animation/MLDeformer/MLDeformerFramework/Source/MLDeformerFrameworkEditor/Private/SMLDeformerInputCurvesWidget.cpp
+++ b/Engine/Plugins/Animation/MLDeformer/MLDeformerFramework/Source/MLDeformerFrameworkEditor/Private/SMLDeformerInputCurvesWidget.cpp
@@ -53,11 +53,6 @@ namespace UE::MLDeformer
 			{
 				Menu.AddMenuEntry(Actions.DeleteInputCurves);
 			}
-
-			if (!GetElements().IsEmpty())
-			{
-				Menu.AddMenuEntry(Actions.ClearInputCurves);
-			}
 		}
 		Menu.EndSection();
 
diff --git a/Engine/Plugins/Animation/MLDeformer/MLDeformerFramework/Source/MLDeformerFrameworkEditor/Public/MLDeformerMorphModelEditorModel.h b/Engine/Plugins/Animation/MLDeformer/MLDeformerFramework/Source/MLDeformerFrameworkEditor/Public/MLDeformerMorphModelEditorModel.h
--- a/Engine/Plugins/Animation/MLDeformer/MLDeformerFramework/Source/MLDeformerFrameworkEditor/Public/MLDeformerMorphModelEditorModel.h
+++ b/Engine/Plugins/Animation/MLDeformer/MLDeformerFramework/Source/MLDeformerFrameworkEditor/Public/MLDeformerMorphModelEditorModel.h
@@ -76,6 +76,9 @@ namespace UE::MLDeformer
 		 */
 		void ApplyBoneToMask(int32 SkeletonBoneIndex, TArrayView<float> MaskBuffer);
 
+		UE_DEPRECATED(5.4, "Please use the RecursiveAddBoneToMaskUpwards that doesn't take the VirtualParentTable as parameter.")
+		void RecursiveAddBoneToMaskUpwards(const FReferenceSkeleton& RefSkel, int32 SkeletonBoneIndex, int32 MaxHierarchyDepth, const TArray<int32>& VirtualParentTable, TArray<int32>& OutBonesAdded, int32 CurHierarchyDepth = 0);
+
 		/**
 		 * Recursively add bone masks for this bone and its parents.
 		 * The number of parents is controlled by the MaxHierarchyDepth.
@@ -83,10 +86,12 @@ namespace UE::MLDeformer
 		 * @param SkeletonBoneIndex The bone index of our current bone to add to the mask.
 		 * @param MaxHierarchyDepth The maximum hierarchy depth. If this is 1, it means that we visit only up to the parent of the specified bone. If it is 2, we also visit the parent of the parent, etc.
 		 * @param OutBoneAdded The list of bones that we already added to the mask. Bones that are in this list will not be added to the mask again.
-		 * @param VirtualParentTable The table of virtual parents. This lets us know that we also have to add say the finger bone vertices to the mask of the hand in our example described above.
 		 * @param CurHierarchyDepth The current recursive hierarchy depth. This is used to track how deep we are when moving up the hierarchy. A value of 1 means we are currently adding the parent of the initial bone.
 		 */
-		void RecursiveAddBoneToMaskUpwards(const FReferenceSkeleton& RefSkel, int32 SkeletonBoneIndex, int32 MaxHierarchyDepth, const TArray<int32>& VirtualParentTable, TArray<int32>& OutBonesAdded, int32 CurHierarchyDepth = 0);
+		void RecursiveAddBoneToMaskUpwards(const FReferenceSkeleton& RefSkel, int32 SkeletonBoneIndex, int32 MaxHierarchyDepth, TArray<int32>& OutBonesAdded, int32 CurHierarchyDepth = 0);
+
+		UE_DEPRECATED(5.4, "Please use the RecursiveAddBoneToMaskDownwards that doesn't take the VirtualParentTable as parameter.")
+		void RecursiveAddBoneToMaskDownwards(const FReferenceSkeleton& RefSkel, int32 SkeletonBoneIndex, int32 MaxHierarchyDepth, const TArray<int32>& VirtualParentTable, TArray<int32>& OutBonesAdded, int32 CurHierarchyDepth = 0);
 
 		/**
 		 * Recursively add bone masks for this bone and its children.
@@ -94,11 +99,10 @@ namespace UE::MLDeformer
 		 * @param RefSkel The reference skeleton.
 		 * @param SkeletonBoneIndex The bone index of our current bone to add to the mask.
 		 * @param MaxHierarchyDepth The maximum hierarchy depth. If this is 1, it means that we visit only down to the children of the specified bone. If it is 2, we also visit the children of the children, etc.
-		 * @param VirtualParentTable The table of virtual parents. This lets us know that we also have to add say the finger bone vertices to the mask of the hand in our example described above.
 		 * @param OutBoneAdded The list of bones that we already added to the mask. Bones that are in this list will not be added to the mask again.
 		 * @param CurHierarchyDepth The current recursive hierarchy depth. This is used to track how deep we are when moving down the hierarchy. A value of 1 means we are currently adding a child node of the initial bone.
 		 */
-		void RecursiveAddBoneToMaskDownwards(const FReferenceSkeleton& RefSkel, int32 SkeletonBoneIndex, int32 MaxHierarchyDepth, const TArray<int32>& VirtualParentTable, TArray<int32>& OutBonesAdded, int32 CurHierarchyDepth = 0);
+		void RecursiveAddBoneToMaskDownwards(const FReferenceSkeleton& RefSkel, int32 SkeletonBoneIndex, int32 MaxHierarchyDepth, TArray<int32>& OutBonesAdded, int32 CurHierarchyDepth = 0);
 
 		/**
 		 * Add any required bone to the mask.
@@ -109,6 +113,7 @@ namespace UE::MLDeformer
 		 * @param VirtualParentTable The table of virtual parents. This lets us know that we also have to add say the finger bone vertices to the mask of the hand in our example described above.
 		 * @param OutBonesAdded The list of bones we added to the mask. This list can contain items already. Bones that are already in the list won't be added to the mask again.
 		 */
+		UE_DEPRECATED(5.4, "This method will be removed.")
 		void AddRequiredBones(const FReferenceSkeleton& RefSkel, int32 SkeletonBoneIndex, const TArray<int32>& VirtualParentTable, TArray<int32>& OutBonesAdded);
 
 		/**
@@ -120,6 +125,7 @@ namespace UE::MLDeformer
 		 * @return The bone index that is the virtual parent for the specified bone. This can point to itself in case it is in the included bone names list.
 		 * @see BuildVirtualParentTable.
 		 */
+		UE_DEPRECATED(5.4, "This method will be removed.")
 		int32 FindVirtualParentIndex(const FReferenceSkeleton& RefSkel, int32 BoneIndex, const TArray<FName>& IncludedBoneNames) const;
 
 		/**
@@ -132,6 +138,7 @@ namespace UE::MLDeformer
 		 * @param IncludedBoneNames The names of the bones that are inputs to the ML model.
 		 * @return Returns a mapping array with the length of number of bones in the reference skeleton. Each array element contains a bone index inside the reference skeleton as well.
 		 */
+		UE_DEPRECATED(5.4, "This method will be removed.")
 		TArray<int32> BuildVirtualParentTable(const FReferenceSkeleton& RefSkel, const TArray<FName>& IncludedBoneNames) const;
 
 		// Helpers.
diff --git a/Engine/Plugins/Animation/MLDeformer/MLDeformerFramework/Source/MLDeformerFrameworkEditor/Public/SMLDeformerBonePickerDialog.h b/Engine/Plugins/Animation/MLDeformer/MLDeformerFramework/Source/MLDeformerFrameworkEditor/Public/SMLDeformerBonePickerDialog.h
--- a/Engine/Plugins/Animation/MLDeformer/MLDeformerFramework/Source/MLDeformerFrameworkEditor/Public/SMLDeformerBonePickerDialog.h
+++ b/Engine/Plugins/Animation/MLDeformer/MLDeformerFramework/Source/MLDeformerFrameworkEditor/Public/SMLDeformerBonePickerDialog.h
@@ -84,6 +84,7 @@ namespace UE::MLDeformer
 		SLATE_BEGIN_ARGS(SMLDeformerBonePickerDialog) {}
 		SLATE_ARGUMENT(FReferenceSkeleton*, RefSkeleton)		// The reference skeleton to grab the bone names and hierarchy from.
 		SLATE_ARGUMENT(TArray<FName>, HighlightBoneNames)		// Highlight bones in this list.
+		SLATE_ARGUMENT(TArray<FName>, InitialSelectedBoneNames)	// Select all bones in this list.
 		SLATE_ARGUMENT(FSlateColor, HighlightBoneNamesColor)	// Bones inside the HighlightBoneNames list will use this color.
 		SLATE_ARGUMENT(TArray<FName>, IncludeList)				// If empty, all bones can be picked. If not empty, only show bones with names inside this list.
 		SLATE_ARGUMENT(FText, Title)							// The window title.
@@ -107,11 +108,14 @@ namespace UE::MLDeformer
 		void OnCancelClicked();
 		void OnFilterTextChanged(const FText& InFilterText);
 		void RefreshBoneTreeElements();
+		void SelectInitialItems();
+		void SelectInitialItemsRecursive(const TSharedPtr<FMLDeformerBonePickerTreeElement>& Item);
 
 	private:
 		TSharedPtr<SMLDeformerBonePickerTreeWidget> BoneTreeWidget;
 		TSharedPtr<SWidget> ExtraWidget;
 		TArray<FName> HighlightBoneNames;
+		TArray<FName> InitialSelectedBoneNames;
 		TArray<FName> IncludeList;
 		FReferenceSkeleton* RefSkeleton = nullptr;
 		TArray<FName> PickedBoneNames;
diff --git a/Engine/Plugins/Animation/MLDeformer/NeuralMorphModel/Source/NeuralMorphModel/Private/NeuralMorphModel.cpp b/Engine/Plugins/Animation/MLDeformer/NeuralMorphModel/Source/NeuralMorphModel/Private/NeuralMorphModel.cpp
--- a/Engine/Plugins/Animation/MLDeformer/NeuralMorphModel/Source/NeuralMorphModel/Private/NeuralMorphModel.cpp
+++ b/Engine/Plugins/Animation/MLDeformer/NeuralMorphModel/Source/NeuralMorphModel/Private/NeuralMorphModel.cpp
@@ -9,6 +9,7 @@
 #include "MLDeformerObjectVersion.h"
 #include "UObject/AssetRegistryTagsContext.h"
 #include "UObject/UObjectGlobals.h"
+#include "Engine/SkeletalMesh.h"
 
 #include UE_INLINE_GENERATED_CPP_BY_NAME(NeuralMorphModel)
 
@@ -50,8 +51,101 @@ UMLDeformerInputInfo* UNeuralMorphModel::CreateInputInfo()
 	return NewObject<UNeuralMorphInputInfo>(this);
 }
 
+
+namespace
+{
+	int32 FindVirtualParentIndex(const FReferenceSkeleton& RefSkel, int32 BoneIndex, const TArray<FName>& IncludedBoneNames)
+	{
+		int32 CurBoneIndex = BoneIndex;
+		while (CurBoneIndex != INDEX_NONE)
+		{
+			const int32 ParentIndex = RefSkel.GetParentIndex(CurBoneIndex);
+			if (ParentIndex == INDEX_NONE)
+			{
+				break;
+			}
+
+			const FName ParentName = RefSkel.GetBoneName(ParentIndex);
+			if (IncludedBoneNames.Contains(ParentName))
+			{
+				return ParentIndex;
+			}
+
+			CurBoneIndex = ParentIndex;
+		}
+
+		return BoneIndex;
+	}
+
+	void ModifyMaskInfosForBackwardCompatibility(const FReferenceSkeleton& RefSkel, const TArray<FBoneReference>& BoneIncludeList, TMap<FName, FNeuralMorphMaskInfo>& MaskInfos)
+	{
+		// If we are not using some masks, just skip it.
+		if (MaskInfos.IsEmpty())
+		{
+			return;
+		}
+
+		// Build a list of bone names that we have setup as inputs.
+		TArray<FName> InputBoneNames;
+		for (const FBoneReference& BoneReference : BoneIncludeList)
+		{
+			InputBoneNames.Add(BoneReference.BoneName);
+		}
+
+		// Build a virtual parent table.
+		// This basically finds the first parent bone up the chain that is also in the bone input list.
+		// So if the input list contains upper and lower arm, then the tip of the index finger's virtual parent will be the lower arm bone
+		// even if there is a hand bone and more finger bones up the chain.
+		TArray<int32> VirtualParentTable;
+		VirtualParentTable.SetNumUninitialized(RefSkel.GetNum());
+		for (int32 BoneIndex = 0; BoneIndex < RefSkel.GetNum(); ++BoneIndex)
+		{
+			VirtualParentTable[BoneIndex] = FindVirtualParentIndex(RefSkel, BoneIndex, InputBoneNames);
+		}
+
+		// For all bone masks.
+		TArray<FName> RequiredBones;
+		for (auto& MaskInfoEntry : MaskInfos)
+		{
+			FNeuralMorphMaskInfo& MaskInfo = MaskInfoEntry.Value;
+			RequiredBones.Reset();
+			for (const FName MaskBoneName : MaskInfo.BoneNames)
+			{
+				const int32 MaskBoneIndex = RefSkel.FindBoneIndex(MaskBoneName);
+				if (MaskBoneIndex == INDEX_NONE)
+				{
+					UE_LOG(LogNeuralMorphModel, Warning, TEXT("Mask contains a bone named '%s' that isn't inside the skeleton."), *MaskBoneName.ToString());
+					continue;
+				}
+	
+				// Find out all required bones.
+				// This adds all child bones that have the current bone as (virtual) parent.
+				if (MaskBoneIndex != INDEX_NONE)
+				{
+					for (int32 Index = 0; Index < VirtualParentTable.Num(); ++Index)
+					{
+						const int32 VirtualParent = VirtualParentTable[Index];
+						const FName CurBoneName = RefSkel.GetBoneName(Index);
+						if (VirtualParent == MaskBoneIndex && !MaskInfo.BoneNames.Contains(CurBoneName))
+						{
+							RequiredBones.AddUnique(CurBoneName);
+						}
+					}
+				}
+			}
+
+			// Add the required bones.
+			for (const FName BoneName : RequiredBones)
+			{
+				MaskInfo.BoneNames.AddUnique(BoneName);
+			}
+		}
+	}
+} // Anonymous namespace.
+
 void UNeuralMorphModel::Serialize(FArchive& Archive)
 {
+	Super::Serialize(Archive);
 	Archive.UsingCustomVersion(UE::MLDeformer::FMLDeformerObjectVersion::GUID);
 
 	if (Archive.IsSaving())
@@ -82,9 +176,24 @@ void UNeuralMorphModel::Serialize(FArchive& Archive)
 			NeuralMorphInputInfo->CopyMembersFrom(CurInputInfo);
 			SetInputInfo(NeuralMorphInputInfo);
 		}		
-	}
+ 	}
 
-	Super::Serialize(Archive);
+	// Since 5.4 (when LOD was added), we modified the way the bone masks are generated.
+	// In UE 5.3 we added some extra bones to the mask. To keep the same behavior, we will upgrade the mask to include those bones again to ensure we get the same results when training.
+	// Newer assets will not do this.
+	#if WITH_EDITORONLY_DATA
+		if (Archive.IsLoading() && 
+			Archive.CustomVer(UE::MLDeformer::FMLDeformerObjectVersion::GUID) < UE::MLDeformer::FMLDeformerObjectVersion::LODSupportAdded)
+		{
+			USkeletalMesh* SkelMesh = GetSkeletalMesh();
+			if (SkelMesh)
+			{
+				const FReferenceSkeleton& RefSkeleton = SkelMesh->GetRefSkeleton();
+				ModifyMaskInfosForBackwardCompatibility(RefSkeleton, GetBoneIncludeList(), BoneMaskInfos);
+				ModifyMaskInfosForBackwardCompatibility(RefSkeleton, GetBoneIncludeList(), BoneGroupMaskInfos);
+			}
+		}
+	#endif
 }
 
 void UNeuralMorphModel::UpdateMissingGroupNames()
diff --git a/Engine/Plugins/Animation/MLDeformer/NeuralMorphModel/Source/NeuralMorphModelEditor/NeuralMorphModelEditor.Build.cs b/Engine/Plugins/Animation/MLDeformer/NeuralMorphModel/Source/NeuralMorphModelEditor/NeuralMorphModelEditor.Build.cs
--- a/Engine/Plugins/Animation/MLDeformer/NeuralMorphModel/Source/NeuralMorphModelEditor/NeuralMorphModelEditor.Build.cs
+++ b/Engine/Plugins/Animation/MLDeformer/NeuralMorphModel/Source/NeuralMorphModelEditor/NeuralMorphModelEditor.Build.cs
@@ -28,7 +28,8 @@ namespace UnrealBuildTool.Rules
 					"MLDeformerFrameworkEditor",
 					"NeuralMorphModel",
 					"PropertyEditor",
-					"ToolWidgets"
+					"ToolWidgets",
+					"DeveloperSettings"
 				}
 			);
 
diff --git a/Engine/Plugins/Animation/MLDeformer/NeuralMorphModel/Source/NeuralMorphModelEditor/Private/NeuralMorphEditorModel.cpp b/Engine/Plugins/Animation/MLDeformer/NeuralMorphModel/Source/NeuralMorphModelEditor/Private/NeuralMorphEditorModel.cpp
--- a/Engine/Plugins/Animation/MLDeformer/NeuralMorphModel/Source/NeuralMorphModelEditor/Private/NeuralMorphEditorModel.cpp
+++ b/Engine/Plugins/Animation/MLDeformer/NeuralMorphModel/Source/NeuralMorphModelEditor/Private/NeuralMorphEditorModel.cpp
@@ -7,6 +7,7 @@
 #include "NeuralMorphInputInfo.h"
 #include "NeuralMorphTrainingModel.h"
 #include "NeuralMorphModelVizSettings.h"
+#include "NeuralMorphEditorProjectSettings.h"
 #include "MLDeformerMorphModelVizSettings.h"
 #include "MLDeformerEditorToolkit.h"
 #include "MLDeformerEditorActor.h"
@@ -287,7 +288,6 @@ namespace UE::NeuralMorphModel
 		const FReferenceSkeleton& RefSkel = SkeletalMesh->GetRefSkeleton();
 
 		// Apply the bones to the mask buffer.
-		const TArray<int32> VirtualParentTable = BuildVirtualParentTable(RefSkel, GetEditorInputInfo()->GetBoneNames());
 		for (const FName MaskBoneName : MaskInfo.BoneNames)
 		{
 			const int32 MaskBoneIndex = RefSkel.FindBoneIndex(MaskBoneName);
@@ -301,23 +301,6 @@ namespace UE::NeuralMorphModel
 					*MaskBoneName.ToString(),
 					*SkeletalMesh->GetName());
 			}
-
-			// Add required bones
-			TArray<int32> RequiredBones;
-			const int32 SkeletonBoneIndex = RefSkel.FindBoneIndex(MaskBoneName);
-			if (SkeletonBoneIndex != INDEX_NONE)
-			{
-				AddRequiredBones(RefSkel, SkeletonBoneIndex, VirtualParentTable, RequiredBones);
-			}
-			for (const int32 RequiredBoneIndex: RequiredBones)
-			{
-				check(RequiredBoneIndex != INDEX_NONE);
-				const FName RequiredBoneName = RefSkel.GetBoneName(RequiredBoneIndex);
-				if (!MaskInfo.BoneNames.Contains(RequiredBoneName))
-				{
-					ApplyBoneToMask(RequiredBoneIndex, ItemMaskBuffer);
-				}
-			}
 		}
 	}
 
@@ -475,6 +458,37 @@ namespace UE::NeuralMorphModel
 		GetNeuralMorphModel()->BoneGroupMaskInfos.Empty();
 	}
 
+	void FNeuralMorphEditorModel::AddTwistBones(const FReferenceSkeleton& RefSkel, TArray<int32>& SkelBoneIndices)
+	{
+		check(UNeuralMorphEditorProjectSettings::Get());
+
+		// Get the twist substring from the per project configuration.
+		// If we have an empty string, just skip the whole twist bone handling.
+		const FString& TwistSubString = UNeuralMorphEditorProjectSettings::Get()->TwistBoneFilter;
+		if (TwistSubString.IsEmpty())
+		{
+			return;
+		}
+
+		// Add child bones with "twist" in their name.
+		TArray<int32> TwistBones;
+		for (const int32 AddedBoneIndex : SkelBoneIndices)
+		{
+			for (int32 Index = 0; Index < RefSkel.GetNum(); ++Index)
+			{
+				if (RefSkel.GetParentIndex(Index) == AddedBoneIndex &&
+					RefSkel.GetBoneName(Index).ToString().Contains(TwistSubString))
+				{
+					TwistBones.AddUnique(Index);
+				}
+			}
+		}
+		for (const int32 TwistBoneIndex : TwistBones)
+		{
+			SkelBoneIndices.AddUnique(TwistBoneIndex);
+		}
+	}
+
 	void FNeuralMorphEditorModel::GenerateBoneMaskInfo(int32 InputInfoBoneIndex, int32 HierarchyDepth)
 	{
 		UNeuralMorphModel* NeuralMorphModel = GetNeuralMorphModel();
@@ -482,7 +496,6 @@ namespace UE::NeuralMorphModel
 		UNeuralMorphInputInfo* NeuralMorphInputInfo = Cast<UNeuralMorphInputInfo>(GetEditorInputInfo());
 		const FReferenceSkeleton& RefSkel = SkeletalMesh->GetRefSkeleton();
 		const TArray<FName>& BoneNames = NeuralMorphInputInfo->GetBoneNames();
-		const TArray<int32> VirtualParentTable = BuildVirtualParentTable(RefSkel, BoneNames);
 		const int32 NumBones = NeuralMorphInputInfo->GetNumBones();
 
 		// Make sure we have a valid bone name.
@@ -502,8 +515,9 @@ namespace UE::NeuralMorphModel
 
 		// Build the list of bones that are required for the mask.
 		TArray<int32> BonesAdded;
-		RecursiveAddBoneToMaskUpwards(RefSkel, SkeletonBoneIndex, HierarchyDepth, VirtualParentTable, BonesAdded);
-		RecursiveAddBoneToMaskDownwards(RefSkel, SkeletonBoneIndex, HierarchyDepth, VirtualParentTable, BonesAdded);
+		RecursiveAddBoneToMaskUpwards(RefSkel, SkeletonBoneIndex, HierarchyDepth, BonesAdded);
+		RecursiveAddBoneToMaskDownwards(RefSkel, SkeletonBoneIndex, HierarchyDepth, BonesAdded);
+		AddTwistBones(RefSkel, BonesAdded);
 
 		// Now that we know which bones we need, add them to the mask.
 		FNeuralMorphMaskInfo* MaskInfo = NeuralMorphModel->BoneMaskInfos.Find(BoneName);
@@ -526,7 +540,6 @@ namespace UE::NeuralMorphModel
 		UNeuralMorphInputInfo* NeuralMorphInputInfo = Cast<UNeuralMorphInputInfo>(GetEditorInputInfo());
 		const FReferenceSkeleton& RefSkel = SkeletalMesh->GetRefSkeleton();
 		const TArray<FName>& BoneNames = NeuralMorphInputInfo->GetBoneNames();
-		const TArray<int32> VirtualParentTable = BuildVirtualParentTable(RefSkel, BoneNames);
 
 		// For all bones.
 		TArray<int32> BonesAdded;
@@ -550,8 +563,9 @@ namespace UE::NeuralMorphModel
 			}
 
 			// Add all bones up and down the chain.
-			RecursiveAddBoneToMaskUpwards(RefSkel, SkeletonBoneIndex, HierarchyDepth, VirtualParentTable, BonesAdded);
-			RecursiveAddBoneToMaskDownwards(RefSkel, SkeletonBoneIndex, HierarchyDepth, VirtualParentTable, BonesAdded);
+			RecursiveAddBoneToMaskUpwards(RefSkel, SkeletonBoneIndex, HierarchyDepth, BonesAdded);
+			RecursiveAddBoneToMaskDownwards(RefSkel, SkeletonBoneIndex, HierarchyDepth, BonesAdded);
+			AddTwistBones(RefSkel, BonesAdded);
 		}
 
 		// Now that we know which bones we need, add them to the mask.
diff --git a/Engine/Plugins/Animation/MLDeformer/NeuralMorphModel/Source/NeuralMorphModelEditor/Private/NeuralMorphEditorProjectSettings.cpp b/Engine/Plugins/Animation/MLDeformer/NeuralMorphModel/Source/NeuralMorphModelEditor/Private/NeuralMorphEditorProjectSettings.cpp
new file mode 100644
--- /dev/null
+++ b/Engine/Plugins/Animation/MLDeformer/NeuralMorphModel/Source/NeuralMorphModelEditor/Private/NeuralMorphEditorProjectSettings.cpp
@@ -0,0 +1,24 @@
+// Copyright Epic Games, Inc. All Rights Reserved.
+
+#include "NeuralMorphEditorProjectSettings.h"
+
+#define LOCTEXT_NAMESPACE "NeuralMorphEditorProjectSettings"
+
+const UNeuralMorphEditorProjectSettings* UNeuralMorphEditorProjectSettings::Get()
+{ 
+	return GetDefault<UNeuralMorphEditorProjectSettings>();
+}
+
+#if WITH_EDITOR
+	FText UNeuralMorphEditorProjectSettings::GetSectionText() const 
+	{ 
+		return LOCTEXT("NeuralMorphModelProjectSettingsName", "Neural Morph Model"); 
+	}
+
+	FText UNeuralMorphEditorProjectSettings::GetSectionDescription() const
+	{
+		return LOCTEXT("NeuralMorphModelProjectSettingsDescription", "Configure project wide settings related to the Neural Morph Model, which is part of the ML Deformer.");
+	}
+#endif
+
+#undef LOCTEXT_NAMESPACE
diff --git a/Engine/Plugins/Animation/MLDeformer/NeuralMorphModel/Source/NeuralMorphModelEditor/Private/SNeuralMorphBoneGroupsWidget.cpp b/Engine/Plugins/Animation/MLDeformer/NeuralMorphModel/Source/NeuralMorphModelEditor/Private/SNeuralMorphBoneGroupsWidget.cpp
--- a/Engine/Plugins/Animation/MLDeformer/NeuralMorphModel/Source/NeuralMorphModelEditor/Private/SNeuralMorphBoneGroupsWidget.cpp
+++ b/Engine/Plugins/Animation/MLDeformer/NeuralMorphModel/Source/NeuralMorphModelEditor/Private/SNeuralMorphBoneGroupsWidget.cpp
@@ -88,11 +88,6 @@ namespace UE::NeuralMorphModel
 			{			
 				Menu.AddMenuEntry(Actions.DeleteSelectedItems);
 			}
-
-			if (!GetRootItems().IsEmpty())
-			{
-				Menu.AddMenuEntry(Actions.ClearGroups);
-			}
 		}
 		Menu.EndSection();
 
@@ -279,42 +274,8 @@ namespace UE::NeuralMorphModel
 				BoneGroup.BoneNames.AddUnique(BoneName);
 			}
 
-			USkeletalMesh* SkeletalMesh = NeuralMorphModel->GetSkeletalMesh();	
-			const FReferenceSkeleton& RefSkel = SkeletalMesh->GetRefSkeleton();
-			const TArray<FName>& EditorInfoBoneNames = InputInfo->GetBoneNames();
-			const TArray<int32> VirtualParentTable = EditorModel->BuildVirtualParentTable(RefSkel, EditorInfoBoneNames);	
 			const int32 HierarchyDepth = StaticCastSharedPtr<SNeuralMorphInputWidget>(InputWidget)->GetHierarchyDepth();
-
-			// Build a list of bone indices that need to be added.
-			// These can be more than just the picked bones, as we grow the mask based on a hierarchy depth as well.
-			TArray<int32> NewMaskBones;
-			for (const FName BoneName : BoneNames)
-			{
-				const int32 SkeletonBoneIndex = RefSkel.FindBoneIndex(BoneName);
-				if (SkeletonBoneIndex == INDEX_NONE)
-				{
-					UE_LOG(LogNeuralMorphModel, Warning, TEXT("Bone '%s' cannot be found in the SkeletalMesh '%s', ignoring the bone inside the group's mask."), *BoneName.ToString(), *SkeletalMesh->GetName());
-					continue;
-				}
-
-				EditorModel->RecursiveAddBoneToMaskUpwards(RefSkel, SkeletonBoneIndex, HierarchyDepth, VirtualParentTable, NewMaskBones);
-				EditorModel->RecursiveAddBoneToMaskDownwards(RefSkel, SkeletonBoneIndex, HierarchyDepth, VirtualParentTable, NewMaskBones);
-				EditorModel->AddRequiredBones(RefSkel, SkeletonBoneIndex, VirtualParentTable, NewMaskBones);
-			}
-
-			// Add the new bones to the mask.
-			FNeuralMorphMaskInfo* MaskInfo = NeuralMorphModel->BoneGroupMaskInfos.Find(BoneGroup.GroupName);
-			if (MaskInfo != nullptr)
-			{
-				for (const int32 BoneIndex : NewMaskBones)
-				{
-					const FName BoneName = RefSkel.GetBoneName(BoneIndex);
-					if (BoneName.IsValid() && !BoneName.IsNone())
-					{
-						MaskInfo->BoneNames.AddUnique(BoneName);
-					}
-				}
-			}
+			EditorModel->GenerateBoneGroupMaskInfo(GroupIndex, HierarchyDepth);
 
 			RefreshTree(true);
 			EditorModel->RebuildEditorMaskInfo();
diff --git a/Engine/Plugins/Animation/MLDeformer/NeuralMorphModel/Source/NeuralMorphModelEditor/Private/SNeuralMorphCurveGroupsWidget.cpp b/Engine/Plugins/Animation/MLDeformer/NeuralMorphModel/Source/NeuralMorphModelEditor/Private/SNeuralMorphCurveGroupsWidget.cpp
--- a/Engine/Plugins/Animation/MLDeformer/NeuralMorphModel/Source/NeuralMorphModelEditor/Private/SNeuralMorphCurveGroupsWidget.cpp
+++ b/Engine/Plugins/Animation/MLDeformer/NeuralMorphModel/Source/NeuralMorphModelEditor/Private/SNeuralMorphCurveGroupsWidget.cpp
@@ -88,11 +88,6 @@ namespace UE::NeuralMorphModel
 			{			
 				Menu.AddMenuEntry(Actions.DeleteSelectedItems);
 			}
-
-			if (!GetRootItems().IsEmpty())
-			{
-				Menu.AddMenuEntry(Actions.ClearGroups);
-			}
 		Menu.EndSection();
 
 		return Menu.MakeWidget();
diff --git a/Engine/Plugins/Animation/MLDeformer/NeuralMorphModel/Source/NeuralMorphModelEditor/Private/SNeuralMorphInputWidget.cpp b/Engine/Plugins/Animation/MLDeformer/NeuralMorphModel/Source/NeuralMorphModelEditor/Private/SNeuralMorphInputWidget.cpp
--- a/Engine/Plugins/Animation/MLDeformer/NeuralMorphModel/Source/NeuralMorphModelEditor/Private/SNeuralMorphInputWidget.cpp
+++ b/Engine/Plugins/Animation/MLDeformer/NeuralMorphModel/Source/NeuralMorphModelEditor/Private/SNeuralMorphInputWidget.cpp
@@ -46,12 +46,12 @@ namespace UE::NeuralMorphModel
 		// Bone commands.
 		UI_COMMAND(ResetAllBoneMasks, "Reset All Bone Masks", "Reset all masks for every bone in the input list.", EUserInterfaceActionType::Button, FInputChord());
 		UI_COMMAND(ResetSelectedBoneMasks, "Reset Selected Bone Masks", "Reset the masks for all selected bones.", EUserInterfaceActionType::Button, FInputChord());
-		UI_COMMAND(ExpandSelectedBoneMasks, "Add Bones To Mask", "Add specific bones to the masks for the selected bones.", EUserInterfaceActionType::Button, FInputChord());
+		UI_COMMAND(ExpandSelectedBoneMasks, "Edit Bone Mask", "Specify which bones to include inside the mask.", EUserInterfaceActionType::Button, FInputChord());
 
 		// Bone group commands.
 		UI_COMMAND(ResetAllBoneGroupMasks, "Reset All Bone Group Masks", "Reset all masks for every bone group.", EUserInterfaceActionType::Button, FInputChord());
 		UI_COMMAND(ResetSelectedBoneGroupMasks, "Reset Selected Bone Group Masks", "Reset the masks for all selected bone groups.", EUserInterfaceActionType::Button, FInputChord());
-		UI_COMMAND(ExpandSelectedBoneGroupMasks, "Add Bones To Group Mask", "Add specific bones to the masks for the selected bone groups.", EUserInterfaceActionType::Button, FInputChord());
+		UI_COMMAND(ExpandSelectedBoneGroupMasks, "Edit Group Mask", "Specify which bones to include inside the mask for this group.", EUserInterfaceActionType::Button, FInputChord());
 	}
 
 	void SNeuralMorphInputWidget::BindCommands()
@@ -197,7 +197,8 @@ namespace UE::NeuralMorphModel
 			.RefSkeleton(&SkelMesh->GetRefSkeleton())
 			.AllowMultiSelect(true)
 			.HighlightBoneNamesColor(FSlateColor(HighlightColor))
-			.HighlightBoneNames(HighlightedBones);
+			.HighlightBoneNames(HighlightedBones)
+			.InitialSelectedBoneNames(HighlightedBones);
 
 		Dialog->ShowModal();
 
@@ -221,6 +222,7 @@ namespace UE::NeuralMorphModel
 				}
 
 				check(MaskInfo);
+				MaskInfo->BoneNames.Reset();
 				for (const FName PickedBoneName : PickedBoneNames)
 				{
 					MaskInfo->BoneNames.AddUnique(PickedBoneName);
@@ -334,7 +336,8 @@ namespace UE::NeuralMorphModel
 			.RefSkeleton(&SkelMesh->GetRefSkeleton())
 			.AllowMultiSelect(true)
 			.HighlightBoneNamesColor(FSlateColor(HighlightColor))
-			.HighlightBoneNames(HighlightedBones);
+			.HighlightBoneNames(HighlightedBones)
+			.InitialSelectedBoneNames(HighlightedBones);
 
 		Dialog->ShowModal();
 
@@ -359,6 +362,7 @@ namespace UE::NeuralMorphModel
 					MaskInfo = &NeuralEditorModel->GetNeuralMorphModel()->BoneGroupMaskInfos.Add(Item->Name, FNeuralMorphMaskInfo());
 				}
 				check(MaskInfo);
+				MaskInfo->BoneNames.Reset();
 				for (const FName PickedBoneName : PickedBoneNames)
 				{
 					MaskInfo->BoneNames.AddUnique(PickedBoneName);
@@ -386,7 +390,10 @@ namespace UE::NeuralMorphModel
 			const int32 NumSelectedItems = InputBonesWidget->GetTreeWidget()->GetNumItemsSelected();
 			if (NumSelectedItems > 0)
 			{
-				MenuBuilder.AddMenuEntry(Commands.ExpandSelectedBoneMasks);
+				if (NumSelectedItems == 1)
+				{
+					MenuBuilder.AddMenuEntry(Commands.ExpandSelectedBoneMasks);
+				}
 				MenuBuilder.AddMenuEntry(Commands.ResetSelectedBoneMasks);
 			}
 		}
@@ -461,7 +468,10 @@ namespace UE::NeuralMorphModel
 			const FNeuralMorphInputWidgetCommands& Commands = FNeuralMorphInputWidgetCommands::Get();
 			if (BoneGroupsWidget->GetNumSelectedGroups() > 0)
 			{
-				MenuBuilder.AddMenuEntry(Commands.ExpandSelectedBoneGroupMasks);
+				if (BoneGroupsWidget->GetNumSelectedGroups() == 1)
+				{
+					MenuBuilder.AddMenuEntry(Commands.ExpandSelectedBoneGroupMasks);
+				}
 				MenuBuilder.AddMenuEntry(Commands.ResetSelectedBoneGroupMasks);
 			}
 		}
diff --git a/Engine/Plugins/Animation/MLDeformer/NeuralMorphModel/Source/NeuralMorphModelEditor/Public/NeuralMorphEditorModel.h b/Engine/Plugins/Animation/MLDeformer/NeuralMorphModel/Source/NeuralMorphModelEditor/Public/NeuralMorphEditorModel.h
--- a/Engine/Plugins/Animation/MLDeformer/NeuralMorphModel/Source/NeuralMorphModelEditor/Public/NeuralMorphEditorModel.h
+++ b/Engine/Plugins/Animation/MLDeformer/NeuralMorphModel/Source/NeuralMorphModelEditor/Public/NeuralMorphEditorModel.h
@@ -160,6 +160,18 @@ namespace UE::NeuralMorphModel
 		 */
 		void DebugDrawItemMask(FPrimitiveDrawInterface* PDI, int32 MaskItemIndex, const FVector& DrawOffset);
 
+	private:
+		/**
+		 * Add child twist bones to the list of skeletal bone indices.
+		 * This will iterate over all bones inside the SkelBoneIndices array and check if there are direct child nodes inside the ref skeleton that contain
+		 * the substring "twist" inside their name. If so, that bone will be added to the SkelBoneIndices array as well.
+		 * The "twist" string is case-insensitive. The "twist" string can be configured on a per project basis inside the UNeuralMorphModel::TwistBoneSubString property, which 
+		 * is exposed in the per project .ini file. On default this value is "twist".
+		 * @param RefSkel The reference skeleton we use to find child nodes.
+		 * @param SkelBoneIndices The input and output array of indices inside the reference skeleton. This method potentially adds new entries to this array.
+		 */
+		void AddTwistBones(const FReferenceSkeleton& RefSkel, TArray<int32>& SkelBoneIndices);
+
 	protected:
 		int32 MaskVizItemIndex = INDEX_NONE;
 	};
diff --git a/Engine/Plugins/Animation/MLDeformer/NeuralMorphModel/Source/NeuralMorphModelEditor/Public/NeuralMorphEditorProjectSettings.h b/Engine/Plugins/Animation/MLDeformer/NeuralMorphModel/Source/NeuralMorphModelEditor/Public/NeuralMorphEditorProjectSettings.h
new file mode 100644
--- /dev/null
+++ b/Engine/Plugins/Animation/MLDeformer/NeuralMorphModel/Source/NeuralMorphModelEditor/Public/NeuralMorphEditorProjectSettings.h
@@ -0,0 +1,38 @@
+// Copyright Epic Games, Inc. All Rights Reserved.
+
+#pragma once
+
+#include "CoreTypes.h"
+#include "Engine/DeveloperSettings.h"
+#include "NeuralMorphEditorProjectSettings.generated.h"
+
+
+UCLASS(config = EditorPerProjectUserSettings)
+class NEURALMORPHMODELEDITOR_API UNeuralMorphEditorProjectSettings
+	: public UDeveloperSettings
+{
+	GENERATED_BODY()
+
+public:
+	static const UNeuralMorphEditorProjectSettings* Get();
+
+	// UDeveloperSettings overrides
+	virtual FName GetContainerName() const	{ return FName("Project"); }
+	virtual FName GetCategoryName() const	{ return FName("Plugins"); }
+	virtual FName GetSectionName() const	{ return FName("Neural Morph Model"); }
+#if WITH_EDITOR
+	virtual FText GetSectionText() const override;
+	virtual FText GetSectionDescription() const override;
+#endif
+	// ~END UDeveloperSettings overrides
+
+	/** 
+	 * The bone name filter that identifies twist bones.
+	 * Any bone name with this string inside its name will be seen as a twist bone.
+	 * This is only used during bone and bone group mask generation inside the Neural Morph Model.
+	 * Twist bones that are a child of the bones included in the bone mask are automatically included inside the mask generation as well.
+	 * We use this case-insensitive sub-string to identify twist bones. Set this to an empty string to disable automatically adding twist bones.
+	 */
+	UPROPERTY(config, EditAnywhere, Category = "Bone and Bone Group masking")
+	FString TwistBoneFilter = TEXT("twist");
+};
