[MLDeformer] Fixed a crash when removing an input animation that didn't get any animation asset set yet.

#rnx

diff --git a/Engine/Plugins/Animation/MLDeformer/MLDeformerFramework/Source/MLDeformerFrameworkEditor/Private/SMLDeformerTimeline.cpp b/Engine/Plugins/Animation/MLDeformer/MLDeformerFramework/Source/MLDeformerFrameworkEditor/Private/SMLDeformerTimeline.cpp
--- a/Engine/Plugins/Animation/MLDeformer/MLDeformerFramework/Source/MLDeformerFrameworkEditor/Private/SMLDeformerTimeline.cpp
+++ b/Engine/Plugins/Animation/MLDeformer/MLDeformerFramework/Source/MLDeformerFrameworkEditor/Private/SMLDeformerTimeline.cpp
@@ -1962,6 +1962,12 @@ namespace UE::MLDeformer
 
 	void SMLDeformerTimeline::OnDeletedTrainingInputAnim(int32 Index)
 	{
+		if (TrainingAnimNames.IsEmpty())
+		{
+			SelectedTrainingAnimNameIndex = INDEX_NONE;
+			return;
+		}
+
 		const int32 OldSelectedItemIndex = SelectedTrainingAnimNameIndex;
 		if (TrainingAnimNames[SelectedTrainingAnimNameIndex]->TrainingInputAnimIndex >= Index)
 		{
