[MLDeformer] Improved handling of delegate bindings, which could prevent crashes.

diff --git a/Engine/Plugins/Animation/MLDeformer/MLDeformerFramework/Source/MLDeformerFramework/Private/MLDeformerComponent.cpp b/Engine/Plugins/Animation/MLDeformer/MLDeformerFramework/Source/MLDeformerFramework/Private/MLDeformerComponent.cpp
--- a/Engine/Plugins/Animation/MLDeformer/MLDeformerFramework/Source/MLDeformerFramework/Private/MLDeformerComponent.cpp
+++ b/Engine/Plugins/Animation/MLDeformer/MLDeformerFramework/Source/MLDeformerFramework/Private/MLDeformerComponent.cpp
@@ -33,6 +33,8 @@ UMLDeformerComponent::UMLDeformerComponent(const FObjectInitializer& ObjectIniti
 
 void UMLDeformerComponent::Init()
 {
+	UnbindDelegates();
+
 	// If there is no deformer asset linked, release what we currently have.
 	if (!DeformerAsset)
 	{
@@ -81,9 +83,7 @@ void UMLDeformerComponent::SetupComponent(UMLDeformerAsset* InDeformerAsset, USk
 	SkelMeshComponent = InSkelMeshComponent;
 
 	// Initialize and make sure we have a model instance.
-	UnbindDelegates();
 	Init();
-	BindDelegates();
 
 	// Verify that a mesh deformer has been setup when the used ML model requires one.
 	if (!bSuppressMeshDeformerLogWarnings && DeformerAsset && ModelInstance && ModelInstance->GetModel() && SkelMeshComponent)
@@ -112,9 +112,9 @@ void UMLDeformerComponent::SetupComponent(UMLDeformerAsset* InDeformerAsset, USk
 void UMLDeformerComponent::BindDelegates()
 {
 	UMLDeformerModel* Model = DeformerAsset ? DeformerAsset->GetModel() : nullptr;
-	if (Model)
+	if (Model && !ReinitModelInstanceDelegateHandle.IsValid())
 	{
-		ReinitModelInstanceDelegateHandle = Model->GetReinitModelInstanceDelegate().AddLambda(
+		ReinitModelInstanceDelegateHandle = Model->GetReinitModelInstanceDelegate().AddLambda( 
 			[this]()
 			{
 				Init();
@@ -128,9 +128,8 @@ void UMLDeformerComponent::UnbindDelegates()
 	if (Model && ReinitModelInstanceDelegateHandle.IsValid())
 	{
 		Model->GetReinitModelInstanceDelegate().Remove(ReinitModelInstanceDelegateHandle);	
+		ReinitModelInstanceDelegateHandle = FDelegateHandle();
 	}
-
-	ReinitModelInstanceDelegateHandle = FDelegateHandle();
 }
 
 void UMLDeformerComponent::BeginDestroy()
@@ -266,6 +265,14 @@ void UMLDeformerComponent::SetDeformerAssetInternal(UMLDeformerAsset* const InDe
 }
 
 #if WITH_EDITOR
+	void UMLDeformerComponent::PreEditChange(FProperty* Property)
+	{
+		if (Property->GetFName() == GET_MEMBER_NAME_CHECKED(UMLDeformerComponent, DeformerAsset))
+		{
+			UnbindDelegates();
+		}
+	}
+
 	void UMLDeformerComponent::PostEditChangeProperty(FPropertyChangedEvent& PropertyChangedEvent)
 	{
 		const FProperty* Property = PropertyChangedEvent.Property;
diff --git a/Engine/Plugins/Animation/MLDeformer/MLDeformerFramework/Source/MLDeformerFramework/Public/MLDeformerComponent.h b/Engine/Plugins/Animation/MLDeformer/MLDeformerFramework/Source/MLDeformerFramework/Public/MLDeformerComponent.h
--- a/Engine/Plugins/Animation/MLDeformer/MLDeformerFramework/Source/MLDeformerFramework/Public/MLDeformerComponent.h
+++ b/Engine/Plugins/Animation/MLDeformer/MLDeformerFramework/Source/MLDeformerFramework/Public/MLDeformerComponent.h
@@ -32,6 +32,7 @@ public:
 	// UObject overrides.
 	void BeginDestroy() override;
 #if WITH_EDITOR
+	void PreEditChange(FProperty* Property) override;
 	void PostEditChangeProperty(FPropertyChangedEvent& PropertyChangedEvent) override;
 #endif
 	// ~END UObject overrides.
