[MLDeformer] Add an extra improvement to the handling of bindings in the ML Deformer Component. Also removed the quality level from the ML Deformer component as that isn't used anymore since we added LOD.

diff --git a/Engine/Plugins/Animation/MLDeformer/MLDeformerFramework/Source/MLDeformerFramework/Private/MLDeformerComponent.cpp b/Engine/Plugins/Animation/MLDeformer/MLDeformerFramework/Source/MLDeformerFramework/Private/MLDeformerComponent.cpp
--- a/Engine/Plugins/Animation/MLDeformer/MLDeformerFramework/Source/MLDeformerFramework/Private/MLDeformerComponent.cpp
+++ b/Engine/Plugins/Animation/MLDeformer/MLDeformerFramework/Source/MLDeformerFramework/Private/MLDeformerComponent.cpp
@@ -79,6 +79,8 @@ void UMLDeformerComponent::SetupComponent(UMLDeformerAsset* InDeformerAsset, USk
 		AddTickPrerequisiteComponent(InSkelMeshComponent);
 	}
 
+	UnbindDelegates();
+
 	DeformerAsset = InDeformerAsset;
 	SkelMeshComponent = InSkelMeshComponent;
 
@@ -120,6 +122,10 @@ void UMLDeformerComponent::BindDelegates()
 				Init();
 			});
 	}
+	else
+	{
+		ReinitModelInstanceDelegateHandle = FDelegateHandle();
+	}
 }
 
 void UMLDeformerComponent::UnbindDelegates()
diff --git a/Engine/Plugins/Animation/MLDeformer/MLDeformerFramework/Source/MLDeformerFramework/Public/MLDeformerComponent.h b/Engine/Plugins/Animation/MLDeformer/MLDeformerFramework/Source/MLDeformerFramework/Public/MLDeformerComponent.h
--- a/Engine/Plugins/Animation/MLDeformer/MLDeformerFramework/Source/MLDeformerFramework/Public/MLDeformerComponent.h
+++ b/Engine/Plugins/Animation/MLDeformer/MLDeformerFramework/Source/MLDeformerFramework/Public/MLDeformerComponent.h
@@ -75,8 +75,9 @@ public:
 	 * In morph based models each quality level defines how many morph targets are active at most.
 	 * @param InQualityLevel The quality level to switch to.
 	 */
-	UFUNCTION(BlueprintCallable, Category = "MLDeformer")
-	void SetQualityLevel(int32 InQualityLevel)					{ QualityLevel = FMath::Max<int32>(InQualityLevel, 0); }
+	UE_DEPRECATED(5.4, "This function will be removed.")
+	UFUNCTION(meta = (DeprecatedFunction, DeprecationMessage = "SetQualityLevel has been deprecated."))
+	void SetQualityLevel(int32 InQualityLevel)					{}
 
 	/** 
 	 * The quality level of the deformer. A value of 0 is the highest quality, 1 is a step lower, etc.
@@ -86,8 +87,9 @@ public:
 	 * In morph based models each quality level defines how many morph targets are active at most.
 	 * @param InQualityLevel The quality level to switch to.
 	 */
-	UFUNCTION(BlueprintCallable, Category = "MLDeformer")
-	int32 GetQualityLevel() const								{ return QualityLevel; }
+	UE_DEPRECATED(5.4, "This function will be removed.")
+	UFUNCTION(meta = (DeprecatedFunction, DeprecationMessage = "GetQualityLevel has been deprecated."))
+	int32 GetQualityLevel() const								{ return 0; }
 
 	/**
 	 * Get the ML Deformer asset that is used by this component.
@@ -162,7 +164,9 @@ public:
 	// Get property names.
 	static FName GetDeformerAssetPropertyName()					{ return GET_MEMBER_NAME_CHECKED(UMLDeformerComponent, DeformerAsset); }
 	static FName GetWeightPropertyName()						{ return GET_MEMBER_NAME_CHECKED(UMLDeformerComponent, Weight); }
-	static FName GetQualityLevelPropertyName()					{ return GET_MEMBER_NAME_CHECKED(UMLDeformerComponent, QualityLevel); }
+
+	UE_DEPRECATED(5.4, "This function will be removed.")
+	static FName GetQualityLevelPropertyName()					{ return GET_MEMBER_NAME_CHECKED(UMLDeformerComponent, QualityLevel_DEPRECATED); }
 
 protected:
 	// AActorComponent overrides.
@@ -246,8 +250,8 @@ protected:
 	 * level 3 will be used in this case, which represents the lowest available quality.
 	 * In morph based models each quality level defines how many morph targets are active at most.
 	 */
-	UPROPERTY(EditAnywhere, Category = "ML Deformer", meta = (ClampMin = "0"))
-	int32 QualityLevel = 0;
+	UPROPERTY()
+	int32 QualityLevel_DEPRECATED = 0;
 
 	/** The deformation model instance. This is used to perform the runtime updates and run the inference. */
 	UPROPERTY(Transient)
