[MLDeformer] Changed the framework and neural morph model from beta to production ready.

diff --git a/Engine/Plugins/Animation/MLDeformer/MLDeformerFramework/MLDeformerFramework.uplugin b/Engine/Plugins/Animation/MLDeformer/MLDeformerFramework/MLDeformerFramework.uplugin
--- a/Engine/Plugins/Animation/MLDeformer/MLDeformerFramework/MLDeformerFramework.uplugin
+++ b/Engine/Plugins/Animation/MLDeformer/MLDeformerFramework/MLDeformerFramework.uplugin
@@ -1,7 +1,7 @@
 {
 	"FileVersion": 3,
 	"Version": 1,
-	"VersionName": "0.2",
+	"VersionName": "1.0",
 	"FriendlyName": "ML Deformer Framework",
 	"Description": "Machine Learning Mesh Deformer Framework",
 	"Category": "Animation",
@@ -11,7 +11,7 @@
 	"MarketplaceURL": "",
 	"SupportURL": "https://forums.unrealengine.com/",
 	"CanContainContent": true,
-	"IsBetaVersion": true,
+	"IsBetaVersion": false,
 	"IsExperimentalVersion": false,
 	"Installed": false,
 	"Modules": [
diff --git a/Engine/Plugins/Animation/MLDeformer/NeuralMorphModel/NeuralMorphModel.uplugin b/Engine/Plugins/Animation/MLDeformer/NeuralMorphModel/NeuralMorphModel.uplugin
--- a/Engine/Plugins/Animation/MLDeformer/NeuralMorphModel/NeuralMorphModel.uplugin
+++ b/Engine/Plugins/Animation/MLDeformer/NeuralMorphModel/NeuralMorphModel.uplugin
@@ -1,7 +1,7 @@
 {
 	"FileVersion": 3,
 	"Version": 1,
-	"VersionName": "0.2",
+	"VersionName": "1.0",
 	"FriendlyName": "ML Deformer Neural Morph Model",
 	"Description": "Neural Morph Model for the ML Deformer Framework",
 	"Category": "Animation",
@@ -11,7 +11,7 @@
 	"MarketplaceURL": "",
 	"SupportURL": "https://forums.unrealengine.com/",
 	"CanContainContent": true,
-	"IsBetaVersion": true,
+	"IsBetaVersion": false,
 	"IsExperimentalVersion": false,
 	"Installed": false,
 	"Modules": [
