[MLDeformer] Fix some crash when pressing Undo in a specific case.

diff --git a/Engine/Plugins/Animation/MLDeformer/MLDeformerFramework/Source/MLDeformerFrameworkEditor/Private/MLDeformerMorphModelVizSettingsDetails.cpp b/Engine/Plugins/Animation/MLDeformer/MLDeformerFramework/Source/MLDeformerFrameworkEditor/Private/MLDeformerMorphModelVizSettingsDetails.cpp
--- a/Engine/Plugins/Animation/MLDeformer/MLDeformerFramework/Source/MLDeformerFrameworkEditor/Private/MLDeformerMorphModelVizSettingsDetails.cpp
+++ b/Engine/Plugins/Animation/MLDeformer/MLDeformerFramework/Source/MLDeformerFrameworkEditor/Private/MLDeformerMorphModelVizSettingsDetails.cpp
@@ -51,7 +51,7 @@ namespace UE::MLDeformer
 		if (MLDeformerComponent)
 		{
 			const UMLDeformerMorphModelInstance* MorphInstance = Cast<UMLDeformerMorphModelInstance>(MLDeformerComponent->GetModelInstance());
-			if (MorphInstance)
+			if (MorphInstance && MorphInstance->GetFinalSkeletalMeshComponent())
 			{
 				const int32 LOD = MorphInstance->GetFinalSkeletalMeshComponent()->GetPredictedLODLevel();
 				FExternalMorphSetWeights* WeightData = MorphInstance->FindWeightData(LOD);
@@ -98,7 +98,7 @@ namespace UE::MLDeformer
 							const UMLDeformerMorphModelInstance* MorphInstance = Cast<UMLDeformerMorphModelInstance>(MLDeformerComponent->GetModelInstance());
 							if (MorphInstance)
 							{
-								LOD = MorphInstance->GetFinalSkeletalMeshComponent()->GetPredictedLODLevel();
+								LOD = MorphInstance->GetFinalSkeletalMeshComponent() ? MorphInstance->GetFinalSkeletalMeshComponent()->GetPredictedLODLevel() : 0;
 							}
 						}
 
