NearestNeighborModel
- add finalize morph button
- close handle on deltas.bin if training is canceled
- raise exception when training in canceled
- fix status display issue
- fix file cache time stamp

diff --git a/Engine/Plugins/Animation/MLDeformer/MLDeformerFramework/Source/MLDeformerFramework/Private/MLDeformerMorphModel.cpp b/Engine/Plugins/Animation/MLDeformer/MLDeformerFramework/Source/MLDeformerFramework/Private/MLDeformerMorphModel.cpp
--- a/Engine/Plugins/Animation/MLDeformer/MLDeformerFramework/Source/MLDeformerFramework/Private/MLDeformerMorphModel.cpp
+++ b/Engine/Plugins/Animation/MLDeformer/MLDeformerFramework/Source/MLDeformerFramework/Private/MLDeformerMorphModel.cpp
@@ -338,6 +338,17 @@ void UMLDeformerMorphModel::UpdateMemoryUsage()
 	GPUMemUsageInBytes += GPUMorphSize;
 	CookedAssetSizeInBytes += GPUMorphSize;
 }
+
+void UMLDeformerMorphModel::FinalizeMorphTargets()
+{
+	MorphTargetDeltas.Empty();
+	UpdateStatistics();
+}
+
+bool UMLDeformerMorphModel::HasRawMorph() const
+{
+	return !MorphTargetDeltas.IsEmpty();
+}
 #endif
 
 #undef LOCTEXT_NAMESPACE
diff --git a/Engine/Plugins/Animation/MLDeformer/MLDeformerFramework/Source/MLDeformerFramework/Public/MLDeformerMorphModel.h b/Engine/Plugins/Animation/MLDeformer/MLDeformerFramework/Source/MLDeformerFramework/Public/MLDeformerMorphModel.h
--- a/Engine/Plugins/Animation/MLDeformer/MLDeformerFramework/Source/MLDeformerFramework/Public/MLDeformerMorphModel.h
+++ b/Engine/Plugins/Animation/MLDeformer/MLDeformerFramework/Source/MLDeformerFramework/Public/MLDeformerMorphModel.h
@@ -33,6 +33,8 @@ public:
 	virtual UMLDeformerModelInstance* CreateModelInstance(UMLDeformerComponent* Component) override;
 #if WITH_EDITOR
 	virtual void UpdateMemoryUsage() override;
+	virtual void FinalizeMorphTargets();
+	bool HasRawMorph() const;
 #endif
 	// ~END UMLDeformerModel overrides.
 
@@ -58,6 +60,7 @@ public:
 	EMLDeformerMaskChannel GetMaskChannel() const					{ return MaskChannel; }
 	bool GetInvertMaskChannel() const								{ return bInvertMaskChannel; }
 
+	UFUNCTION(BlueprintPure, Category = "MLDeformerMorphModel")
 	bool CanDynamicallyUpdateMorphTargets() const;
 
 	void SetMorphDeltaZeroThreshold(float Threshold)				{ MorphDeltaZeroThreshold = Threshold; }
@@ -355,7 +358,7 @@ private:
 	 * The advantage of this can be that it is higher performance than recomputing the normals.
 	 * The disadvantage is it can result in lower quality and uses more memory for the stored morph targets.
 	 */
-	UPROPERTY(EditAnywhere, Category = "Morph Targets")
+	UPROPERTY(EditAnywhere, Category = "Morph Targets", meta = (EditCondition = "CanDynamicallyUpdateMorphTargets()"))
 	bool bIncludeNormals = false;
 
 	/**
@@ -363,14 +366,14 @@ private:
 	 * This essentially removes small deltas from morph targets, which will lower the memory usage at runtime, however when set too high it can also introduce visual artifacts.
 	 * A value of 0 will result in the highest quality morph targets, at the cost of higher runtime memory usage.
 	 */
-	UPROPERTY(EditAnywhere, Category = "Morph Targets", DisplayName = "Delta Zero Threshold", meta = (ClampMin = "0.0", ClampMax = "1.0", ForceUnits="cm"))
+	UPROPERTY(EditAnywhere, Category = "Morph Targets", DisplayName = "Delta Zero Threshold", meta = (ClampMin = "0.0", ClampMax = "1.0", ForceUnits="cm", EditCondition = "CanDynamicallyUpdateMorphTargets()"))
 	float MorphDeltaZeroThreshold = 0.0025f;
 
 	/** 
 	 * The morph target compression level. Higher values result in larger compression, but could result in visual artifacts.
 	 * Most of the times this is a value between 20 and 200.
 	 */
-	UPROPERTY(EditAnywhere, Category = "Morph Targets", DisplayName = "Compression Level", meta = (ClampMin = "0.01", ClampMax = "1000"))
+	UPROPERTY(EditAnywhere, Category = "Morph Targets", DisplayName = "Compression Level", meta = (ClampMin = "0.01", ClampMax = "1000", EditCondition = "CanDynamicallyUpdateMorphTargets()"))
 	float MorphCompressionLevel = 20.0f;
 
 	/**
@@ -378,14 +381,14 @@ private:
 	 * You can use this feather out influence of the ML Deformer in specific areas, such as neck line seams, where the head mesh connects with the body.
 	 * The painted vertex color values will be like a weight multiplier on the ML deformer deltas applied to that vertex. You can invert the mask as well.
 	 */
-	UPROPERTY(EditAnywhere, Category = "Morph Targets")
+	UPROPERTY(EditAnywhere, Category = "Morph Targets", meta = (EditCondition = "CanDynamicallyUpdateMorphTargets()"))
 	EMLDeformerMaskChannel MaskChannel = EMLDeformerMaskChannel::Disabled;
 
 	/** 
 	 * Enable this if you want to invert the mask channel values. For example if you painted the neck seam vertices in red, and you wish the vertices that got painted to NOT move, you have to invert the mask.
 	 * On default you paint areas where the deformer should be active. If you enable the invert option, you paint areas where the deformer will not be active.
 	 */
-	UPROPERTY(EditAnywhere, Category = "Morph Targets", meta = (EditCondition = "MaskChannel != EMLDeformerMaskChannel::Disabled"))
+	UPROPERTY(EditAnywhere, Category = "Morph Targets", meta = (EditCondition = "CanDynamicallyUpdateMorphTargets() && MaskChannel != EMLDeformerMaskChannel::Disabled"))
 	bool bInvertMaskChannel = false;
 
 	/** The fence that let's us wait for all render commands to finish, before this instance is destroyed. */
diff --git a/Engine/Plugins/Animation/MLDeformer/MLDeformerFramework/Source/MLDeformerFrameworkEditor/Private/MLDeformerMorphModelDetails.cpp b/Engine/Plugins/Animation/MLDeformer/MLDeformerFramework/Source/MLDeformerFrameworkEditor/Private/MLDeformerMorphModelDetails.cpp
--- a/Engine/Plugins/Animation/MLDeformer/MLDeformerFramework/Source/MLDeformerFrameworkEditor/Private/MLDeformerMorphModelDetails.cpp
+++ b/Engine/Plugins/Animation/MLDeformer/MLDeformerFramework/Source/MLDeformerFrameworkEditor/Private/MLDeformerMorphModelDetails.cpp
@@ -8,6 +8,7 @@
 #include "DetailWidgetRow.h"
 #include "IDetailGroup.h"
 #include "SWarningOrErrorBox.h"
+#include "Widgets/Input/SButton.h"
 #include "Widgets/Layout/SBox.h"
 
 #define LOCTEXT_NAMESPACE "MLDeformerMorphModelDetails"
@@ -52,7 +53,7 @@ namespace UE::MLDeformer
 		MaskGroup.AddPropertyRow(DetailLayoutBuilder->GetProperty(UMLDeformerMorphModel::GetMaskChannelPropertyName(), UMLDeformerMorphModel::StaticClass()));
 		MaskGroup.AddPropertyRow(DetailLayoutBuilder->GetProperty(UMLDeformerMorphModel::GetInvertMaskChannelPropertyName(), UMLDeformerMorphModel::StaticClass()));
 
-		if (MorphModel && !MorphModel->CanDynamicallyUpdateMorphTargets())
+		if (MorphModel && MorphModel->HasRawMorph() && !MorphModel->CanDynamicallyUpdateMorphTargets())
 		{
 			const FText DeltaCountMismatchErrorText = LOCTEXT("MorphDeltaCountMismatch", "Dynamic morph target updates disabled until retrained. This is because the vertex count changed after the model was trained.");
 			FDetailWidgetRow& DeltaMismatchErrorRow = MorphTargetCategoryBuilder->AddCustomRow(FText::FromString("MorphDeltaCountMismatchError"))
@@ -67,6 +68,28 @@ namespace UE::MLDeformer
 					]
 				];
 		}
+
+		MorphTargetCategoryBuilder->AddCustomRow(LOCTEXT("FinalizeMorphTargetsButton", "Finalize Morph Targets"))
+			.WholeRowContent()
+			[
+				SNew(SBox)
+				.Padding(2, 2)
+				.MaxDesiredWidth(200.f)
+				[
+					SNew(SButton)
+					.Text(LOCTEXT("FinalizeMorphTargetsButtonText", "Finalize Morph Targets"))
+					.ToolTipText(LOCTEXT("FinalizeMorphTargetsButtonTooltip", "Delete raw data and fianlize the morph targets. This will reduce disk size but will make morph data ineditable. Re-training is required to edit the morph targets again."))
+					.HAlign(HAlign_Center)
+					.OnClicked_Lambda([this]() -> FReply
+					{
+						if (MorphModel)
+						{
+							MorphModel->FinalizeMorphTargets();
+						}
+						return FReply::Handled();
+					})
+				]
+			];
 	}
 }	// namespace UE::MLDeformer
 
diff --git a/Engine/Plugins/Animation/MLDeformer/NearestNeighborModel/Source/NearestNeighborModel/Private/NearestNeighborModel.cpp b/Engine/Plugins/Animation/MLDeformer/NearestNeighborModel/Source/NearestNeighborModel/Private/NearestNeighborModel.cpp
--- a/Engine/Plugins/Animation/MLDeformer/NearestNeighborModel/Source/NearestNeighborModel/Private/NearestNeighborModel.cpp
+++ b/Engine/Plugins/Animation/MLDeformer/NearestNeighborModel/Source/NearestNeighborModel/Private/NearestNeighborModel.cpp
@@ -1012,6 +1012,12 @@ UE::NearestNeighborModel::EOpFlag UNearestNeighborModelSection::NormalizeVertexW
 	return EOpFlag::Success;
 }
 
+void UNearestNeighborModel::UpdateVersion()
+{
+	using UE::NearestNeighborModel::FNearestNeighborModelCustomVersion;
+	Version = FNearestNeighborModelCustomVersion::LatestVersion;
+}
+
 TArray<FName> UNearestNeighborModelSection::GetVertexAttributeNames() const
 {
 	if (!Model)
@@ -1158,7 +1164,7 @@ void UNearestNeighborModel::PostEditChangeProperty(FPropertyChangedEvent& Proper
 
 void UNearestNeighborModel::PostEditChangeChainProperty(FPropertyChangedChainEvent& PropertyChangedEvent) 
 {
-	Super::PostEditChangeProperty(PropertyChangedEvent);
+	Super::PostEditChangeChainProperty(PropertyChangedEvent);
 	const FProperty* Property = PropertyChangedEvent.Property;
 	if (!Property)
 	{
@@ -1194,6 +1200,14 @@ void UNearestNeighborModel::PostEditChangeChainProperty(FPropertyChangedChainEve
 	}
 }
 
+void UNearestNeighborModel::FinalizeMorphTargets()
+{
+	for(FSection* Section : Sections)
+	{
+		Section->FinalizeMorphTargets();
+	}
+	Super::FinalizeMorphTargets();
+}
 #endif
 
 UMLDeformerInputInfo* UNearestNeighborModel::CreateInputInfo()
@@ -1531,7 +1545,6 @@ void UNearestNeighborModel::ClearReferences()
 	}
 }
 
-
 bool UNearestNeighborModel::IsBeforeCustomVersionWasAdded() const
 {
 	using UE::NearestNeighborModel::FNearestNeighborModelCustomVersion;
@@ -1658,6 +1671,15 @@ void UNearestNeighborModelSection::ClearReferences()
 	NeighborMeshes = nullptr;
 }
 
+void UNearestNeighborModelSection::FinalizeMorphTargets()
+{
+	VertexMean.Empty();
+	Basis.Empty();
+	AssetNeighborCoeffs.Empty();
+	AssetNeighborOffsets.Empty();
+	AssetNeighborIndexMap.Empty();
+}
+
 UE::NearestNeighborModel::EOpFlag UNearestNeighborModel::UpdateForTraining()
 {
 	if (IsReadyForTraining())
@@ -1817,6 +1839,7 @@ bool UNearestNeighborModel::LoadOptimizedNetworkFromFile(const FString& Filename
 			using UE::NearestNeighborModel::Private::GetArchitectureString;
 			NetworkLastWriteArchitectureString = GetArchitectureString(InputDim, HiddenLayerDims, OutputDim);
 			InvalidateInference();
+			UpdateVersion();
 			return true;
 		}
 		else
diff --git a/Engine/Plugins/Animation/MLDeformer/NearestNeighborModel/Source/NearestNeighborModel/Public/NearestNeighborModel.h b/Engine/Plugins/Animation/MLDeformer/NearestNeighborModel/Source/NearestNeighborModel/Public/NearestNeighborModel.h
--- a/Engine/Plugins/Animation/MLDeformer/NearestNeighborModel/Source/NearestNeighborModel/Public/NearestNeighborModel.h
+++ b/Engine/Plugins/Animation/MLDeformer/NearestNeighborModel/Source/NearestNeighborModel/Public/NearestNeighborModel.h
@@ -177,6 +177,7 @@ public:
 	EOpFlag NormalizeVertexWeights();
 
 	void ClearReferences();
+	void FinalizeMorphTargets();
 
 
 	static FName GetNumBasisPropertyName() { return GET_MEMBER_NAME_CHECKED(UNearestNeighborModelSection, NumPCACoeffs); }
@@ -344,6 +345,7 @@ public:
 #if WITH_EDITOR
 	virtual void PostEditChangeProperty(FPropertyChangedEvent& PropertyChangedEvent) override;
 	virtual void PostEditChangeChainProperty(FPropertyChangedChainEvent& PropertyChangedEvent) override;
+	virtual void FinalizeMorphTargets() override;
 #endif
 	// ~END UObject overrides.
 
@@ -619,6 +621,7 @@ private:
 	void UpdateCachedNetworkTimestamp();
 
 	void NormalizeVertexWeights();
+	void UpdateVersion();
 	// FNearestNeighborModelDetails needs to call private function GetSection(int32).
 	friend class UE::NearestNeighborModel::FNearestNeighborModelDetails;
 #endif
diff --git a/Engine/Plugins/Animation/MLDeformer/NearestNeighborModel/Source/NearestNeighborModelEditor/Private/NearestNeighborEditorModel.cpp b/Engine/Plugins/Animation/MLDeformer/NearestNeighborModel/Source/NearestNeighborModelEditor/Private/NearestNeighborEditorModel.cpp
--- a/Engine/Plugins/Animation/MLDeformer/NearestNeighborModel/Source/NearestNeighborModelEditor/Private/NearestNeighborEditorModel.cpp
+++ b/Engine/Plugins/Animation/MLDeformer/NearestNeighborModel/Source/NearestNeighborModelEditor/Private/NearestNeighborEditorModel.cpp
@@ -504,7 +504,7 @@ namespace UE::NearestNeighborModel
 		UpdateIsReadyForTrainingState();
 		if (!IsReadyForTraining())
 		{
-			UE_LOG(LogNearestNeighborModel, Error, TEXT("Model is not ready for training. Please check training data or reload MLDeformer editor."));
+			UE_LOG(LogNearestNeighborModel, Error, TEXT("Model is not ready for training. Please check training data is not empty or reload MLDeformer editor."));
 			return EOpFlag::Error;
 		}
 		UpdateResult |= UpdateNearestNeighborData();
diff --git a/Engine/Plugins/Animation/MLDeformer/NearestNeighborModel/Source/NearestNeighborModelEditor/Private/NearestNeighborModelDetails.cpp b/Engine/Plugins/Animation/MLDeformer/NearestNeighborModel/Source/NearestNeighborModelEditor/Private/NearestNeighborModelDetails.cpp
--- a/Engine/Plugins/Animation/MLDeformer/NearestNeighborModel/Source/NearestNeighborModelEditor/Private/NearestNeighborModelDetails.cpp
+++ b/Engine/Plugins/Animation/MLDeformer/NearestNeighborModel/Source/NearestNeighborModelEditor/Private/NearestNeighborModelDetails.cpp
@@ -322,7 +322,7 @@ namespace UE::NearestNeighborModel
 					{
 						return FText::FromString(TEXT("Valid"));
 					}
-					else if (NearestNeighborModel->IsReadyForTraining())
+					else if (NearestNeighborModel->IsReadyForTraining() && NearestNeighborModel->CanDynamicallyUpdateMorphTargets())
 					{
 						return FText::FromString(TEXT("Invalid, please Update"));
 					}
@@ -370,7 +370,7 @@ namespace UE::NearestNeighborModel
 					const UNearestNeighborModel* const NearestNeighborModel = GetCastModel();
 					if (NearestNeighborModel)
 					{
-						return NearestNeighborModel->IsReadyForTraining();
+						return NearestNeighborModel->IsReadyForTraining() && NearestNeighborModel->CanDynamicallyUpdateMorphTargets();
 					}
 					else
 					{
@@ -493,7 +493,7 @@ namespace UE::NearestNeighborModel
 				const TOptional<FDateTime> Timestamp = TimestampFunc();
 				if (Timestamp.IsSet())
 				{
-					return FText::FromString(Private::ToString(Timestamp.GetValue()));
+					return FText::FromString(Private::ToString(Private::ToLocalTime(Timestamp.GetValue())));
 				}
 				else
 				{
