[MLDeformer] Fixed a bug in the Finalize Morph Targets button where the memory usage wasn't reported correctly. Also fixed an issue where the editor uncooked asset didn't include the compressed morph targets. And fixed a typo in the button's tooltip and added a bit more information. Also made the Finalize Morph Targets button disabled when there is nothing to finalize.

diff --git a/Engine/Plugins/Animation/MLDeformer/MLDeformerFramework/Source/MLDeformerFramework/Private/MLDeformerMorphModel.cpp b/Engine/Plugins/Animation/MLDeformer/MLDeformerFramework/Source/MLDeformerFramework/Private/MLDeformerMorphModel.cpp
--- a/Engine/Plugins/Animation/MLDeformer/MLDeformerFramework/Source/MLDeformerFramework/Private/MLDeformerMorphModel.cpp
+++ b/Engine/Plugins/Animation/MLDeformer/MLDeformerFramework/Source/MLDeformerFramework/Private/MLDeformerMorphModel.cpp
@@ -337,12 +337,15 @@ void UMLDeformerMorphModel::UpdateMemoryUsage()
 	const uint64 GPUMorphSize = GetCompressedMorphDataSizeInBytes();
 	GPUMemUsageInBytes += GPUMorphSize;
 	CookedAssetSizeInBytes += GPUMorphSize;
+
+	EditorAssetSizeInBytes += CompressedMorphDataSizeInBytes;
 }
 
 void UMLDeformerMorphModel::FinalizeMorphTargets()
 {
 	MorphTargetDeltas.Empty();
 	UpdateStatistics();
+	UpdateMemoryUsage();
 }
 
 bool UMLDeformerMorphModel::HasRawMorph() const
diff --git a/Engine/Plugins/Animation/MLDeformer/MLDeformerFramework/Source/MLDeformerFrameworkEditor/Private/MLDeformerMorphModelDetails.cpp b/Engine/Plugins/Animation/MLDeformer/MLDeformerFramework/Source/MLDeformerFrameworkEditor/Private/MLDeformerMorphModelDetails.cpp
--- a/Engine/Plugins/Animation/MLDeformer/MLDeformerFramework/Source/MLDeformerFrameworkEditor/Private/MLDeformerMorphModelDetails.cpp
+++ b/Engine/Plugins/Animation/MLDeformer/MLDeformerFramework/Source/MLDeformerFrameworkEditor/Private/MLDeformerMorphModelDetails.cpp
@@ -74,12 +74,15 @@ namespace UE::MLDeformer
 			[
 				SNew(SBox)
 				.Padding(2, 2)
-				.MaxDesiredWidth(200.f)
 				[
 					SNew(SButton)
 					.Text(LOCTEXT("FinalizeMorphTargetsButtonText", "Finalize Morph Targets"))
-					.ToolTipText(LOCTEXT("FinalizeMorphTargetsButtonTooltip", "Delete raw data and fianlize the morph targets. This will reduce disk size but will make morph data ineditable. Re-training is required to edit the morph targets again."))
+					.ToolTipText(LOCTEXT("FinalizeMorphTargetsButtonTooltip", 
+						"Delete the raw vertex deltas, basically turning the editor asset into a cooked asset.\n"
+						"This will reduce the disk size of the uncooked asset, but will make morph target mask and compression settings uneditable until the model is retrained again.\n"
+						"Finalizing isn't required, but can be used to reduce the size of the files you submit to source control."))
 					.HAlign(HAlign_Center)
+					.IsEnabled_Lambda([this]() { return MorphModel && !MorphModel->GetMorphTargetDeltas().IsEmpty(); })
 					.OnClicked_Lambda([this]() -> FReply
 					{
 						if (MorphModel)
