[MLDeformer] Removed the delta cutoff threshold property and filtered out bad deltas. This also fixes the problem where the delta cutoff length was removing deltas in post-skinning space in the viewport, while removing them with the same length in pre-skinning space. This would possibly cause deltas removed during training that shouldn't have been removed.

Also improved a few property deprecations to be more correct.

diff --git a/Engine/Plugins/Animation/MLDeformer/MLDeformerFramework/Source/MLDeformerFramework/Private/MLDeformerModel.cpp b/Engine/Plugins/Animation/MLDeformer/MLDeformerFramework/Source/MLDeformerFramework/Private/MLDeformerModel.cpp
--- a/Engine/Plugins/Animation/MLDeformer/MLDeformerFramework/Source/MLDeformerFramework/Private/MLDeformerModel.cpp
+++ b/Engine/Plugins/Animation/MLDeformer/MLDeformerFramework/Source/MLDeformerFramework/Private/MLDeformerModel.cpp
@@ -178,7 +178,6 @@ void UMLDeformerModel::GetAssetRegistryTags(FAssetRegistryTagsContext Context) c
 	#if WITH_EDITORONLY_DATA
 		Context.AddTag(FAssetRegistryTag("MLDeformer.NumBones", FString::FromInt(BoneIncludeList.Num()), FAssetRegistryTag::TT_Numerical));
 		Context.AddTag(FAssetRegistryTag("MLDeformer.NumCurves", FString::FromInt(CurveIncludeList.Num()), FAssetRegistryTag::TT_Numerical));
-		Context.AddTag(FAssetRegistryTag("MLDeformer.DeltaCutoffLength", FString::Printf(TEXT("%f"), DeltaCutoffLength), FAssetRegistryTag::TT_Numerical));
 		Context.AddTag(FAssetRegistryTag("MLDeformer.MaxTrainingFrames", FString::FromInt(MaxTrainingFrames), FAssetRegistryTag::TT_Numerical));
 	#endif
 
diff --git a/Engine/Plugins/Animation/MLDeformer/MLDeformerFramework/Source/MLDeformerFramework/Public/MLDeformerModel.h b/Engine/Plugins/Animation/MLDeformer/MLDeformerFramework/Source/MLDeformerFramework/Public/MLDeformerModel.h
--- a/Engine/Plugins/Animation/MLDeformer/MLDeformerFramework/Source/MLDeformerFramework/Public/MLDeformerModel.h
+++ b/Engine/Plugins/Animation/MLDeformer/MLDeformerFramework/Source/MLDeformerFramework/Public/MLDeformerModel.h
@@ -423,7 +423,7 @@ public:
 	 * @return Returns true when bone transfomations should be a part of the network inputs, during the training process.
 	 */
 	UE_DEPRECATED(5.3, "This method and property has been removed and shouldn't be used anymore.")
-	bool ShouldIncludeBonesInTraining() const					{ return bIncludeBones_DEPRECATED; }
+	bool ShouldIncludeBonesInTraining() const					{ PRAGMA_DISABLE_DEPRECATION_WARNINGS; return bIncludeBones_DEPRECATED; PRAGMA_ENABLE_DEPRECATION_WARNINGS; }
 
 	/**
 	 * Set whether we want to include bones during training or not.
@@ -431,7 +431,7 @@ public:
 	 * @param bInclude Set to true if you wish bone transforms to be included during training and at inference time.
 	 */
 	UE_DEPRECATED(5.3, "This method and property has been removed and shouldn't be used anymore.")
-	void SetShouldIncludeBonesInTraining(bool bInclude)			{ bIncludeBones_DEPRECATED = bInclude; }
+	void SetShouldIncludeBonesInTraining(bool bInclude)			{ PRAGMA_DISABLE_DEPRECATION_WARNINGS; bIncludeBones_DEPRECATED = bInclude; PRAGMA_ENABLE_DEPRECATION_WARNINGS; }
 
 	/**
 	 * Check whether we should include curve values as input to the model during training or not.
@@ -439,7 +439,7 @@ public:
 	 * @return Returns true when curve values should be a part of the network inputs, during the training process.
 	 */
 	UE_DEPRECATED(5.3, "This method and property has been removed and shouldn't be used anymore.")
-	bool ShouldIncludeCurvesInTraining() const					{ return bIncludeCurves_DEPRECATED; }
+	bool ShouldIncludeCurvesInTraining() const					{ PRAGMA_DISABLE_DEPRECATION_WARNINGS; return bIncludeCurves_DEPRECATED; PRAGMA_ENABLE_DEPRECATION_WARNINGS; }
 
 	/**
 	 * Set whether we want to include curves during training.
@@ -447,7 +447,7 @@ public:
 	 * @param bInclude Set to true to include curves during training and inference time.
 	 */
 	UE_DEPRECATED(5.3, "This method and property has been removed and shouldn't be used anymore.")
-	void SetShouldIncludeCurvesInTraining(bool bInclude)		{ bIncludeCurves_DEPRECATED = bInclude; }
+	void SetShouldIncludeCurvesInTraining(bool bInclude)		{ PRAGMA_DISABLE_DEPRECATION_WARNINGS; bIncludeCurves_DEPRECATED = bInclude; PRAGMA_ENABLE_DEPRECATION_WARNINGS; }
 
 	/**
 	 * The delegate that gets fired when a property value changes.
@@ -563,7 +563,8 @@ public:
 	 * Skipping those deltas will prevent issues.
 	 * @return The length after which deltas will be ignored. So anything delta length larger than this value will be ignored.
 	 */
-	float GetDeltaCutoffLength() const										{ return DeltaCutoffLength; }
+	UE_DEPRECATED(5.5, "The delta cutoff length has been removed and it not used anymore.")
+	float GetDeltaCutoffLength() const										{ return 100000.0f; }
 
 	/**
 	 * Set the delta cutoff length. Deltas that are larger than this length will be set to zero.
@@ -571,7 +572,8 @@ public:
 	 * Skipping those deltas will prevent issues.
 	 * @param Length The new delta cutoff length.
 	 */
-	void SetDeltaCutoffLength(float Length)									{ DeltaCutoffLength = Length; }
+	UE_DEPRECATED(5.5, "The delta cutoff length has been removed and it not used anymore.")
+	void SetDeltaCutoffLength(float Length)									{ }
 
 	/**
 	 * Set the visualization settings object.
@@ -589,14 +591,16 @@ public:
 	static FName GetBoneIncludeListPropertyName()		{ return GET_MEMBER_NAME_CHECKED(UMLDeformerModel, BoneIncludeList); }
 	static FName GetCurveIncludeListPropertyName()		{ return GET_MEMBER_NAME_CHECKED(UMLDeformerModel, CurveIncludeList); }
 	static FName GetMaxTrainingFramesPropertyName()		{ return GET_MEMBER_NAME_CHECKED(UMLDeformerModel, MaxTrainingFrames); }
-	static FName GetDeltaCutoffLengthPropertyName()		{ return GET_MEMBER_NAME_CHECKED(UMLDeformerModel, DeltaCutoffLength); }
 	static FName GetMaxNumLODsPropertyName()			{ return GET_MEMBER_NAME_CHECKED(UMLDeformerModel, MaxNumLODs); }
 
 	UE_DEPRECATED(5.3, "This property has been removed and shouldn't be used anymore.")
-	static FName GetShouldIncludeBonesPropertyName()	{ return GET_MEMBER_NAME_CHECKED(UMLDeformerModel, bIncludeBones_DEPRECATED); }
+	static FName GetShouldIncludeBonesPropertyName()	{ PRAGMA_DISABLE_DEPRECATION_WARNINGS; return GET_MEMBER_NAME_CHECKED(UMLDeformerModel, bIncludeBones_DEPRECATED); PRAGMA_ENABLE_DEPRECATION_WARNINGS; }
 
 	UE_DEPRECATED(5.3, "This property has been removed and shouldn't be used anymore.")
-	static FName GetShouldIncludeCurvesPropertyName()	{ return GET_MEMBER_NAME_CHECKED(UMLDeformerModel, bIncludeCurves_DEPRECATED); }
+	static FName GetShouldIncludeCurvesPropertyName()	{ PRAGMA_DISABLE_DEPRECATION_WARNINGS; return GET_MEMBER_NAME_CHECKED(UMLDeformerModel, bIncludeCurves_DEPRECATED); PRAGMA_ENABLE_DEPRECATION_WARNINGS; }
+
+	UE_DEPRECATED(5.5, "The delta cutoff length has been removed and it not used anymore.")
+	static FName GetDeltaCutoffLengthPropertyName()		{ PRAGMA_DISABLE_DEPRECATION_WARNINGS; return GET_MEMBER_NAME_CHECKED(UMLDeformerModel, DeltaCutoffLength_DEPRECATED); PRAGMA_ENABLE_DEPRECATION_WARNINGS; }
 #endif	// #if WITH_EDITORONLY_DATA
 
 protected:
@@ -730,11 +734,13 @@ private:
 	TObjectPtr<UMLDeformerVizSettings> VizSettings;
 
 	/** Specifies whether bone transformations should be included as inputs during the training process. */
-	UPROPERTY()
+	UE_DEPRECATED(5.3, "This property has been removed.")
+	UPROPERTY(meta = (DeprecatedProperty, DeprecationMessage = "This property has been removed and isn't used anymore."))
 	bool bIncludeBones_DEPRECATED = true;
 
 	/** Specifies whether curve values (a float per curve) should be included as inputs during the training process. */
-	UPROPERTY()
+	UE_DEPRECATED(5.3, "This property has been removed.")
+	UPROPERTY(meta = (DeprecatedProperty, DeprecationMessage = "This property has been removed and isn't used anymore."))
 	bool bIncludeCurves_DEPRECATED = false;
 
 	/** The transform that aligns the Geometry Cache to the SkeletalMesh. This will mostly apply some scale and a rotation, but no translation. */
@@ -757,7 +763,8 @@ private:
 	 * Sometimes there can be some vertices that cause some issues that cause deltas to be very long. We can ignore these deltas by setting a cutoff value. 
 	 * Deltas that are longer than the cutoff value (in units), will be ignored and set to zero length. 
 	 */
-	UPROPERTY(EditAnywhere, Category = "Training Settings", meta = (ClampMin = "0.01", ForceUnits="cm"))
-	float DeltaCutoffLength = 30.0f;
+	UE_DEPRECATED(5.5, "This property has been removed.")
+	UPROPERTY(meta = (DeprecatedProperty, DeprecationMessage = "This property has been removed and isn't used anymore."))
+	float DeltaCutoffLength_DEPRECATED = 30.0f;
 #endif
 };
diff --git a/Engine/Plugins/Animation/MLDeformer/MLDeformerFramework/Source/MLDeformerFrameworkEditor/Private/MLDeformerEditorModel.cpp b/Engine/Plugins/Animation/MLDeformer/MLDeformerFramework/Source/MLDeformerFrameworkEditor/Private/MLDeformerEditorModel.cpp
--- a/Engine/Plugins/Animation/MLDeformer/MLDeformerFramework/Source/MLDeformerFrameworkEditor/Private/MLDeformerEditorModel.cpp
+++ b/Engine/Plugins/Animation/MLDeformer/MLDeformerFramework/Source/MLDeformerFrameworkEditor/Private/MLDeformerEditorModel.cpp
@@ -1091,8 +1091,7 @@ namespace UE::MLDeformer
 			}
 		}
 		else
-		if (Property->GetFName() == UMLDeformerModel::GetAlignmentTransformPropertyName() ||
-		    Property->GetFName() == UMLDeformerModel::GetDeltaCutoffLengthPropertyName())
+		if (Property->GetFName() == UMLDeformerModel::GetAlignmentTransformPropertyName())
 		{
 			if (PropertyChangedEvent.ChangeType == EPropertyChangeType::ValueSet)
 			{
@@ -1780,7 +1779,7 @@ namespace UE::MLDeformer
 								VertexDeltas[ArrayIndex + 2]);
 							const FVector VertexPos = (FVector)LinearSkinnedPositions[Index];
 							PDI->DrawLine(VertexPos, VertexPos + Delta, DeltasColor, DepthGroup);
-						}				
+						}
 					}
 				}
 			}
diff --git a/Engine/Plugins/Animation/MLDeformer/MLDeformerFramework/Source/MLDeformerFrameworkEditor/Private/MLDeformerGeomCacheSampler.cpp b/Engine/Plugins/Animation/MLDeformer/MLDeformerFramework/Source/MLDeformerFrameworkEditor/Private/MLDeformerGeomCacheSampler.cpp
--- a/Engine/Plugins/Animation/MLDeformer/MLDeformerFramework/Source/MLDeformerFrameworkEditor/Private/MLDeformerGeomCacheSampler.cpp
+++ b/Engine/Plugins/Animation/MLDeformer/MLDeformerFramework/Source/MLDeformerFrameworkEditor/Private/MLDeformerGeomCacheSampler.cpp
@@ -74,7 +74,6 @@ namespace UE::MLDeformer
 		UGeometryCache* GeometryCache = GeometryCacheComponent->GetGeometryCache();
 		if (SkeletalMeshComponent && SkeletalMesh && GeometryCacheComponent && GeometryCache)
 		{
-			const float DeltaCutoffLength = Model->GetDeltaCutoffLength();
 			const FTransform& AlignmentTransform = Model->GetAlignmentTransform();
 
 			// For all mesh mappings we found.
@@ -120,40 +119,37 @@ namespace UE::MLDeformer
 						const int32 GeomCacheVertexIndex = MeshMapping.SkelMeshToTrackVertexMap[VertexIndex];
 						if (GeomCacheVertexIndex != INDEX_NONE && GeomCacheMeshData.Positions.IsValidIndex(GeomCacheVertexIndex))
 						{
-							FVector3f Delta = FVector3f::ZeroVector;
+							const int32 RenderVertexIndex = MeshMapping.ImportedVertexToRenderVertexMap[VertexIndex];
+							if (RenderVertexIndex == INDEX_NONE)
+							{
+								continue;
+							}
 
-							const int32 ArrayIndex = 3 * SkinnedVertexIndex;
+							// Calculate the pre-skinning delta.
+							FVector3f Delta = FVector3f::ZeroVector;
 							if (VertexDeltaSpace == EVertexDeltaSpace::PreSkinning)
 							{
 								// Calculate the inverse skinning transform for this vertex.
-								const int32 RenderVertexIndex = MeshMapping.ImportedVertexToRenderVertexMap[VertexIndex];
-								if (RenderVertexIndex != INDEX_NONE)
-								{
-									const FMatrix44f InvSkinningTransform = CalcInverseSkinningTransform(RenderVertexIndex, SkelMeshLODData, SkinWeightBuffer);
-
-									// Calculate the pre-skinning data.
-									const FSkeletalMeshLODRenderData& LODData = SkeletalMesh->GetResourceForRendering()->LODRenderData[0];
-									const FVector3f UnskinnedPosition = LODData.StaticVertexBuffers.PositionVertexBuffer.VertexPosition(RenderVertexIndex);
-									const FVector3f GeomCacheVertexPos = (FVector3f)AlignmentTransform.TransformPosition((FVector)GeomCacheMeshData.Positions[GeomCacheVertexIndex]);
-									const FVector3f PreSkinningTargetPos = InvSkinningTransform.TransformPosition(GeomCacheVertexPos);
-									Delta = PreSkinningTargetPos - UnskinnedPosition;
-								}
+								const FMatrix44f InvSkinningTransform = CalcInverseSkinningTransform(RenderVertexIndex, SkelMeshLODData, SkinWeightBuffer);
+
+								// Calculate the pre-skinning data.
+								const FSkeletalMeshLODRenderData& LODData = SkeletalMesh->GetResourceForRendering()->LODRenderData[0];
+								const FVector3f UnskinnedPosition = LODData.StaticVertexBuffers.PositionVertexBuffer.VertexPosition(RenderVertexIndex);
+								const FVector3f GeomCacheVertexPos = (FVector3f)AlignmentTransform.TransformPosition((FVector)GeomCacheMeshData.Positions[GeomCacheVertexIndex]);
+								const FVector3f PreSkinningTargetPos = InvSkinningTransform.TransformPosition(GeomCacheVertexPos);
+								Delta = PreSkinningTargetPos - UnskinnedPosition;
 							}
-							else // We're post skinning.
+							else if (VertexDeltaSpace == EVertexDeltaSpace::PostSkinning)
 							{
-								check(VertexDeltaSpace == EVertexDeltaSpace::PostSkinning);
 								const FVector3f SkinnedVertexPos = SkinnedVertexPositions[SkinnedVertexIndex];
 								const FVector3f GeomCacheVertexPos = (FVector3f)AlignmentTransform.TransformPosition((FVector)GeomCacheMeshData.Positions[GeomCacheVertexIndex]);
 								Delta = GeomCacheVertexPos - SkinnedVertexPos;
 							}
 
-							// Set the delta.
-							if (Delta.Length() < DeltaCutoffLength)
-							{
-								VertexDeltas[ArrayIndex] = Delta.X;
-								VertexDeltas[ArrayIndex + 1] = Delta.Y;
-								VertexDeltas[ArrayIndex + 2] = Delta.Z;
-							}
+							const int32 ArrayIndex = 3 * SkinnedVertexIndex;
+							VertexDeltas[ArrayIndex + 0] = Delta.X;
+							VertexDeltas[ArrayIndex + 1] = Delta.Y;
+							VertexDeltas[ArrayIndex + 2] = Delta.Z;
 						}
 					}
 				});	// ParallelFor
@@ -161,7 +157,7 @@ namespace UE::MLDeformer
 		}
 		else
 		{
-			VertexDeltas.Reset(0);
+			VertexDeltas.Reset();
 		}
 	}
 
diff --git a/Engine/Plugins/Animation/MLDeformer/MLDeformerFramework/Source/MLDeformerFrameworkEditor/Private/MLDeformerModelDetails.cpp b/Engine/Plugins/Animation/MLDeformer/MLDeformerFramework/Source/MLDeformerFrameworkEditor/Private/MLDeformerModelDetails.cpp
--- a/Engine/Plugins/Animation/MLDeformer/MLDeformerFramework/Source/MLDeformerFrameworkEditor/Private/MLDeformerModelDetails.cpp
+++ b/Engine/Plugins/Animation/MLDeformer/MLDeformerFramework/Source/MLDeformerFrameworkEditor/Private/MLDeformerModelDetails.cpp
@@ -193,7 +193,6 @@ namespace UE::MLDeformer
 		}
 
 		TrainingSettingsCategoryBuilder->AddProperty(UMLDeformerModel::GetMaxTrainingFramesPropertyName(), UMLDeformerModel::StaticClass());
-		TrainingSettingsCategoryBuilder->AddProperty(UMLDeformerModel::GetDeltaCutoffLengthPropertyName(), UMLDeformerModel::StaticClass());
 		AddTrainingSettingsErrors();
 
 		LODSettingsCategoryBuilder->AddProperty(UMLDeformerModel::GetMaxNumLODsPropertyName(), UMLDeformerModel::StaticClass());
diff --git a/Engine/Plugins/Animation/MLDeformer/NearestNeighborModel/Source/NearestNeighborModel/Private/NearestNeighborModel.cpp b/Engine/Plugins/Animation/MLDeformer/NearestNeighborModel/Source/NearestNeighborModel/Private/NearestNeighborModel.cpp
--- a/Engine/Plugins/Animation/MLDeformer/NearestNeighborModel/Source/NearestNeighborModel/Private/NearestNeighborModel.cpp
+++ b/Engine/Plugins/Animation/MLDeformer/NearestNeighborModel/Source/NearestNeighborModel/Private/NearestNeighborModel.cpp
@@ -1210,7 +1210,6 @@ void UNearestNeighborModel::PostEditChangeProperty(FPropertyChangedEvent& Proper
 		Property->GetFName() == UMLDeformerModel::GetBoneIncludeListPropertyName() ||
 		Property->GetFName() == UMLDeformerModel::GetCurveIncludeListPropertyName() ||
 		Property->GetFName() == UMLDeformerModel::GetMaxTrainingFramesPropertyName() ||
-		Property->GetFName() == UMLDeformerModel::GetDeltaCutoffLengthPropertyName() ||
 		Property->GetFName() == UNearestNeighborModel::GetHiddenLayerDimsPropertyName() ||
 		Property->GetFName() == UNearestNeighborModel::GetNumEpochsPropertyName() ||
 		Property->GetFName() == UNearestNeighborModel::GetBatchSizePropertyName() ||
