[MLDeformer] Show an error when the deformer can produce visual shading artifacts. This can happen when no deformer graph is used while skin cache is disabled and the materials don't have 'Use with morph targets' enabled.

diff --git a/Engine/Plugins/Animation/MLDeformer/MLDeformerFramework/Source/MLDeformerFramework/Private/MLDeformerGeomCacheHelpers.cpp b/Engine/Plugins/Animation/MLDeformer/MLDeformerFramework/Source/MLDeformerFramework/Private/MLDeformerGeomCacheHelpers.cpp
--- a/Engine/Plugins/Animation/MLDeformer/MLDeformerFramework/Source/MLDeformerFramework/Private/MLDeformerGeomCacheHelpers.cpp
+++ b/Engine/Plugins/Animation/MLDeformer/MLDeformerFramework/Source/MLDeformerFramework/Private/MLDeformerGeomCacheHelpers.cpp
@@ -154,7 +154,6 @@ namespace UE::MLDeformer
 		}
 
 		FSkeletalMeshConstAttributes MeshAttributes(*MeshDescription);
-
 		if (!MeshAttributes.HasSourceGeometryParts())
 		{
 			return;
@@ -172,9 +171,8 @@ namespace UE::MLDeformer
 
 		const int32 NumSourceGeoParts = MeshAttributes.GetNumSourceGeometryParts();
 		const FSkeletalMeshAttributes::FSourceGeometryPartNameConstRef GeoPartNames = MeshAttributes.GetSourceGeometryPartNames();
-		const FSkeletalMeshAttributes::FSourceGeometryPartVertexOffsetAndCountConstRef GeoPartOffsetAndCounts = MeshAttributes.GetSourceGeometryPartVertexOffsetAndCounts();
-		
-
+		const FSkeletalMeshAttributes::FSourceGeometryPartVertexOffsetAndCountConstRef GeoPartOffsetAndCounts = MeshAttributes.GetSourceGeometryPartVertexOffsetAndCounts();	
+	
 		const bool bIsSoloMesh = (GeomCache->Tracks.Num() == 1 && NumSourceGeoParts == 1);	// Do we just have one mesh and one track?
 		for (int32 TrackIndex = 0; TrackIndex < GeomCache->Tracks.Num(); ++TrackIndex)
 		{
@@ -257,6 +255,21 @@ namespace UE::MLDeformer
 						}
 					});
 
+					// Figure out what materials we use.
+					for (int32 VertexIndex = 0; VertexIndex < NumSkelMeshVerts; ++VertexIndex)
+					{
+						const int32 RenderVertexIndex = Mapping.ImportedVertexToRenderVertexMap[VertexIndex];
+
+						int32 SectionIndex = INDEX_NONE;
+						int32 SectionVertexIndex = INDEX_NONE;
+						LODModel.GetSectionFromVertexIndex(RenderVertexIndex, SectionIndex, SectionVertexIndex);
+						if (SectionIndex != INDEX_NONE && SectionVertexIndex != INDEX_NONE)
+						{
+							const int32 MaterialIndex = LODModel.Sections[SectionIndex].MaterialIndex;
+							Mapping.MaterialIndices.AddUnique(MaterialIndex);
+						}
+					}
+
 					// We found a match, no need to iterate over more Tracks.
 					bFoundMatch = true;
 					if (!bSuppressLog)
diff --git a/Engine/Plugins/Animation/MLDeformer/MLDeformerFramework/Source/MLDeformerFramework/Public/MLDeformerGeomCacheHelpers.h b/Engine/Plugins/Animation/MLDeformer/MLDeformerFramework/Source/MLDeformerFramework/Public/MLDeformerGeomCacheHelpers.h
--- a/Engine/Plugins/Animation/MLDeformer/MLDeformerFramework/Source/MLDeformerFramework/Public/MLDeformerGeomCacheHelpers.h
+++ b/Engine/Plugins/Animation/MLDeformer/MLDeformerFramework/Source/MLDeformerFramework/Public/MLDeformerGeomCacheHelpers.h
@@ -22,6 +22,7 @@ namespace UE::MLDeformer
 		int32 TrackIndex = INDEX_NONE;	// The geometry cache track that this mesh is mapped to.
 		TArray<int32> SkelMeshToTrackVertexMap;	// This maps imported model individual meshes to the geomcache track's mesh data.
 		TArray<int32> ImportedVertexToRenderVertexMap; // Map the imported DCC vertex number to a render vertex. This is just one of the duplicates, which shares the same position.
+		TArray<int32> MaterialIndices;	// List of material indices that point inside the SkelMesh. These materials are used on parts of the mesh that are modified by the ML Deformer.
 	};
 
 	/**
diff --git a/Engine/Plugins/Animation/MLDeformer/MLDeformerFramework/Source/MLDeformerFrameworkEditor/Private/MLDeformerEditorModel.cpp b/Engine/Plugins/Animation/MLDeformer/MLDeformerFramework/Source/MLDeformerFrameworkEditor/Private/MLDeformerEditorModel.cpp
--- a/Engine/Plugins/Animation/MLDeformer/MLDeformerFramework/Source/MLDeformerFrameworkEditor/Private/MLDeformerEditorModel.cpp
+++ b/Engine/Plugins/Animation/MLDeformer/MLDeformerFramework/Source/MLDeformerFrameworkEditor/Private/MLDeformerEditorModel.cpp
@@ -1156,6 +1156,7 @@ namespace UE::MLDeformer
 		{
 			UpdateDeformerGraph();
 			GetEditor()->GetVizSettingsDetailsView()->ForceRefresh();
+			GetEditor()->GetModelDetailsView()->ForceRefresh();
 		}
 	}
 
diff --git a/Engine/Plugins/Animation/MLDeformer/MLDeformerFramework/Source/MLDeformerFrameworkEditor/Private/MLDeformerMorphModelDetails.cpp b/Engine/Plugins/Animation/MLDeformer/MLDeformerFramework/Source/MLDeformerFrameworkEditor/Private/MLDeformerMorphModelDetails.cpp
--- a/Engine/Plugins/Animation/MLDeformer/MLDeformerFramework/Source/MLDeformerFrameworkEditor/Private/MLDeformerMorphModelDetails.cpp
+++ b/Engine/Plugins/Animation/MLDeformer/MLDeformerFramework/Source/MLDeformerFrameworkEditor/Private/MLDeformerMorphModelDetails.cpp
@@ -3,6 +3,7 @@
 #include "MLDeformerMorphModelDetails.h"
 #include "MLDeformerMorphModelEditorModel.h"
 #include "MLDeformerMorphModel.h"
+#include "MLDeformerGeomCacheHelpers.h"
 #include "DetailLayoutBuilder.h"
 #include "DetailCategoryBuilder.h"
 #include "DetailWidgetRow.h"
@@ -10,6 +11,10 @@
 #include "SWarningOrErrorBox.h"
 #include "Widgets/Input/SButton.h"
 #include "Widgets/Layout/SBox.h"
+#include "Animation/DebugSkelMeshComponent.h"
+#include "Engine/SkinnedAssetCommon.h"
+#include "Materials/MaterialInterface.h"
+#include "Materials/Material.h"
 
 #define LOCTEXT_NAMESPACE "MLDeformerMorphModelDetails"
 
@@ -35,11 +40,130 @@ namespace UE::MLDeformer
 		MorphTargetCategoryBuilder = &DetailLayoutBuilder->EditCategory("Morph Targets", FText::GetEmpty(), ECategoryPriority::Important);
 	}
 
+	bool FMLDeformerMorphModelDetails::ShouldShowShadingError() const
+	{
+		if (!Model)
+		{
+			return false;
+		}
+
+		if (Model->GetVizSettings()->GetDeformerGraph())
+		{
+			return false;
+		}
+
+		const FMLDeformerEditorActor* MLActor = EditorModel->FindEditorActor(ActorID_Test_MLDeformed);
+		if (!MLActor)
+		{
+			return false;
+		}
+
+		const UDebugSkelMeshComponent* SkelMeshComponent = MLActor->GetSkeletalMeshComponent();
+		if (!SkelMeshComponent)
+		{
+			return false;
+		}
+
+		// Check if we have skin cache enabled.
+		const int32 LOD = 0;
+		if (SkelMeshComponent->IsSkinCacheAllowed(LOD))
+		{
+			return false;
+		}
+
+		// Check the materials on parts of the mesh that are modified by the ML Deformer.
+		TArray<int32> MaterialIndices;
+		const int32 NumTrainingAnims = GeomCacheEditorModel->GetNumTrainingInputAnims();
+		for (int32 AnimIndex = 0; AnimIndex < NumTrainingAnims; ++AnimIndex)
+		{
+			const FMLDeformerGeomCacheTrainingInputAnim* Anim = static_cast<FMLDeformerGeomCacheTrainingInputAnim*>(GeomCacheEditorModel->GetTrainingInputAnim(AnimIndex));
+			if (!Anim || !Anim->IsEnabled())
+			{
+				continue;
+			}
+
+			// Get the sampler and try to initialize it when needed.
+			FMLDeformerGeomCacheSampler* Sampler = static_cast<FMLDeformerGeomCacheSampler*>(GeomCacheEditorModel->GetSamplerForTrainingAnim(AnimIndex));
+			if (!Sampler->IsInitialized())
+			{
+				Sampler->Init(GeomCacheEditorModel, AnimIndex);
+			}
+
+			if (Sampler->IsInitialized()) // This could still fail, when it failed to initialize before.
+			{
+				for (const FMLDeformerGeomCacheMeshMapping& MeshMapping : Sampler->GetMeshMappings())
+				{
+					for (const int32 MaterialIndex : MeshMapping.MaterialIndices)
+					{
+						MaterialIndices.AddUnique(MaterialIndex);
+					}
+				}
+			}
+		}
+
+		// Now check if the used materials have the 'Use with morph targets' flag disabled.
+		const USkeletalMesh* SkelMesh = SkelMeshComponent->GetSkeletalMeshAsset();
+		if (SkelMesh)
+		{
+			const TArray<FSkeletalMaterial>& Materials = SkelMesh->GetMaterials();
+			bool bHasMaterialErrors = false;
+			for (const int32 MaterialIndex : MaterialIndices)
+			{
+				const TObjectPtr<UMaterialInterface> MaterialInterface = Materials[MaterialIndex].MaterialInterface;
+				if (MaterialInterface && MaterialInterface->GetMaterial() && !MaterialInterface->GetMaterial()->bUsedWithMorphTargets)
+				{
+					UE_LOG(LogMLDeformer, Warning, TEXT("Material '%s' (Index=%d) has the 'Used with Morph Targets' property disabled, while no deformer graph or skin cache is used. This can cause issues with ML Deformer."),
+						*Materials[MaterialIndex].MaterialSlotName.ToString(), MaterialIndex);
+					bHasMaterialErrors = true;
+				}
+			}
+			
+			if (bHasMaterialErrors)
+			{
+				return true;
+			}
+		}
+
+		return false;
+	}
+
 	void FMLDeformerMorphModelDetails::CustomizeDetails(IDetailLayoutBuilder& DetailBuilder)
 	{
 		// Create all the detail categories and add the properties of the base class.
 		FMLDeformerGeomCacheModelDetails::CustomizeDetails(DetailBuilder);
 
+		if (ShouldShowShadingError())
+		{
+			FDetailWidgetRow& CannotRunMLDeformerWarning = BaseMeshCategoryBuilder->AddCustomRow(FText::FromString("CannotRunMLDeformerError"))
+				.WholeRowContent()
+				[
+					SNew(SBox)
+					.Padding(FMargin(0.0f, 4.0f))
+					[
+						SNew(SWarningOrErrorBox)
+						.MessageStyle(EMessageStyle::Error)
+						.Message_Lambda
+						(
+							[this]()
+							{
+								return LOCTEXT("CannotRunMLDeformerErrorMessage", 
+									"This ML Deformer cannot work properly because:\n"
+									"\n"
+									"- No deformer graph is used.\n"
+									"- Skin cache is disabled in the project settings or the Skeletal Mesh.\n"
+									"- There are materials that have the 'Used with morph targets' property disabled.\n"
+									"\n"
+									"This can lead to visual shading artifacts or the deformer not working at all.\n"
+									"See the log for more details.\n"
+									"\n"
+									"To fix this, either use a deformer graph, enable skin cache, or enable the mentioned material property."
+								);
+							}
+						)
+					]
+				];
+		}
+
 		MorphTargetCategoryBuilder->AddProperty(DetailLayoutBuilder->GetProperty(UMLDeformerMorphModel::GetClampMorphTargetWeightsPropertyName(), UMLDeformerMorphModel::StaticClass()))
 			.Visibility(MorphModelEditorModel->IsMorphWeightClampingSupported() ? EVisibility::Visible : EVisibility::Collapsed);
 
diff --git a/Engine/Plugins/Animation/MLDeformer/MLDeformerFramework/Source/MLDeformerFrameworkEditor/Public/MLDeformerMorphModelDetails.h b/Engine/Plugins/Animation/MLDeformer/MLDeformerFramework/Source/MLDeformerFrameworkEditor/Public/MLDeformerMorphModelDetails.h
--- a/Engine/Plugins/Animation/MLDeformer/MLDeformerFramework/Source/MLDeformerFrameworkEditor/Public/MLDeformerMorphModelDetails.h
+++ b/Engine/Plugins/Animation/MLDeformer/MLDeformerFramework/Source/MLDeformerFrameworkEditor/Public/MLDeformerMorphModelDetails.h
@@ -30,6 +30,21 @@ namespace UE::MLDeformer
 		virtual bool UpdateMemberPointers(const TArray<TWeakObjectPtr<UObject>>& Objects) override;
 		// ~END FMLDeformerModelDetails overrides.
 
+	protected:
+		/**
+		 * Check whether we should show an error in the UI about possible shading artifacts.
+		 * The ML Deformer cannot run properly if all of the following things are true:
+		 * - Skin cache is disabled
+		 * - Material settings have "Use with morph targets" disabled.
+		 * - No deformer graph is being used.
+		 * 
+		 * This method checks those points.
+		 * If not, we can show a warning in the UI.
+		 * 
+		 * @return Returns true if an error should be displayed in the UI about this issue.
+		 */
+		bool ShouldShowShadingError() const;
+
 	protected:
 		/** A pointer to the morph model. This is updated when UpdateMemberPointers is called. */
 		TObjectPtr<UMLDeformerMorphModel> MorphModel = nullptr;
