[MLDeformer] Copy over all shared settings when switching between model types.

diff --git a/Engine/Plugins/Animation/MLDeformer/MLDeformerFramework/Source/MLDeformerFramework/Public/MLDeformerGeomCacheModel.h b/Engine/Plugins/Animation/MLDeformer/MLDeformerFramework/Source/MLDeformerFramework/Public/MLDeformerGeomCacheModel.h
--- a/Engine/Plugins/Animation/MLDeformer/MLDeformerFramework/Source/MLDeformerFramework/Public/MLDeformerGeomCacheModel.h
+++ b/Engine/Plugins/Animation/MLDeformer/MLDeformerFramework/Source/MLDeformerFramework/Public/MLDeformerGeomCacheModel.h
@@ -102,7 +102,8 @@ private:
 	UPROPERTY(meta=(DeprecatedProperty, DeprecationMessage="Use the training input anims instead."))
 	TSoftObjectPtr<UGeometryCache> GeometryCache_DEPRECATED;
 
-	UPROPERTY(EditAnywhere, Category = "Target Mesh")
+	/** The training animation inputs. */
+	UPROPERTY(EditAnywhere, Category = "Inputs")
 	TArray<FMLDeformerGeomCacheTrainingInputAnim> TrainingInputAnims;
 #endif // WITH_EDITORONLY_DATA
 };
diff --git a/Engine/Plugins/Animation/MLDeformer/MLDeformerFramework/Source/MLDeformerFramework/Public/MLDeformerGeomCacheVizSettings.h b/Engine/Plugins/Animation/MLDeformer/MLDeformerFramework/Source/MLDeformerFramework/Public/MLDeformerGeomCacheVizSettings.h
--- a/Engine/Plugins/Animation/MLDeformer/MLDeformerFramework/Source/MLDeformerFramework/Public/MLDeformerGeomCacheVizSettings.h
+++ b/Engine/Plugins/Animation/MLDeformer/MLDeformerFramework/Source/MLDeformerFramework/Public/MLDeformerGeomCacheVizSettings.h
@@ -31,6 +31,12 @@ public:
 	 */
 	UGeometryCache* GetTestGroundTruth() const			{ return GroundTruth.LoadSynchronous(); }
 
+	/**
+	 * Set the test ground truth geometry cache.
+	 * @param InGroundTruth The ground truth model.
+	 */
+	void SetTestGroundTruth(TSoftObjectPtr<UGeometryCache> InGroundTruth)	{ GroundTruth = InGroundTruth; }
+
 	// Get property names.
 	static FName GetTestGroundTruthPropertyName()		{ return GET_MEMBER_NAME_CHECKED(UMLDeformerGeomCacheVizSettings, GroundTruth); }
 
diff --git a/Engine/Plugins/Animation/MLDeformer/MLDeformerFramework/Source/MLDeformerFramework/Public/MLDeformerMorphModel.h b/Engine/Plugins/Animation/MLDeformer/MLDeformerFramework/Source/MLDeformerFramework/Public/MLDeformerMorphModel.h
--- a/Engine/Plugins/Animation/MLDeformer/MLDeformerFramework/Source/MLDeformerFramework/Public/MLDeformerMorphModel.h
+++ b/Engine/Plugins/Animation/MLDeformer/MLDeformerFramework/Source/MLDeformerFramework/Public/MLDeformerMorphModel.h
@@ -64,12 +64,15 @@ public:
 	bool CanDynamicallyUpdateMorphTargets() const;
 
 	void SetMorphDeltaZeroThreshold(float Threshold)				{ MorphDeltaZeroThreshold = Threshold; }
-	void SetMorphCompressionlevel(float Tolerance)					{ MorphCompressionLevel = Tolerance; }
+	void SetMorphCompressionLevel(float Tolerance)					{ MorphCompressionLevel = Tolerance; }
 	void SetIncludeMorphTargetNormals(bool bInclude)				{ bIncludeNormals = bInclude; }
 	void SetMaskChannel(EMLDeformerMaskChannel Channel)				{ MaskChannel = Channel; }
 	void SetInvertMaskChannel(bool bInvert)							{ bInvertMaskChannel = bInvert; }
 	void SetClampMorphTargetsWeights(bool bEnabled)					{ bClampMorphWeights = bEnabled; }
 
+	UE_DEPRECATED(5.5, "Please use SetMorphCompressionLevel, with upper case L for Level.")
+	void SetMorphCompressionlevel(float Tolerance)					{ MorphCompressionLevel = Tolerance; }
+
 	UE_DEPRECATED(5.3, "Use SetMaskChannel instead.")
 	void SetWeightMask(EMLDeformerMaskChannel Channel)				{ MaskChannel = Channel; }
 
diff --git a/Engine/Plugins/Animation/MLDeformer/MLDeformerFramework/Source/MLDeformerFramework/Public/MLDeformerVizSettings.h b/Engine/Plugins/Animation/MLDeformer/MLDeformerFramework/Source/MLDeformerFramework/Public/MLDeformerVizSettings.h
--- a/Engine/Plugins/Animation/MLDeformer/MLDeformerFramework/Source/MLDeformerFramework/Public/MLDeformerVizSettings.h
+++ b/Engine/Plugins/Animation/MLDeformer/MLDeformerFramework/Source/MLDeformerFramework/Public/MLDeformerVizSettings.h
@@ -74,6 +74,24 @@ public:
 	void SetTestingFrameNumber(int32 FrameNumber)			{ TestingFrameNumber = FrameNumber; }
 	void SetWeight(float InWeight)							{ Weight = InWeight; }
 	void SetTestAnimSequence(UAnimSequence* InAnim)  		{ TestAnimSequence = InAnim; }
+	void SetMeshSpacing(float Spacing)						{ MeshSpacing = Spacing; }
+	void SetLabelHeight(float Height)						{ LabelHeight = Height; }
+	void SetLabelScale(float Scale)							{ LabelScale = Scale; }
+	void SetDrawLabels(bool bDraw)							{ bDrawLabels = bDraw; }
+	void SetAnimPlaySpeed(float SpeedFactor)				{ AnimPlaySpeed = SpeedFactor; }
+	void SetDrawLinearSkinnedActor(bool bDraw)				{ bDrawLinearSkinnedActor = bDraw; }
+	void SetDrawMLDeformedActor(bool bDraw)					{ bDrawMLDeformedActor = bDraw; }
+	void SetDrawGroundTruthActor(bool bDraw)				{ bDrawGroundTruthActor = bDraw; }
+	void SetDrawMLCompareActors(bool bDraw)					{ bDrawMLCompareActors = bDraw; }
+	void SetDrawDebugActorBounds(bool bDraw)				{ bDrawDebugActorBounds = bDraw; }
+	void SetDrawVertexDeltas(bool bDraw)					{ bDrawDeltas = bDraw; }
+	void SetXRayDeltas(bool bEnable)						{ bXRayDeltas = bEnable; }
+	void SetHeatMapMax(float Value)							{ HeatMapMax = Value; }
+	void SetGroundTruthLerp(float Alpha)					{ GroundTruthLerp = Alpha; }
+	void SetShowHeatMap(bool bShow)							{ bShowHeatMap = bShow; }
+	void SetHeatMapMode(EMLDeformerHeatMapMode Mode)		{ HeatMapMode = Mode; }
+	void SetDebugBoundsColor(const FColor& Color)			{ DebugBoundsColor = Color; }
+	void SetCompareActors(const TArray<FMLDeformerCompareActor>& Actors) { CompareActors = Actors; }
 
 	FVector GetMeshSpacingOffsetVector() const				{ return FVector(MeshSpacing, 0.0f, 0.0f); }
 	float GetMeshSpacing() const							{ return MeshSpacing; }
@@ -103,7 +121,6 @@ public:
 	TArray<FMLDeformerCompareActor>& GetCompareActors()		{ return CompareActors; }
 	FColor GetDebugBoundsColor() const						{ return DebugBoundsColor; }
 
-
 	UE_DEPRECATED(5.4, "This method will be removed.")
 	void SetQualityLevel(int32 InQualityLevel)				{ QualityLevel_DEPRECATED = FMath::Max<int32>(InQualityLevel, 0); }
 
diff --git a/Engine/Plugins/Animation/MLDeformer/MLDeformerFramework/Source/MLDeformerFrameworkEditor/Private/MLDeformerEditorModel.cpp b/Engine/Plugins/Animation/MLDeformer/MLDeformerFramework/Source/MLDeformerFrameworkEditor/Private/MLDeformerEditorModel.cpp
--- a/Engine/Plugins/Animation/MLDeformer/MLDeformerFramework/Source/MLDeformerFrameworkEditor/Private/MLDeformerEditorModel.cpp
+++ b/Engine/Plugins/Animation/MLDeformer/MLDeformerFramework/Source/MLDeformerFrameworkEditor/Private/MLDeformerEditorModel.cpp
@@ -166,6 +166,67 @@ namespace UE::MLDeformer
 		Collector.AddReferencedObject(EditorInputInfo);
 	}
 
+	void FMLDeformerEditorModel::CopyBaseSettingsFromModel(const FMLDeformerEditorModel* SourceEditorModel)
+	{
+		UMLDeformerModel* SourceModel = SourceEditorModel->GetModel();
+		if (SourceModel == nullptr)
+		{
+			return;
+		}
+
+		UMLDeformerModel* TargetModel = GetModel();
+		check(TargetModel);
+
+		// Copy model properties.
+		TargetModel->SetSkeletalMesh(SourceModel->GetSkeletalMesh());
+		TargetModel->SetAlignmentTransform(SourceModel->GetAlignmentTransform());
+		TargetModel->SetTrainingFrameLimit(SourceModel->GetTrainingFrameLimit());
+		if (TargetModel->DoesSupportBones())
+		{
+			TargetModel->SetBoneIncludeList(SourceModel->GetBoneIncludeList());
+		}
+		if (TargetModel->DoesSupportCurves())
+		{
+			TargetModel->SetCurveIncludeList(SourceModel->GetCurveIncludeList());
+		}
+
+		// Copy visualization properties.
+		const UMLDeformerVizSettings* SourceVizSettings = SourceModel->GetVizSettings();
+		UMLDeformerVizSettings* TargetVizSettings = TargetModel->GetVizSettings();
+		check(TargetVizSettings);
+		if (SourceVizSettings)
+		{
+			TargetVizSettings->SetLabelHeight(SourceVizSettings->GetLabelHeight());
+			TargetVizSettings->SetVisualizationMode(SourceVizSettings->GetVisualizationMode());
+			TargetVizSettings->SetMeshSpacing(SourceVizSettings->GetMeshSpacing());
+			TargetVizSettings->SetLabelScale(SourceVizSettings->GetLabelScale());
+			TargetVizSettings->SetDrawLabels(SourceVizSettings->GetDrawLabels());
+			TargetVizSettings->SetWeight(SourceVizSettings->GetWeight());
+			TargetVizSettings->SetTestAnimSequence((UAnimSequence*)SourceVizSettings->GetTestAnimSequence());
+			TargetVizSettings->SetTrainingFrameNumber(SourceVizSettings->GetTrainingFrameNumber());
+			TargetVizSettings->SetTestingFrameNumber(SourceVizSettings->GetTestingFrameNumber());
+			TargetVizSettings->SetAnimPlaySpeed(SourceVizSettings->GetAnimPlaySpeed());
+			TargetVizSettings->SetDrawLinearSkinnedActor(SourceVizSettings->GetDrawLinearSkinnedActor());
+			TargetVizSettings->SetDrawMLDeformedActor(SourceVizSettings->GetDrawMLDeformedActor());
+			TargetVizSettings->SetDrawGroundTruthActor(SourceVizSettings->GetDrawGroundTruthActor());
+			TargetVizSettings->SetDrawMLCompareActors(SourceVizSettings->GetDrawMLCompareActors());
+			TargetVizSettings->SetDrawDebugActorBounds(SourceVizSettings->GetDrawDebugActorBounds());
+			TargetVizSettings->SetDrawVertexDeltas(SourceVizSettings->GetDrawVertexDeltas());
+			TargetVizSettings->SetXRayDeltas(SourceVizSettings->GetXRayDeltas());
+			TargetVizSettings->SetHeatMapMax(SourceVizSettings->GetHeatMapMax());
+			TargetVizSettings->SetGroundTruthLerp(SourceVizSettings->GetGroundTruthLerp());
+			TargetVizSettings->SetShowHeatMap(SourceVizSettings->GetShowHeatMap());
+			TargetVizSettings->SetHeatMapMode(SourceVizSettings->GetHeatMapMode());
+			TargetVizSettings->SetDebugBoundsColor(SourceVizSettings->GetDebugBoundsColor());
+			TargetVizSettings->SetCompareActors(SourceVizSettings->GetCompareActors());
+
+			// NOTE: We don't copy the DeformerGraph on purpose, as this model might need another one, so we let it get the default graph.
+		}
+
+		TargetModel->UpdateMemoryUsage();
+		UpdateMemoryUsage();	
+	}
+
 	void FMLDeformerEditorModel::UpdateEditorInputInfo()
 	{
 		InitInputInfo(EditorInputInfo);
diff --git a/Engine/Plugins/Animation/MLDeformer/MLDeformerFramework/Source/MLDeformerFrameworkEditor/Private/MLDeformerEditorToolkit.cpp b/Engine/Plugins/Animation/MLDeformer/MLDeformerFramework/Source/MLDeformerFrameworkEditor/Private/MLDeformerEditorToolkit.cpp
--- a/Engine/Plugins/Animation/MLDeformer/MLDeformerFramework/Source/MLDeformerFrameworkEditor/Private/MLDeformerEditorToolkit.cpp
+++ b/Engine/Plugins/Animation/MLDeformer/MLDeformerFramework/Source/MLDeformerFrameworkEditor/Private/MLDeformerEditorToolkit.cpp
@@ -296,7 +296,6 @@ namespace UE::MLDeformer
 		if (ActiveModel)
 		{
 			ActiveModel->ClearWorldAndPersonaPreviewScene();
-			ActiveModel.Reset();
 
 			if (ModelDetailsView)
 			{
@@ -327,7 +326,11 @@ namespace UE::MLDeformer
 		EditorModel->Init(InitSettings);
 
 		// Tell the editor we use this model now.
-		ActiveModel = TSharedPtr<FMLDeformerEditorModel>(EditorModel);
+		TSharedPtr<FMLDeformerEditorModel> NewModel = TSharedPtr<FMLDeformerEditorModel>(EditorModel);
+		NewModel->CopyBaseSettingsFromModel(ActiveModel.Get());
+
+		ActiveModel.Reset();
+		ActiveModel = NewModel;
 
 		// Create the new scene for this model.
 		if (PersonaToolkit)
diff --git a/Engine/Plugins/Animation/MLDeformer/MLDeformerFramework/Source/MLDeformerFrameworkEditor/Private/MLDeformerGeomCacheEditorModel.cpp b/Engine/Plugins/Animation/MLDeformer/MLDeformerFramework/Source/MLDeformerFrameworkEditor/Private/MLDeformerGeomCacheEditorModel.cpp
--- a/Engine/Plugins/Animation/MLDeformer/MLDeformerFramework/Source/MLDeformerFrameworkEditor/Private/MLDeformerGeomCacheEditorModel.cpp
+++ b/Engine/Plugins/Animation/MLDeformer/MLDeformerFramework/Source/MLDeformerFrameworkEditor/Private/MLDeformerGeomCacheEditorModel.cpp
@@ -35,6 +35,30 @@ namespace UE::MLDeformer
 		}
 	}
 
+	void FMLDeformerGeomCacheEditorModel::CopyBaseSettingsFromModel(const FMLDeformerEditorModel* SourceEditorModel)
+	{
+		UMLDeformerGeomCacheModel* SourceGeomCacheModel = Cast<UMLDeformerGeomCacheModel>(SourceEditorModel->GetModel());
+		if (SourceGeomCacheModel) // If we have a model derived from the UMLDeformerGeomCacheModel
+		{
+			// Copy the training input animations.
+			UMLDeformerGeomCacheModel* TargetModel = Cast<UMLDeformerGeomCacheModel>(GetModel());
+			TArray<FMLDeformerGeomCacheTrainingInputAnim>& TargetInputAnims = TargetModel->GetTrainingInputAnims();
+			TargetInputAnims = SourceGeomCacheModel->GetTrainingInputAnims();
+
+			// Copy the ground truth in the visualization settings.
+			UMLDeformerGeomCacheVizSettings* SourceVizSettings = Cast<UMLDeformerGeomCacheVizSettings>(SourceGeomCacheModel->GetVizSettings());
+			if (SourceVizSettings)
+			{
+				UMLDeformerGeomCacheVizSettings* TargetVizSettings = Cast<UMLDeformerGeomCacheVizSettings>(TargetModel->GetVizSettings());
+				check(TargetVizSettings)
+				TargetVizSettings->SetTestGroundTruth(SourceVizSettings->GetTestGroundTruth());
+			}
+		}
+
+		// Copy the settings of the base class.
+		FMLDeformerEditorModel::CopyBaseSettingsFromModel(SourceEditorModel);
+	}
+
 	TSharedPtr<FMLDeformerSampler> FMLDeformerGeomCacheEditorModel::CreateSamplerObject() const
 	{
 		return MakeShared<FMLDeformerGeomCacheSampler>();
diff --git a/Engine/Plugins/Animation/MLDeformer/MLDeformerFramework/Source/MLDeformerFrameworkEditor/Private/MLDeformerMorphModelEditorModel.cpp b/Engine/Plugins/Animation/MLDeformer/MLDeformerFramework/Source/MLDeformerFrameworkEditor/Private/MLDeformerMorphModelEditorModel.cpp
--- a/Engine/Plugins/Animation/MLDeformer/MLDeformerFramework/Source/MLDeformerFrameworkEditor/Private/MLDeformerMorphModelEditorModel.cpp
+++ b/Engine/Plugins/Animation/MLDeformer/MLDeformerFramework/Source/MLDeformerFrameworkEditor/Private/MLDeformerMorphModelEditorModel.cpp
@@ -32,6 +32,25 @@ namespace UE::MLDeformer
 		return new FMLDeformerMorphModelEditorModel();
 	}
 
+	void FMLDeformerMorphModelEditorModel::CopyBaseSettingsFromModel(const FMLDeformerEditorModel* SourceEditorModel)
+	{
+		// Copy the morph related settings.
+		const UMLDeformerMorphModel* SourceMorphModel = Cast<UMLDeformerMorphModel>(SourceEditorModel->GetModel());
+		if (SourceMorphModel)
+		{
+			UMLDeformerMorphModel* TargetModel = Cast<UMLDeformerMorphModel>(GetModel());
+			check(TargetModel);
+			TargetModel->SetMorphCompressionLevel(SourceMorphModel->GetMorphCompressionLevel());
+			TargetModel->SetMorphDeltaZeroThreshold(SourceMorphModel->GetMorphDeltaZeroThreshold());
+			TargetModel->SetIncludeMorphTargetNormals(SourceMorphModel->GetIncludeMorphTargetNormals());
+			TargetModel->SetMaskChannel(SourceMorphModel->GetMaskChannel());
+			TargetModel->SetInvertMaskChannel(SourceMorphModel->GetInvertMaskChannel());
+		}
+
+		// Copy all base class settings.
+		FMLDeformerGeomCacheEditorModel::CopyBaseSettingsFromModel(SourceEditorModel);
+	}
+
 	bool FMLDeformerMorphModelEditorModel::IsInputMaskingSupported() const
 	{
 		return false;
diff --git a/Engine/Plugins/Animation/MLDeformer/MLDeformerFramework/Source/MLDeformerFrameworkEditor/Public/MLDeformerEditorModel.h b/Engine/Plugins/Animation/MLDeformer/MLDeformerFramework/Source/MLDeformerFrameworkEditor/Public/MLDeformerEditorModel.h
--- a/Engine/Plugins/Animation/MLDeformer/MLDeformerFramework/Source/MLDeformerFrameworkEditor/Public/MLDeformerEditorModel.h
+++ b/Engine/Plugins/Animation/MLDeformer/MLDeformerFramework/Source/MLDeformerFrameworkEditor/Public/MLDeformerEditorModel.h
@@ -94,6 +94,13 @@ namespace UE::MLDeformer
 		virtual FString GetReferencerName() const override { return TEXT("FMLDeformerEditorModel"); }
 		// ~END FGCObject overrides.
 
+		/**
+		 * Copy the common property values from a source model into this model.		 
+		 * This includes the skeletal mesh, input animations, and shared visualization settings.
+		 * @param SourceEditorModel The source model to copy the shared settings from.
+		 */
+		virtual void CopyBaseSettingsFromModel(const FMLDeformerEditorModel* SourceEditorModel);
+
 		/**
 		 * Change the skeletal mesh on a specific skeletal mesh component.
 		 * This internally will not only call SkelMeshComponent->SetSkeletalMesh(Mesh), but also will reassign the materials.
@@ -107,7 +114,6 @@ namespace UE::MLDeformer
 		 */
 		static void ChangeSkeletalMeshOnComponent(USkeletalMeshComponent* SkelMeshComponent, USkeletalMesh* Mesh);
 
-		// Required overrides.
 		/**
 		 * Get the number of training frames.
 		 * This must not be clamped yet to the maximum training frame limit. It must return the full amount of frames available in your training data.
@@ -212,8 +218,28 @@ namespace UE::MLDeformer
 		 */
 		virtual TSharedPtr<FMLDeformerSampler> CreateSamplerObject() const;
 
+		/**
+		 * Get the number of training input animations.
+		 * Each training input anim is a pair of a source animation and a target mesh deformation with the same poses.
+		 * This data is fed into the training process.
+		 * @return The number of training input anims.
+		 */
 		virtual int32 GetNumTrainingInputAnims() const;
+
+		/**
+		 * Get a given training input animation.
+		 * Each training input anim is a pair of a source animation and a target mesh deformation with the same poses.
+		 * This data is fed into the training process.
+		 * @param Index The training input animation to get.
+		 * @see GetNumTrainingInputAnims
+		 */
 		virtual FMLDeformerTrainingInputAnim* GetTrainingInputAnim(int32 Index) const;
+
+		/**
+		 * Update the timeline's training input animation list.
+		 * This is the list of animations that show inside the timeline when in training mode.
+		 * Each training input animation pair needs some name that is shown inside the timeline.
+		 */
 		virtual void UpdateTimelineTrainingAnimList();
 
 		/**
@@ -283,8 +309,25 @@ namespace UE::MLDeformer
 		 */
 		virtual void CreateTestMLDeformedActor(UWorld* World);
 
+		/**
+		 * Create the compare actors for the testing mode.
+		 * @param World The world to create the actors in.
+		 */
 		virtual void CreateTestCompareActors(UWorld* World);
+
+		/**
+		 * Check if this model is compatible with a specific deformer asset.
+		 * For example it will check whether the other deformer asset's skeletal mesh is the same as the one used
+		 * by this model.
+		 * @param Deformer The deformer asset to check against.
+		 * @return Returns true if the specified asset is compatible with this one, otherwise false is returned.
+		 */
 		virtual bool IsCompatibleDeformer(UMLDeformerAsset* Deformer) const;
+
+		/**
+		 * Update the mesh offset factors, which is basically controlling where to place the actors in the scene.
+		 * The spacing between the characters is the mesh offset factor multiplied by the mesh spacing.
+		 */
 		virtual void UpdateMeshOffsetFactors();
 
 		/**
@@ -744,6 +787,11 @@ namespace UE::MLDeformer
 		UE_DEPRECATED(5.4, "Please use GetSamplerForTrainingAnim instead.")
 		FMLDeformerSampler* GetSampler() const;
 
+		/**
+		 * Get the vertex delta sampler for a given training input animation.
+		 * @param AnimIndex The training input animation index.
+		 * @return A pointer to the sampler for this specific input anim, or nullptr in case the index is out of range.
+		 */
 		FMLDeformerSampler* GetSamplerForTrainingAnim(int32 AnimIndex) const	{ return Samplers.IsValidIndex(AnimIndex) ? Samplers[AnimIndex].Get() : nullptr; }
 
 		/**
diff --git a/Engine/Plugins/Animation/MLDeformer/MLDeformerFramework/Source/MLDeformerFrameworkEditor/Public/MLDeformerGeomCacheEditorModel.h b/Engine/Plugins/Animation/MLDeformer/MLDeformerFramework/Source/MLDeformerFrameworkEditor/Public/MLDeformerGeomCacheEditorModel.h
--- a/Engine/Plugins/Animation/MLDeformer/MLDeformerFramework/Source/MLDeformerFrameworkEditor/Public/MLDeformerGeomCacheEditorModel.h
+++ b/Engine/Plugins/Animation/MLDeformer/MLDeformerFramework/Source/MLDeformerFrameworkEditor/Public/MLDeformerGeomCacheEditorModel.h
@@ -28,6 +28,7 @@ namespace UE::MLDeformer
 
 		// FMLDeformerEditorModel overrides.
 		virtual void Init(const InitSettings& Settings) override;
+		virtual void CopyBaseSettingsFromModel(const FMLDeformerEditorModel* SourceEditorModel) override;
 		virtual FMLDeformerEditorActor* CreateEditorActor(const FMLDeformerEditorActor::FConstructSettings& Settings) const override;
 		virtual TSharedPtr<FMLDeformerSampler> CreateSamplerObject() const override;
 		virtual void CreateTrainingGroundTruthActor(UWorld* World) override;
diff --git a/Engine/Plugins/Animation/MLDeformer/MLDeformerFramework/Source/MLDeformerFrameworkEditor/Public/MLDeformerMorphModelEditorModel.h b/Engine/Plugins/Animation/MLDeformer/MLDeformerFramework/Source/MLDeformerFrameworkEditor/Public/MLDeformerMorphModelEditorModel.h
--- a/Engine/Plugins/Animation/MLDeformer/MLDeformerFramework/Source/MLDeformerFrameworkEditor/Public/MLDeformerMorphModelEditorModel.h
+++ b/Engine/Plugins/Animation/MLDeformer/MLDeformerFramework/Source/MLDeformerFrameworkEditor/Public/MLDeformerMorphModelEditorModel.h
@@ -28,6 +28,7 @@ namespace UE::MLDeformer
 		// ~END FGCObject overrides.
 
 		// FMLDeformerEditorModel overrides.
+		virtual void CopyBaseSettingsFromModel(const FMLDeformerEditorModel* SourceEditorModel) override;
 		virtual void Render(const FSceneView* View, FViewport* Viewport, FPrimitiveDrawInterface* PDI) override;
 		virtual void OnPropertyChanged(FPropertyChangedEvent& PropertyChangedEvent) override;
 		virtual FString GetHeatMapDeformerGraphPath() const override;
