[MLDeformer] Fixed a bug where the heatmap in ground truth mode would flicker.

diff --git a/Engine/Plugins/Animation/MLDeformer/MLDeformerFramework/Source/MLDeformerFramework/Private/MLDeformerGraphDebugDataInterface.cpp b/Engine/Plugins/Animation/MLDeformer/MLDeformerFramework/Source/MLDeformerFramework/Private/MLDeformerGraphDebugDataInterface.cpp
--- a/Engine/Plugins/Animation/MLDeformer/MLDeformerFramework/Source/MLDeformerFramework/Private/MLDeformerGraphDebugDataInterface.cpp
+++ b/Engine/Plugins/Animation/MLDeformer/MLDeformerFramework/Source/MLDeformerFramework/Private/MLDeformerGraphDebugDataInterface.cpp
@@ -130,9 +130,12 @@ FComputeDataProviderRenderProxy* UMLDeformerGraphDebugDataProvider::GetRenderPro
 	if (DeformerComponent && DeformerAsset && DeformerComponent->GetModelInstance() && DeformerComponent->GetModelInstance()->IsValidForDataProvider())
 	{
 		UE::MLDeformer::FMLDeformerGraphDebugDataProviderProxy* Proxy = new UE::MLDeformer::FMLDeformerGraphDebugDataProviderProxy(DeformerComponent, DeformerAsset, this);
-		const float SampleTime = DeformerComponent->GetModelInstance()->GetSkeletalMeshComponent()->GetPosition();
 		UMLDeformerModel* Model = DeformerAsset->GetModel();
-		Model->SampleGroundTruthPositions(SampleTime, Proxy->GetGroundTruthPositions());
+		if (Model->GetVizSettings()->GetHeatMapMode() == EMLDeformerHeatMapMode::GroundTruth)
+		{
+			const float SampleTime = DeformerComponent->GetModelInstance()->GetSkeletalMeshComponent()->GetPosition();
+			Model->SampleGroundTruthPositions(SampleTime, Proxy->GetGroundTruthPositions());
+		}
 		Proxy->HandleZeroGroundTruthPositions();
 		return Proxy;
 	}
diff --git a/Engine/Plugins/Animation/MLDeformer/MLDeformerFramework/Source/MLDeformerFrameworkEditor/Private/MLDeformerEditorModel.cpp b/Engine/Plugins/Animation/MLDeformer/MLDeformerFramework/Source/MLDeformerFrameworkEditor/Private/MLDeformerEditorModel.cpp
--- a/Engine/Plugins/Animation/MLDeformer/MLDeformerFramework/Source/MLDeformerFrameworkEditor/Private/MLDeformerEditorModel.cpp
+++ b/Engine/Plugins/Animation/MLDeformer/MLDeformerFramework/Source/MLDeformerFrameworkEditor/Private/MLDeformerEditorModel.cpp
@@ -910,6 +910,7 @@ namespace UE::MLDeformer
 			SkeletalMeshComponent->SetPosition(CurrentPlayTime);
 			SkeletalMeshComponent->SetPlayRate(TestAnimSpeed);
 			SkeletalMeshComponent->Play(true);
+			UpdateStepInterpolationMode();
 		}
 
 		// Update all compare actors.
@@ -1200,6 +1201,11 @@ namespace UE::MLDeformer
 		{
 			SetHeatMapMaterialEnabled(Model->GetVizSettings()->GetShowHeatMap());
 			UpdateDeformerGraph();
+			UpdateStepInterpolationMode();
+		}
+		else if (Property->GetFName() == UMLDeformerVizSettings::GetHeatMapModePropertyName())
+		{
+			UpdateStepInterpolationMode();
 		}
 		else
 		if (Property->GetFName() == UMLDeformerVizSettings::GetDrawLinearSkinnedActorPropertyName() ||
@@ -3161,6 +3167,28 @@ namespace UE::MLDeformer
 			}
 		}
 	}
+
+	void FMLDeformerEditorModel::UpdateStepInterpolationMode()
+	{
+		// Force step interpolation when viewing the heatmap in ground truth mode.
+		UMLDeformerVizSettings* VizSettings = Model->GetVizSettings();
+		FMLDeformerEditorActor* MLDeformerEditorActor = static_cast<FMLDeformerEditorActor*>(FindEditorActor(ActorID_Test_MLDeformed));
+		if (MLDeformerEditorActor)
+		{
+			USkeletalMeshComponent* SkelMeshComponent = MLDeformerEditorActor->GetSkeletalMeshComponent();
+			if (SkelMeshComponent)
+			{
+				UAnimSingleNodeInstance* AnimInstance = Cast<UAnimSingleNodeInstance>(SkelMeshComponent->GetAnimInstance());
+				if (AnimInstance)
+				{
+					const bool bForceStepInterpolation = VizSettings->GetShowHeatMap() && (VizSettings->GetHeatMapMode() == EMLDeformerHeatMapMode::GroundTruth) && VizSettings->HasTestGroundTruth();
+					const TOptional<EAnimInterpolationType> InterpolationOverride = bForceStepInterpolation ? EAnimInterpolationType::Step : TOptional<EAnimInterpolationType>();
+					AnimInstance->SetInterpolationOverride(InterpolationOverride);
+				}
+			}
+		}
+	}
+
 }	// namespace UE::MLDeformer
 
 #undef LOCTEXT_NAMESPACE
diff --git a/Engine/Plugins/Animation/MLDeformer/MLDeformerFramework/Source/MLDeformerFrameworkEditor/Public/MLDeformerEditorModel.h b/Engine/Plugins/Animation/MLDeformer/MLDeformerFramework/Source/MLDeformerFrameworkEditor/Public/MLDeformerEditorModel.h
--- a/Engine/Plugins/Animation/MLDeformer/MLDeformerFramework/Source/MLDeformerFrameworkEditor/Public/MLDeformerEditorModel.h
+++ b/Engine/Plugins/Animation/MLDeformer/MLDeformerFramework/Source/MLDeformerFrameworkEditor/Public/MLDeformerEditorModel.h
@@ -1080,6 +1080,13 @@ namespace UE::MLDeformer
 
 		void UpdateLODMappings();
 
+		/**
+		 * Update whether we force to use step interpolation on the ML Deformed model or not.
+		 * Step interpolation is needed when we are enabled heatmaps and have a ground truth mesh and are in ground truth heatmap mode.
+		 * This will modify the AnimInstance's ForceStepInterpolation state on just the ML Deformed model. It can also disable it when needed.
+		 */
+		void UpdateStepInterpolationMode();
+
 	protected:
 		struct FLODInfo
 		{
diff --git a/Engine/Plugins/Animation/MLDeformer/VertexDeltaModel/Source/VertexDeltaModel/Private/VertexDeltaGraphDebugDataInterface.cpp b/Engine/Plugins/Animation/MLDeformer/VertexDeltaModel/Source/VertexDeltaModel/Private/VertexDeltaGraphDebugDataInterface.cpp
--- a/Engine/Plugins/Animation/MLDeformer/VertexDeltaModel/Source/VertexDeltaModel/Private/VertexDeltaGraphDebugDataInterface.cpp
+++ b/Engine/Plugins/Animation/MLDeformer/VertexDeltaModel/Source/VertexDeltaModel/Private/VertexDeltaGraphDebugDataInterface.cpp
@@ -130,9 +130,12 @@ FComputeDataProviderRenderProxy* UDEPRECATED_VertexDeltaGraphDebugDataProvider::
 	if (DeformerComponent && DeformerAsset && DeformerComponent->GetModelInstance() && DeformerComponent->GetModelInstance()->IsValidForDataProvider())
 	{
 		UE::VertexDeltaModel::FVertexDeltaGraphDebugDataProviderProxy* Proxy = new UE::VertexDeltaModel::FVertexDeltaGraphDebugDataProviderProxy(DeformerComponent, DeformerAsset, this);
-		const float SampleTime = DeformerComponent->GetModelInstance()->GetSkeletalMeshComponent()->GetPosition();
 		UMLDeformerModel* Model = DeformerAsset->GetModel();
-		Model->SampleGroundTruthPositions(SampleTime, Proxy->GetGroundTruthPositions());
+		if (Model->GetVizSettings()->GetHeatMapMode() == EMLDeformerHeatMapMode::GroundTruth)
+		{
+			const float SampleTime = DeformerComponent->GetModelInstance()->GetSkeletalMeshComponent()->GetPosition();
+			Model->SampleGroundTruthPositions(SampleTime, Proxy->GetGroundTruthPositions());
+		}
 		Proxy->HandleZeroGroundTruthPositions();
 		return Proxy;
 	}
