[MLDeformer] Fixed a little bug when switching model types. It didn't auto select the correct input anim in the timeline.

diff --git a/Engine/Plugins/Animation/MLDeformer/MLDeformerFramework/Source/MLDeformerFrameworkEditor/Private/MLDeformerEditorModel.cpp b/Engine/Plugins/Animation/MLDeformer/MLDeformerFramework/Source/MLDeformerFrameworkEditor/Private/MLDeformerEditorModel.cpp
--- a/Engine/Plugins/Animation/MLDeformer/MLDeformerFramework/Source/MLDeformerFrameworkEditor/Private/MLDeformerEditorModel.cpp
+++ b/Engine/Plugins/Animation/MLDeformer/MLDeformerFramework/Source/MLDeformerFrameworkEditor/Private/MLDeformerEditorModel.cpp
@@ -223,6 +223,8 @@ namespace UE::MLDeformer
 			// NOTE: We don't copy the DeformerGraph on purpose, as this model might need another one, so we let it get the default graph.
 		}
 
+		ActiveTrainingInputAnimIndex = SourceEditorModel->GetActiveTrainingInputAnimIndex();
+
 		TargetModel->UpdateMemoryUsage();
 		UpdateMemoryUsage();	
 	}
