[MLDeformer] Fixed a crash when having an ML Deformer asset open, with HeatMap enabled, while untrained in testing mode and reloading the skeletal mesh it uses through "Asset Actions -> Reload".

diff --git a/Engine/Plugins/Animation/MLDeformer/MLDeformerFramework/Source/MLDeformerFramework/Private/MLDeformerGraphDebugDataInterface.cpp b/Engine/Plugins/Animation/MLDeformer/MLDeformerFramework/Source/MLDeformerFramework/Private/MLDeformerGraphDebugDataInterface.cpp
--- a/Engine/Plugins/Animation/MLDeformer/MLDeformerFramework/Source/MLDeformerFramework/Private/MLDeformerGraphDebugDataInterface.cpp
+++ b/Engine/Plugins/Animation/MLDeformer/MLDeformerFramework/Source/MLDeformerFramework/Private/MLDeformerGraphDebugDataInterface.cpp
@@ -163,7 +163,7 @@ namespace UE::MLDeformer
 			VertexMapBufferSRV = Model->GetVertexMapBuffer().ShaderResourceViewRHI;
 			HeatMapMode = static_cast<int32>(VizSettings->GetHeatMapMode());
 			HeatMapMax = 1.0f / FMath::Max(VizSettings->GetHeatMapMax(), 0.00001f);
-			GroundTruthLerp = (ModelInstance && ModelInstance->GetSkeletalMeshComponent()->GetPredictedLODLevel() == 0) ? VizSettings->GetGroundTruthLerp() : 0.0f;
+			GroundTruthLerp = (ModelInstance && ModelInstance->GetSkeletalMeshComponent() && ModelInstance->GetSkeletalMeshComponent()->GetPredictedLODLevel() == 0) ? VizSettings->GetGroundTruthLerp() : 0.0f;
 		}
 	}
 
