UMLDeformerMorphModel: Fix indeterminism during cook by clearing MorphTargetWeights before they are serialized the first time during cook,so that their property name is not harvested by PackageHarvester because the empty array is equal to the class default object's empty array and the property serialization is skipped. Also restore the MorphTargetWeights after the Serialize call is done, because cook savepackage by contract is not supposed to modify the editor version of the data without restoring it.
#rnx

diff --git a/Engine/Plugins/Animation/MLDeformer/MLDeformerFramework/Source/MLDeformerFramework/Private/MLDeformerMorphModel.cpp b/Engine/Plugins/Animation/MLDeformer/MLDeformerFramework/Source/MLDeformerFramework/Private/MLDeformerMorphModel.cpp
--- a/Engine/Plugins/Animation/MLDeformer/MLDeformerFramework/Source/MLDeformerFramework/Private/MLDeformerMorphModel.cpp
+++ b/Engine/Plugins/Animation/MLDeformer/MLDeformerFramework/Source/MLDeformerFramework/Private/MLDeformerMorphModel.cpp
@@ -25,37 +25,62 @@ void UMLDeformerMorphModel::Serialize(FArchive& Archive)
 {
 	TRACE_CPUPROFILER_EVENT_SCOPE(UMLDeformerMorphModel::Serialize)
 
-	Super::Serialize(Archive);
-	Archive.UsingCustomVersion(UE::MLDeformer::FMLDeformerObjectVersion::GUID);
+	bool bModifiedPropertiesForCook = false;
+	TArray<FVector3f> SavedMorphTargetDeltas;
+	auto ModifyPropertiesForCook = [this, &bModifiedPropertiesForCook, &SavedMorphTargetDeltas]()
+		{
+			bModifiedPropertiesForCook = true;
+			SavedMorphTargetDeltas = MoveTemp(MorphTargetDeltas);
+			MorphTargetDeltas.Empty();
+		};
+	auto RestorePropertiesForCook = [this, &bModifiedPropertiesForCook, &SavedMorphTargetDeltas]()
+		{
+			if (!bModifiedPropertiesForCook)
+			{
+				return;
+			}
+			MorphTargetDeltas = MoveTemp(SavedMorphTargetDeltas);
+		};
+	ON_SCOPE_EXIT
+	{
+		RestorePropertiesForCook();
+	};
 
-	// Check if we have initialized our compressed morph buffers.
-	bool bHasMorphData = false;
+	int32 NumSaveLODs = 0;
 	if (Archive.IsSaving())
 	{
 		// Strip editor only data on cook.
-		int32 NumLODs = GetNumLODs();
+		NumSaveLODs = GetNumLODs();
 		if (Archive.IsCooking())
-		{			
-			MorphTargetDeltas.Empty();
+		{
+			ModifyPropertiesForCook();
 
 			// Check if we want to limit the number of LODs (can be per platform/device).
 			UE::MLDeformer::FMLDeformerModule& MLDeformerModule = FModuleManager::LoadModuleChecked<UE::MLDeformer::FMLDeformerModule>("MLDeformerFramework");
 			const int32 MaxLODLevels = FMath::Clamp(MLDeformerModule.GetMaxLODLevelsOnCookCVar().GetInt(), 1, 1000);	// Limit to 1000 LODs, which should never be reached.
 
 			// Get lowest value between what we generated, console variable and the UI/property max lods value.
-			NumLODs = FMath::Min3(NumLODs, MaxLODLevels, GetMaxNumLODs());
+			NumSaveLODs = FMath::Min3(NumSaveLODs, MaxLODLevels, GetMaxNumLODs());
 
-			UE_LOG(LogMLDeformer, Display, TEXT("Cooking MLD asset '%s' with %d LOD levels"), *GetFullName(), NumLODs);
+			UE_LOG(LogMLDeformer, Display, TEXT("Cooking MLD asset '%s' with %d LOD levels"), *GetFullName(), NumSaveLODs);
 		}
 		else
 		{
 			// Get lowest number between how many LODs we have generated and the number of LODs we setup in the UI/Property.
-			NumLODs = FMath::Min(NumLODs, GetMaxNumLODs());
+			NumSaveLODs = FMath::Min(NumSaveLODs, GetMaxNumLODs());
 		}
+	}
+
+	Super::Serialize(Archive);
+	Archive.UsingCustomVersion(UE::MLDeformer::FMLDeformerObjectVersion::GUID);
 
+	// Check if we have initialized our compressed morph buffers.
+	bool bHasMorphData = false;
+	if (Archive.IsSaving())
+	{
 		// Save all LOD levels, strip out LODs we don't want.
-		Archive << NumLODs;
-		for (int32 LOD = 0; LOD < NumLODs; ++LOD)
+		Archive << NumSaveLODs;
+		for (int32 LOD = 0; LOD < NumSaveLODs; ++LOD)
 		{
 			bHasMorphData = GetMorphTargetSet(LOD).IsValid() ? GetMorphTargetSet(LOD)->MorphBuffers.IsMorphCPUDataValid() : false;
 			Archive << bHasMorphData;
