[MLDeformer] Add support for painted masks using weight maps using the skeletal mesh attributes. Maps can now also be painted for bone and bone group masks in the Neural Morph Model local mode.

Directly fixed an issue when switching from global to local mode after having trained in global mode. This had to resample data.

#virtualized

diff --git a/Engine/Plugins/Animation/MLDeformer/MLDeformerFramework/MLDeformerFramework.uplugin b/Engine/Plugins/Animation/MLDeformer/MLDeformerFramework/MLDeformerFramework.uplugin
--- a/Engine/Plugins/Animation/MLDeformer/MLDeformerFramework/MLDeformerFramework.uplugin
+++ b/Engine/Plugins/Animation/MLDeformer/MLDeformerFramework/MLDeformerFramework.uplugin
@@ -42,6 +42,16 @@
 		{
 			"Name": "GameplayInsights",
 			"Enabled": true
+		},
+		{
+			"Name": "ModelingToolsEditorMode",
+			"Enabled": true,
+			"TargetAllowList": [ "Editor" ]
+		},
+		{
+			"Name": "SkeletalMeshModelingTools",
+			"Enabled": true,
+			"TargetAllowList": [ "Editor" ]
 		}
 	]
 }
\ No newline at end of file
diff --git a/Engine/Plugins/Animation/MLDeformer/MLDeformerFramework/Source/MLDeformerFramework/Private/MLDeformerModel.cpp b/Engine/Plugins/Animation/MLDeformer/MLDeformerFramework/Source/MLDeformerFramework/Private/MLDeformerModel.cpp
--- a/Engine/Plugins/Animation/MLDeformer/MLDeformerFramework/Source/MLDeformerFramework/Private/MLDeformerModel.cpp
+++ b/Engine/Plugins/Animation/MLDeformer/MLDeformerFramework/Source/MLDeformerFramework/Private/MLDeformerModel.cpp
@@ -267,6 +267,36 @@ bool UMLDeformerModel::IsCompatibleDebugActor(const AActor* Actor, UMLDeformerCo
 
 
 #if WITH_EDITOR
+	TArray<FName> UMLDeformerModel::GetVertexAttributeNames() const
+	{
+		const USkeletalMesh* SkelMesh = GetSkeletalMesh();
+		if (!SkelMesh)
+		{
+			return {};
+		}
+
+		const int32 LODIndex = 0;
+		if (!SkelMesh->HasMeshDescription(LODIndex))
+		{
+			return {};
+		}
+
+		const FMeshDescription* MeshDescription = SkelMesh->GetMeshDescription(LODIndex);		
+		const TAttributesSet<FVertexID> AttributesSet = MeshDescription->VertexAttributes();
+
+		TArray<FName> AttributeNames;
+		AttributesSet.ForEach([&AttributeNames, &AttributesSet](const FName AttributeName, const auto AttributesRef)
+		{
+			const bool bIsAutoGenerated = (AttributesRef.GetFlags() & EMeshAttributeFlags::AutoGenerated) != EMeshAttributeFlags::None;
+			if (!bIsAutoGenerated && AttributesSet.HasAttributeOfType<float>(AttributeName))
+			{
+				AttributeNames.Add(AttributeName);
+			}
+		});
+
+		return AttributeNames;
+	}
+
 	void UMLDeformerModel::UpdateNumTargetMeshVertices()
 	{ 
 		NumTargetMeshVerts = 0;
diff --git a/Engine/Plugins/Animation/MLDeformer/MLDeformerFramework/Source/MLDeformerFramework/Private/MLDeformerMorphModel.cpp b/Engine/Plugins/Animation/MLDeformer/MLDeformerFramework/Source/MLDeformerFramework/Private/MLDeformerMorphModel.cpp
--- a/Engine/Plugins/Animation/MLDeformer/MLDeformerFramework/Source/MLDeformerFramework/Private/MLDeformerMorphModel.cpp
+++ b/Engine/Plugins/Animation/MLDeformer/MLDeformerFramework/Source/MLDeformerFramework/Private/MLDeformerMorphModel.cpp
@@ -9,6 +9,7 @@
 #include "MLDeformerModule.h"
 #include "Components/ExternalMorphSet.h"
 #include "Components/SkeletalMeshComponent.h"
+#include "Engine/SkeletalMesh.h"
 #include "Rendering/MorphTargetVertexInfoBuffers.h"
 #include "UObject/AssetRegistryTagsContext.h"
 
@@ -186,7 +187,8 @@ int32 UMLDeformerMorphModel::GetNumMorphTargets(int32 LOD) const
 bool UMLDeformerMorphModel::CanDynamicallyUpdateMorphTargets() const
 {
 	const int32 LOD = 0;
-	return GetMorphTargetDeltas().Num() == (GetNumBaseMeshVerts() * GetNumMorphTargets(LOD));
+	const bool bResult = (GetMorphTargetDeltas().Num() == (GetNumBaseMeshVerts() * GetNumMorphTargets(LOD))) || bHasOnlyEmptyMorphs;
+	return bResult;
 }
 
 UMLDeformerModelInstance* UMLDeformerMorphModel::CreateModelInstance(UMLDeformerComponent* Component)
@@ -346,6 +348,14 @@ int32 UMLDeformerMorphModel::GetNumActiveMorphs(int32 QualityLevel) const
 	return 0;
 }
 
+#if WITH_EDITORONLY_DATA
+FName UMLDeformerMorphModel::GetGlobalMaskAttributeName() const
+{
+	return VertexAttributeName;
+}
+#endif
+
+
 #if WITH_EDITOR
 void UMLDeformerMorphModel::UpdateMemoryUsage()
 {
diff --git a/Engine/Plugins/Animation/MLDeformer/MLDeformerFramework/Source/MLDeformerFramework/Public/MLDeformerMasking.h b/Engine/Plugins/Animation/MLDeformer/MLDeformerFramework/Source/MLDeformerFramework/Public/MLDeformerMasking.h
new file mode 100644
--- /dev/null
+++ b/Engine/Plugins/Animation/MLDeformer/MLDeformerFramework/Source/MLDeformerFramework/Public/MLDeformerMasking.h
@@ -0,0 +1,43 @@
+// Copyright Epic Games, Inc. All Rights Reserved.
+#pragma once
+
+#include "CoreTypes.h"
+#include "MLDeformerMasking.generated.h"
+
+
+/** Specifies whether a mask should be generated or not. */
+UENUM()
+enum class EMLDeformerMaskingMode : uint8
+{
+	/** Use a generated mask, generated from the skinning weights of a set of bones. */
+	Generated,
+
+	/** Use a vertex attribute on the skeletal mesh. This can be a manually painted map. */
+	VertexAttribute
+};
+
+
+/**
+ * Information needed to generate the mask for a specific bone.
+ * This includes a list of bone names. Each bone has a skinning influence mask.
+ * The final mask will be a merge of all the bone masks of bones listed inside this info struct.
+ * There will be an array of these structs, one for each bone.
+ */
+USTRUCT()
+struct MLDEFORMERFRAMEWORK_API FMLDeformerMaskInfo
+{
+	GENERATED_BODY()
+
+public:
+	/** The list of bone names that should be included in the mask generation. */
+	UPROPERTY()
+	TArray<FName> BoneNames;
+
+	/** The masking mode. */
+	UPROPERTY()
+	EMLDeformerMaskingMode MaskMode = EMLDeformerMaskingMode::Generated;
+
+	/** If the masking mode is set to VertexAttribute then we can check which attribute to use using this member. */
+	UPROPERTY()
+	FName VertexAttributeName;
+};
diff --git a/Engine/Plugins/Animation/MLDeformer/MLDeformerFramework/Source/MLDeformerFramework/Public/MLDeformerModel.h b/Engine/Plugins/Animation/MLDeformer/MLDeformerFramework/Source/MLDeformerFramework/Public/MLDeformerModel.h
--- a/Engine/Plugins/Animation/MLDeformer/MLDeformerFramework/Source/MLDeformerFramework/Public/MLDeformerModel.h
+++ b/Engine/Plugins/Animation/MLDeformer/MLDeformerFramework/Source/MLDeformerFramework/Public/MLDeformerModel.h
@@ -38,7 +38,10 @@ enum class EMLDeformerMaskChannel : uint8
 	VertexColorBlue,
 
 	/** The alpha vertex color channel. */
-	VertexColorAlpha
+	VertexColorAlpha,
+
+	/** Use a set of vertex attributes on the skeletal mesh. You can create and editor those using the Skeletal Mesh editor plugin. */
+	VertexAttribute
 };
 
 
@@ -80,6 +83,7 @@ namespace UE::MLDeformer
 		/** The array of integers we want to store on the GPU. This buffer will be emptied after successfully calling InitRHI. */
 		TArray<int32> VertexMap;
 	};
+
 }	// namespace UE::MLDeformer
 
 // DEPRECATED: Use FMLDeformerReinitModelInstancesDelegate instead.
@@ -254,6 +258,10 @@ public:
 #endif
 
 #if WITH_EDITOR
+	/** Get the attribute names of all vertex attributes on the skeletal mesh. This only includes ones that aren't auto generated by the engine. */
+	UFUNCTION()
+	TArray<FName> GetVertexAttributeNames() const;
+
 	/**
 	 * Update the cached number of vertices of both base and target meshes.
 	 */ 
diff --git a/Engine/Plugins/Animation/MLDeformer/MLDeformerFramework/Source/MLDeformerFramework/Public/MLDeformerMorphModel.h b/Engine/Plugins/Animation/MLDeformer/MLDeformerFramework/Source/MLDeformerFramework/Public/MLDeformerMorphModel.h
--- a/Engine/Plugins/Animation/MLDeformer/MLDeformerFramework/Source/MLDeformerFramework/Public/MLDeformerMorphModel.h
+++ b/Engine/Plugins/Animation/MLDeformer/MLDeformerFramework/Source/MLDeformerFramework/Public/MLDeformerMorphModel.h
@@ -59,6 +59,7 @@ public:
 	bool GetIncludeMorphTargetNormals() const						{ return bIncludeNormals; }
 	EMLDeformerMaskChannel GetMaskChannel() const					{ return MaskChannel; }
 	bool GetInvertMaskChannel() const								{ return bInvertMaskChannel; }
+	bool GetHasOnlyEmptyMorphs() const								{ return bHasOnlyEmptyMorphs; }
 
 	UFUNCTION(BlueprintPure, Category = "MLDeformerMorphModel")
 	bool CanDynamicallyUpdateMorphTargets() const;
@@ -69,6 +70,7 @@ public:
 	void SetMaskChannel(EMLDeformerMaskChannel Channel)				{ MaskChannel = Channel; }
 	void SetInvertMaskChannel(bool bInvert)							{ bInvertMaskChannel = bInvert; }
 	void SetClampMorphTargetsWeights(bool bEnabled)					{ bClampMorphWeights = bEnabled; }
+	void SetHasOnlyEmptyMorphs(bool bOnlyEmpty)						{ bHasOnlyEmptyMorphs = bOnlyEmpty; }
 
 	UE_DEPRECATED(5.5, "Please use SetMorphCompressionLevel, with upper case L for Level.")
 	void SetMorphCompressionlevel(float Tolerance)					{ MorphCompressionLevel = Tolerance; }
@@ -98,6 +100,10 @@ public:
 	static FName GetUncompressedMorphDataSizeInBytesPropertyName()	{ return GET_MEMBER_NAME_CHECKED(UMLDeformerMorphModel, UncompressedMorphDataSizeInBytes); }
 	static FName GetClampMorphTargetWeightsPropertyName()			{ return GET_MEMBER_NAME_CHECKED(UMLDeformerMorphModel, bClampMorphWeights); }
 
+#if WITH_EDITORONLY_DATA
+	static FName GetGlobalMaskAttributePropertyName()				{ return GET_MEMBER_NAME_CHECKED(UMLDeformerMorphModel, VertexAttributeName); }
+#endif
+
 	UE_DEPRECATED(5.4, "This method will be removed.")
 	static FName GetQualityLevelsPropertyName()						{ return GET_MEMBER_NAME_CHECKED(UMLDeformerMorphModel, QualityLevels_DEPRECATED); }
 
@@ -287,6 +293,11 @@ public:
 	UE_DEPRECATED(5.4, "Please use the GetNumActiveMorphsForLOD.")
 	int32 GetNumActiveMorphs(int32 QualityLevel) const;
 
+#if WITH_EDITORONLY_DATA
+	/** Get the name of the attribute used for the global mask. */
+	FName GetGlobalMaskAttributeName() const;
+#endif
+
 private:
 	/** The compressed morph target data, ready for the GPU. */
 	TSharedPtr<FExternalMorphSet> MorphTargetSet_DEPRECATED;
@@ -384,9 +395,15 @@ private:
 	 * You can use this feather out influence of the ML Deformer in specific areas, such as neck line seams, where the head mesh connects with the body.
 	 * The painted vertex color values will be like a weight multiplier on the ML deformer deltas applied to that vertex. You can invert the mask as well.
 	 */
-	UPROPERTY(EditAnywhere, Category = "Morph Targets", meta = (EditCondition = "CanDynamicallyUpdateMorphTargets()"))
+	UPROPERTY(EditAnywhere, Category = "Morph Targets", meta = (EditCondition = "CanDynamicallyUpdateMorphTargets"))
 	EMLDeformerMaskChannel MaskChannel = EMLDeformerMaskChannel::Disabled;
 
+#if WITH_EDITORONLY_DATA
+	/** The global vertex attribute attribute name. This is an attribute on the skeletal mesh, which can be created using the skeletal mesh editor. */
+	UPROPERTY(EditAnywhere, Category = "Morph Targets", meta = (EditCondition = "CanDynamicallyUpdateMorphTargets() && MaskChannel == EMLDeformerMaskChannel::VertexAttribute", GetOptions = "GetVertexAttributeNames", NoResetToDefault))
+	FName VertexAttributeName;
+#endif
+
 	/** 
 	 * Enable this if you want to invert the mask channel values. For example if you painted the neck seam vertices in red, and you wish the vertices that got painted to NOT move, you have to invert the mask.
 	 * On default you paint areas where the deformer should be active. If you enable the invert option, you paint areas where the deformer will not be active.
@@ -396,4 +413,7 @@ private:
 
 	/** The fence that let's us wait for all render commands to finish, before this instance is destroyed. */
 	FRenderCommandFence RenderCommandFence;
+
+	/** Set to true when all the morph targets got their deltas filtered out. This can happen when your global mask is all 0's for example. */
+	bool bHasOnlyEmptyMorphs = true;
 };
diff --git a/Engine/Plugins/Animation/MLDeformer/MLDeformerFramework/Source/MLDeformerFramework/Public/MLDeformerObjectVersion.h b/Engine/Plugins/Animation/MLDeformer/MLDeformerFramework/Source/MLDeformerFramework/Public/MLDeformerObjectVersion.h
--- a/Engine/Plugins/Animation/MLDeformer/MLDeformerFramework/Source/MLDeformerFramework/Public/MLDeformerObjectVersion.h
+++ b/Engine/Plugins/Animation/MLDeformer/MLDeformerFramework/Source/MLDeformerFramework/Public/MLDeformerObjectVersion.h
@@ -19,6 +19,9 @@ namespace UE::MLDeformer
 			// We added LOD support.
 			LODSupportAdded,
 
+			// We deprecated the UNeuralMorphMaskInfo and should use FMLDeformerMaskInfo instead.
+			MaskInfoMovedToFramework,
+
 			// ----- New versions can be added above this line -----
 			VersionPlusOne,
 			LatestVersion = VersionPlusOne - 1
diff --git a/Engine/Plugins/Animation/MLDeformer/MLDeformerFramework/Source/MLDeformerFrameworkEditor/MLDeformerFrameworkEditor.Build.cs b/Engine/Plugins/Animation/MLDeformer/MLDeformerFramework/Source/MLDeformerFrameworkEditor/MLDeformerFrameworkEditor.Build.cs
--- a/Engine/Plugins/Animation/MLDeformer/MLDeformerFramework/Source/MLDeformerFrameworkEditor/MLDeformerFrameworkEditor.Build.cs
+++ b/Engine/Plugins/Animation/MLDeformer/MLDeformerFramework/Source/MLDeformerFrameworkEditor/MLDeformerFrameworkEditor.Build.cs
@@ -54,6 +54,20 @@ namespace UnrealBuildTool.Rules
 					"RHI",
 					"AssetDefinition",
 					"SkeletalMeshDescription",
+
+					"EditorInteractiveToolsFramework",
+					"InteractiveToolsFramework",
+					"MeshModelingTools",
+					"MeshModelingToolsEditorOnly",
+					"MeshModelingToolsEditorOnlyExp",
+					"MeshModelingToolsExp",
+					"ModelingComponentsEditorOnly",
+					"ModelingToolsEditorMode",
+					"ModelingComponents",
+					"SkeletalMeshEditor",
+					"SkeletalMeshUtilitiesCommon",
+					"WidgetRegistration",
+					"StatusBar"
 				}
 			);
 
diff --git a/Engine/Plugins/Animation/MLDeformer/MLDeformerFramework/Source/MLDeformerFrameworkEditor/Private/MLDeformerApplicationMode.cpp b/Engine/Plugins/Animation/MLDeformer/MLDeformerFramework/Source/MLDeformerFrameworkEditor/Private/MLDeformerApplicationMode.cpp
--- a/Engine/Plugins/Animation/MLDeformer/MLDeformerFramework/Source/MLDeformerFrameworkEditor/Private/MLDeformerApplicationMode.cpp
+++ b/Engine/Plugins/Animation/MLDeformer/MLDeformerFramework/Source/MLDeformerFrameworkEditor/Private/MLDeformerApplicationMode.cpp
@@ -35,6 +35,7 @@ namespace UE::MLDeformer
 		// Register Persona tabs.
 		FPersonaModule& PersonaModule = FModuleManager::LoadModuleChecked<FPersonaModule>("Persona");
 		TabFactories.RegisterFactory(PersonaModule.CreatePersonaViewportTabFactory(InHostingApp, ViewportArgs));
+		TabFactories.RegisterFactory(PersonaModule.CreatePersonaToolboxTabFactory(MLDeformerEditor));
 		TabFactories.RegisterFactory(PersonaModule.CreateDetailsTabFactory(InHostingApp, FOnDetailsCreated::CreateSP(&MLDeformerEditor.Get(), &FMLDeformerEditorToolkit::HandleDetailsCreated)));
 
 		// Register custom tabs.
@@ -49,7 +50,7 @@ namespace UE::MLDeformer
 		}
 
 		// Create tab layout.
-		TabLayout = FTabManager::NewLayout("Standalone_MLDeformerEditor_Layout_v5.3.001")
+		TabLayout = FTabManager::NewLayout("Standalone_MLDeformerEditor_Layout_v5.5.001")
 			->AddArea
 			(
 				FTabManager::NewPrimaryArea()
@@ -64,6 +65,8 @@ namespace UE::MLDeformer
 						FTabManager::NewStack()
 						->SetSizeCoefficient(0.3f)
 						->AddTab(FMLDeformerVizSettingsTabSummoner::TabID, ETabState::OpenedTab)
+						->AddTab(FPersonaTabs::ToolboxID, ETabState::OpenedTab)
+						->SetForegroundTab(FMLDeformerVizSettingsTabSummoner::TabID)
 					)
 					->Split
 					(
diff --git a/Engine/Plugins/Animation/MLDeformer/MLDeformerFramework/Source/MLDeformerFrameworkEditor/Private/MLDeformerEditorModel.cpp b/Engine/Plugins/Animation/MLDeformer/MLDeformerFramework/Source/MLDeformerFrameworkEditor/Private/MLDeformerEditorModel.cpp
--- a/Engine/Plugins/Animation/MLDeformer/MLDeformerFramework/Source/MLDeformerFrameworkEditor/Private/MLDeformerEditorModel.cpp
+++ b/Engine/Plugins/Animation/MLDeformer/MLDeformerFramework/Source/MLDeformerFrameworkEditor/Private/MLDeformerEditorModel.cpp
@@ -16,6 +16,7 @@
 #include "MLDeformerModelInstance.h"
 #include "MLDeformerTrainingInputAnim.h"
 #include "MLDeformerOctree.h"
+#include "MLDeformerPaintMode.h"
 #include "AnimationEditorPreviewActor.h"
 #include "AnimationEditorViewportClient.h"
 #include "EditorModeManager.h"
@@ -712,7 +713,6 @@ namespace UE::MLDeformer
 		UpdateActorTransforms();
 		UpdateLabels();
 		UpdateActorLODs();
-		CheckTrainingDataFrameChanged();
 		ApplyDebugActorTransforms();
 
 		// Debug draw elements inside the PIE viewport when PIE is active.
@@ -806,6 +806,8 @@ namespace UE::MLDeformer
 		const bool bShowTrainingData = (VizSettings->GetVisualizationMode() == EMLDeformerVizMode::TrainingData);
 		const bool bShowTestData = (VizSettings->GetVisualizationMode() == EMLDeformerVizMode::TestData);
 		const bool bIsDebugging = (GetEditor()->GetDebugActor() != nullptr);
+		const bool bInDefaultMode = GetEditor()->IsDefaultModeActive();
+
 		for (FMLDeformerEditorActor* EditorActor : EditorActors)
 		{
 			if (EditorActor)
@@ -829,6 +831,7 @@ namespace UE::MLDeformer
 				}
 
 				bIsVisible &= EditorActor->HasVisualMesh();
+				bIsVisible &= bInDefaultMode;
 
 				EditorActor->SetVisibility(bIsVisible);
 			}
@@ -959,10 +962,10 @@ namespace UE::MLDeformer
 		UpdateIsReadyForTrainingState();
 
 		const int32 TrainingFrame = Model->GetVizSettings()->GetTrainingFrameNumber();
-		SetTrainingFrame(TrainingFrame);
+		SetTrainingFrame(TrainingFrame, false);
 
 		const int32 TestFrame = Model->GetVizSettings()->GetTestingFrameNumber();
-		SetTestFrame(TestFrame);
+		SetTestFrame(TestFrame, false);
 
 		UpdateEditorInputInfo();
 		CheckTrainingDataFrameChanged();
@@ -973,6 +976,7 @@ namespace UE::MLDeformer
 		float PlayOffset = static_cast<float>(NewScrubTime);
 
 		UMLDeformerVizSettings* VizSettings = Model->GetVizSettings();
+		bool bFrameChanged = true;
 		if (VizSettings->GetVisualizationMode() == EMLDeformerVizMode::TrainingData)
 		{
 			const int32 TargetFrame = GetTrainingFrameAtTime(NewScrubTime);
@@ -985,6 +989,10 @@ namespace UE::MLDeformer
 				}
 			}
 			VizSettings->SetTrainingFrameNumber(TargetFrame);
+			if (!bIsScrubbing)
+			{
+				CheckTrainingDataFrameChanged();
+			}
 		}
 		else if (VizSettings->GetVisualizationMode() == EMLDeformerVizMode::TestData)
 		{
@@ -999,6 +1007,44 @@ namespace UE::MLDeformer
 			}
 			VizSettings->SetTestingFrameNumber(TargetFrame);
 		}
+		
+		bIsScrubbingTimeline = bIsScrubbing;
+		UpdatePaintModePose(/*bFullUpdate=*/!bIsScrubbing);
+	}
+
+	void FMLDeformerEditorModel::UpdatePaintModePose(bool bFullUpdate)
+	{
+		if (!Editor->IsInitialized() || !Editor->IsPaintModeActive())
+		{
+			return;
+		}
+
+		UEdMode* ActiveMode = Editor->GetEditorModeManager().GetActiveScriptableMode(UMLDeformerPaintMode::Id);
+		UMLDeformerPaintMode* PaintMode = CastChecked<UMLDeformerPaintMode>(ActiveMode);
+		if (!PaintMode)
+		{
+			return;
+		}
+
+		UMLDeformerVizSettings* VizSettings = Model->GetVizSettings();
+		FMLDeformerEditorActor* EditorActor = nullptr;
+		if (VizSettings->GetVisualizationMode() == EMLDeformerVizMode::TrainingData)
+		{
+			EditorActor = FindEditorActor(ActorID_Train_Base);
+		}
+		else if (VizSettings->GetVisualizationMode() == EMLDeformerVizMode::TestData)
+		{
+			EditorActor = FindEditorActor(ActorID_Test_MLDeformed);
+		}
+		else
+		{
+			checkf(false, TEXT("Unexpected visualization mode"));
+		}
+
+		if (EditorActor)
+		{		
+			PaintMode->UpdatePose(EditorActor->GetSkeletalMeshComponent(), bFullUpdate);
+		}
 	}
 
 	float FMLDeformerEditorModel::CorrectedFrameTime(int32 FrameNumber, float FrameTimeForFrameNumber, FFrameRate FrameRate)
@@ -1053,25 +1099,25 @@ namespace UE::MLDeformer
 		return Model->GetVizSettings()->GetTestAnimSequence() ? Model->GetVizSettings()->GetTestAnimSequence()->GetFrameAtTime(TimeInSeconds) : 0;
 	}
 
-	void FMLDeformerEditorModel::SetTrainingFrame(int32 FrameNumber)
+	void FMLDeformerEditorModel::SetTrainingFrame(int32 FrameNumber, bool bIsScrubbing)
 	{
 		UMLDeformerVizSettings* VizSettings = Model->GetVizSettings();
 		VizSettings->SetTrainingFrameNumber(FrameNumber);
 		ClampCurrentTrainingFrameIndex();
 		if (VizSettings->GetVisualizationMode() == EMLDeformerVizMode::TrainingData)
 		{
-			OnTimeSliderScrubPositionChanged(GetTrainingTimeAtFrame(FrameNumber), false);
+			OnTimeSliderScrubPositionChanged(GetTrainingTimeAtFrame(FrameNumber), bIsScrubbing);
 		}
 	}
 
-	void FMLDeformerEditorModel::SetTestFrame(int32 FrameNumber)
+	void FMLDeformerEditorModel::SetTestFrame(int32 FrameNumber, bool bIsScrubbing)
 	{
 		UMLDeformerVizSettings* VizSettings = Model->GetVizSettings();
 		VizSettings->SetTestingFrameNumber(FrameNumber);
 		ClampCurrentTestFrameIndex();
 		if (VizSettings->GetVisualizationMode() == EMLDeformerVizMode::TestData)
 		{
-			OnTimeSliderScrubPositionChanged(GetTestTimeAtFrame(FrameNumber), false);
+			OnTimeSliderScrubPositionChanged(GetTestTimeAtFrame(FrameNumber), bIsScrubbing);
 		}
 	}
 
@@ -1249,11 +1295,11 @@ namespace UE::MLDeformer
 		{
 			if (VizSettings->GetVisualizationMode() == EMLDeformerVizMode::TrainingData)
 			{
-				SetTrainingFrame(GetTrainingFrameAtTime(CalcTimelinePosition()));
+				SetTrainingFrame(GetTrainingFrameAtTime(CalcTimelinePosition()), false);
 			}
 			else if (VizSettings->GetVisualizationMode() == EMLDeformerVizMode::TestData)
 			{
-				SetTestFrame(GetTestFrameAtTime(CalcTimelinePosition()));
+				SetTestFrame(GetTestFrameAtTime(CalcTimelinePosition()), false);
 			}
 		}
 
@@ -1804,7 +1850,7 @@ namespace UE::MLDeformer
 
 			// Draw the deltas for the current frame.
 			USkeletalMesh* SkelMesh = Model->GetSkeletalMesh();
-			if (bDrawDeltas && Model->GetSkeletalMesh() && SkelMesh->GetImportedModel())
+			if (bDrawDeltas && Model->GetSkeletalMesh() && SkelMesh->GetImportedModel() && !IsScrubbingTimeline())
 			{
 				const FLinearColor DeltasColor = FMLDeformerEditorStyle::Get().GetColor("MLDeformer.Deltas.Color");
 				const FLinearColor DebugVectorsColor = FMLDeformerEditorStyle::Get().GetColor("MLDeformer.DebugVectors.Color");
@@ -1823,10 +1869,11 @@ namespace UE::MLDeformer
 					// Render a delta for every render vertex in this LOD.
 					if (!Mapping.IsEmpty()) // If we have LOD mapping data.
 					{
+						const TArray<int32>& MeshToImportVertexMap = SkelMesh->GetImportedModel()->LODModels[0].MeshToImportVertexMap;
 						for (int32 RenderVertexIndex = 0; RenderVertexIndex < Mapping.Num(); ++RenderVertexIndex)
 						{
 							const int32 LOD0VertexIndex = Mapping[RenderVertexIndex]; // Get the vertex number in LOD0.
-							const int32 VertexIndex = SkelMesh->GetImportedModel()->LODModels[0].MeshToImportVertexMap[LOD0VertexIndex]; // Get the imported vertex number.
+							const int32 VertexIndex = MeshToImportVertexMap[LOD0VertexIndex]; // Get the imported vertex number.
 							const int32 ArrayIndex = 3 * VertexIndex; // 3 floats per delta, so multiply by 3.
 							const FVector Delta(
 								VertexDeltas[ArrayIndex], 
@@ -2730,36 +2777,46 @@ namespace UE::MLDeformer
 		return (float)ScrubPosition.AsDecimal() / (float)GetFrameRate();
 	}
 
-	void FMLDeformerEditorModel::SetScrubPosition(FFrameTime NewScrubPostion)
+	void FMLDeformerEditorModel::SetScrubPosition(FFrameTime NewScrubPosition)
+	{
+		SetScrubPosition(NewScrubPosition, false);
+	}
+
+	void FMLDeformerEditorModel::SetScrubPosition(FFrameTime NewScrubPosition, bool bIsScrubbing)
 	{
-		const double FrameNumber = NewScrubPostion.AsDecimal() * (double)GetFrameRate() / (double)GetTickResolution();
+		const double FrameNumber = NewScrubPosition.AsDecimal() * (double)GetFrameRate() / (double)GetTickResolution();
 		ScrubPosition.FrameNumber = (int32)FrameNumber;
 
 		const UMLDeformerVizSettings* VizSettings = GetModel()->GetVizSettings();
 		if (VizSettings->GetVisualizationMode() == EMLDeformerVizMode::TrainingData)
 		{
-			SetTrainingFrame(ScrubPosition.FrameNumber.Value);
+			SetTrainingFrame(ScrubPosition.FrameNumber.Value, bIsScrubbing);
 		}
 		else if (VizSettings->GetVisualizationMode() == EMLDeformerVizMode::TestData)
 		{
-			SetTestFrame(ScrubPosition.FrameNumber.Value);
+			SetTestFrame(ScrubPosition.FrameNumber.Value, bIsScrubbing);
 		}
 	}
 
-	void FMLDeformerEditorModel::SetScrubPosition(FFrameNumber NewScrubPostion)
+	void FMLDeformerEditorModel::SetScrubPosition(FFrameNumber NewScrubPosition)
+	{
+		SetScrubPosition(NewScrubPosition, false);
+	}
+
+	void FMLDeformerEditorModel::SetScrubPosition(FFrameNumber NewScrubPosition, bool bIsScrubbing)
 	{
 		const double Resolution = GetTickResolution();
-		const double ClampedValue = FMath::Clamp((double)NewScrubPostion.Value, PlaybackRange.GetLowerBoundValue() * Resolution, PlaybackRange.GetUpperBoundValue() * Resolution);
+		const double ClampedValue = FMath::Clamp((double)NewScrubPosition.Value, PlaybackRange.GetLowerBoundValue() * Resolution, PlaybackRange.GetUpperBoundValue() * Resolution);
 		const int32 FrameValue = ClampedValue * (double)GetFrameRate() / Resolution;
 
 		const UMLDeformerVizSettings* VizSettings = GetModel()->GetVizSettings();
 		if (VizSettings->GetVisualizationMode() == EMLDeformerVizMode::TrainingData)
 		{
-			SetTrainingFrame(FrameValue);
+			SetTrainingFrame(FrameValue, bIsScrubbing);
 		}
 		else if (VizSettings->GetVisualizationMode() == EMLDeformerVizMode::TestData)
 		{
-			SetTestFrame(FrameValue);
+			SetTestFrame(FrameValue, bIsScrubbing);
 		}
 	}
 
@@ -3185,12 +3242,27 @@ namespace UE::MLDeformer
 				{
 					const bool bForceStepInterpolation = VizSettings->GetShowHeatMap() && (VizSettings->GetHeatMapMode() == EMLDeformerHeatMapMode::GroundTruth) && VizSettings->HasTestGroundTruth();
 					const TOptional<EAnimInterpolationType> InterpolationOverride = bForceStepInterpolation ? EAnimInterpolationType::Step : TOptional<EAnimInterpolationType>();
-					AnimInstance->SetInterpolationOverride(InterpolationOverride);
+					AnimInstance->SetInterpolationOverride(InterpolationOverride);				
 				}
 			}
 		}
 	}
 
+	TVertexAttributesConstRef<float> FMLDeformerEditorModel::FindVertexAttributes(FName AttributeName) const
+	{
+		const USkeletalMesh* SkelMesh = Model->GetSkeletalMesh();
+		if (SkelMesh)
+		{
+			const int32 LodLevel = 0;
+			const FMeshDescription* MeshDescription = SkelMesh->GetMeshDescription(LodLevel);
+			if (MeshDescription)
+			{
+				return MeshDescription->VertexAttributes().GetAttributesRef<float>(AttributeName);
+			}
+		}
+		return TVertexAttributesConstRef<float>();
+	}
+
 }	// namespace UE::MLDeformer
 
 #undef LOCTEXT_NAMESPACE
diff --git a/Engine/Plugins/Animation/MLDeformer/MLDeformerFramework/Source/MLDeformerFrameworkEditor/Private/MLDeformerEditorStyle.cpp b/Engine/Plugins/Animation/MLDeformer/MLDeformerFramework/Source/MLDeformerFrameworkEditor/Private/MLDeformerEditorStyle.cpp
--- a/Engine/Plugins/Animation/MLDeformer/MLDeformerFramework/Source/MLDeformerFrameworkEditor/Private/MLDeformerEditorStyle.cpp
+++ b/Engine/Plugins/Animation/MLDeformer/MLDeformerFramework/Source/MLDeformerFrameworkEditor/Private/MLDeformerEditorStyle.cpp
@@ -56,6 +56,7 @@ namespace UE::MLDeformer
 		Set("MLDeformer.Morphs.ExcludedVertexColor", FLinearColor(0.1f, 0.1f, 0.1f));
 
 		Set("MLDeformer.Morphs.MaskIncludedVertexColor", FLinearColor(0.0f, 1.0f, 1.0f));
+		Set("MLDeformer.Morphs.MaskIncludedVertexColorPainted", FLinearColor(0.0f, 0.5f, 1.0f));
 		Set("MLDeformer.Morphs.MaskExcludedVertexColor", FLinearColor(0.1f, 0.1f, 0.1f));
 
 		Set("MLDeformer.Statistics.Performance", FLinearColor(0.0f, 1.0f, 0.0f));
diff --git a/Engine/Plugins/Animation/MLDeformer/MLDeformerFramework/Source/MLDeformerFrameworkEditor/Private/MLDeformerEditorToolkit.cpp b/Engine/Plugins/Animation/MLDeformer/MLDeformerFramework/Source/MLDeformerFrameworkEditor/Private/MLDeformerEditorToolkit.cpp
--- a/Engine/Plugins/Animation/MLDeformer/MLDeformerFramework/Source/MLDeformerFrameworkEditor/Private/MLDeformerEditorToolkit.cpp
+++ b/Engine/Plugins/Animation/MLDeformer/MLDeformerFramework/Source/MLDeformerFrameworkEditor/Private/MLDeformerEditorToolkit.cpp
@@ -7,8 +7,12 @@
 #include "MLDeformerModel.h"
 #include "MLDeformerApplicationMode.h"
 #include "MLDeformerEditorMode.h"
+#include "MLDeformerPaintMode.h"
+#include "MLDeformerPaintModeToolkit.h"
 #include "MLDeformerEditorModule.h"
 #include "MLDeformerEditorModel.h"
+#include "MLDeformerVizSettingsTabSummoner.h"
+#include "MLDeformerTimelineTabSummoner.h"
 #include "MLDeformerInputInfo.h"
 #include "MLDeformerEditorStyle.h"
 #include "MLDeformerVizSettings.h"
@@ -20,6 +24,7 @@
 #include "EditorViewportClient.h"
 #include "Modules/ModuleManager.h"
 #include "PersonaModule.h"
+#include "PersonaTabs.h"
 #include "IPersonaToolkit.h"
 #include "IAssetFamily.h"
 #include "IPersonaViewport.h"
@@ -96,11 +101,13 @@ namespace UE::MLDeformer
 		GetEditorModeManager().ActivateMode(FMLDeformerEditorMode::ModeName);
 
 		FMLDeformerEditorMode* EditorMode = static_cast<FMLDeformerEditorMode*>(GetEditorModeManager().GetActiveMode(FMLDeformerEditorMode::ModeName));
-		EditorMode->SetEditorToolkit(this);
+		if (EditorMode)
+		{
+			EditorMode->SetEditorToolkit(this);
+		}
 
 		SAssignNew(DebugWidget, SMLDeformerDebugSelectionWidget)
-			.MLDeformerEditor(this)
-			.Visibility(this, &FMLDeformerEditorToolkit::GetDebuggingVisibility);
+			.MLDeformerEditor(this);
 
 		ExtendToolbar();
 		RegenerateMenusAndToolbars();
@@ -109,8 +116,8 @@ namespace UE::MLDeformer
 		if (ActiveModel)
 		{
 			ActiveModel->UpdateIsReadyForTrainingState();
-			ActiveModel->SetTrainingFrame(ActiveModel->GetModel()->GetVizSettings()->GetTrainingFrameNumber());
-			ActiveModel->SetTestFrame(ActiveModel->GetModel()->GetVizSettings()->GetTestingFrameNumber());
+			ActiveModel->SetTrainingFrame(ActiveModel->GetModel()->GetVizSettings()->GetTrainingFrameNumber(), false);
+			ActiveModel->SetTestFrame(ActiveModel->GetModel()->GetVizSettings()->GetTestingFrameNumber(), false);
 			ActiveModel->InvalidateDeltas();
 		}
 
@@ -475,19 +482,13 @@ namespace UE::MLDeformer
 	void FMLDeformerEditorToolkit::FillToolbar(FToolBarBuilder& ToolbarBuilder)
 	{
 		// Training button and model selection.
-		ToolbarBuilder.BeginSection("Training");
+		ToolbarBuilder.BeginSection("Main");
 		{
 			ToolbarBuilder.AddToolBarButton(
 				FUIAction
 				(
-					FExecuteAction::CreateLambda
-					(
-						[this]()
-						{
-							Train(/*bSuppressDialogs*/false);
-						}
-					),
-					FCanExecuteAction::CreateRaw(this, &FMLDeformerEditorToolkit::IsTrainButtonEnabled)
+					FExecuteAction::CreateLambda([this]() { Train(/*bSuppressDialogs*/false); }),
+					FCanExecuteAction::CreateLambda([this]() { return IsTrainButtonEnabled() && !IsPaintModeActive(); })
 				),
 				NAME_None,
 				LOCTEXT("TrainModel", "Train Model"),
@@ -497,16 +498,28 @@ namespace UE::MLDeformer
 			);
 
 			ToolbarBuilder.AddSeparator();
+
+			// Model selection.
 			TSharedPtr<FUICommandList> CommandList = GetToolkitCommands();
 			ToolbarBuilder.AddComboButton(
-				FUIAction(),
+				FUIAction
+				(
+					FExecuteAction(),
+					FCanExecuteAction::CreateLambda([this]() { return !IsPaintModeActive(); })						
+				),
 				FOnGetContent::CreateRaw(this, &FMLDeformerEditorToolkit::GenerateModelButtonContents, CommandList.ToSharedRef()),
 				TAttribute<FText>::CreateRaw(this, &FMLDeformerEditorToolkit::GetActiveModelName),
 				LOCTEXT("ActiveModelTooltip", "The currently active ML Deformer Model."),
 				FSlateIcon(FMLDeformerEditorStyle::Get().GetStyleSetName(), "MLDeformer.VizSettings.TabIcon")
 			);
+
+			// Visualization mode.
 			ToolbarBuilder.AddComboButton(
-				FUIAction(),
+				FUIAction
+				(
+					FExecuteAction(),
+					FCanExecuteAction::CreateLambda([this]() { return !IsPaintModeActive(); })						
+				),
 				FOnGetContent::CreateRaw(this, &FMLDeformerEditorToolkit::GenerateVizModeButtonContents, CommandList.ToSharedRef()),
 				TAttribute<FText>::CreateRaw(this, &FMLDeformerEditorToolkit::GetCurrentVizModeName),			
 				LOCTEXT("VizModeModeTooltip", "The visualization mode, specifying whether you are working on training or testing."),
@@ -522,7 +535,11 @@ namespace UE::MLDeformer
 			{
 				TSharedPtr<FUICommandList> CommandList = GetToolkitCommands();
 				ToolbarBuilder.AddComboButton(
-					FUIAction(),
+					FUIAction
+					(
+						FExecuteAction(),
+						FCanExecuteAction::CreateLambda([this]() { return !IsPaintModeActive(); })						
+					),
 					FOnGetContent::CreateRaw(this, &FMLDeformerEditorToolkit::GenerateToolsMenuContents, CommandList.ToSharedRef()),
 					TAttribute<FText>::CreateLambda(
 						[this]()
@@ -540,15 +557,20 @@ namespace UE::MLDeformer
 		// Debugging.
 		ToolbarBuilder.BeginSection("Debugging");
 		{
+			DebugWidget->SetEnabled(TAttribute<bool>::CreateLambda([this]()
+			{
+				return !IsPaintModeActive() && GetDebuggingVisibility() == EVisibility::Visible;
+			}));
+
 			ToolbarBuilder.AddWidget(DebugWidget.ToSharedRef());
 			ToolbarBuilder.AddToolBarButton
 			(
 				FUIAction
 				(
 					FExecuteAction::CreateLambda([this]() { if (DebugWidget.IsValid()) { DebugWidget->Refresh(); } }),
-					FCanExecuteAction::CreateLambda([](){ return true; }),
-					FGetActionCheckState::CreateLambda([](){ return ECheckBoxState::Checked; }),
-					FIsActionButtonVisible::CreateLambda([this](){ return GetDebuggingVisibility() == EVisibility::Visible; })
+					FCanExecuteAction::CreateLambda([this]() { return !IsPaintModeActive() && GetDebuggingVisibility() == EVisibility::Visible; }),
+					FGetActionCheckState::CreateLambda([]() { return ECheckBoxState::Checked; }),
+					FIsActionButtonVisible::CreateLambda([this]() { return true;/*return GetDebuggingVisibility() == EVisibility::Visible;*/ })
 				),
 				NAME_None,
 				FText(),
@@ -558,6 +580,43 @@ namespace UE::MLDeformer
 			);
 		}
 		ToolbarBuilder.EndSection();
+
+		ToolbarBuilder.BeginSection("EditorModes");
+		{
+			ToolbarBuilder.AddToolBarButton(
+				FUIAction
+				(
+					FExecuteAction::CreateLambda
+					(
+						[this]()
+						{
+							if (!GetEditorModeManager().IsModeActive(UMLDeformerPaintMode::Id))
+							{
+								EnablePaintMode();
+							}
+							else
+							{
+								DisablePaintMode();
+							}							
+						}
+					),
+					FCanExecuteAction::CreateLambda([this]()
+					{
+						if (ActiveModel.IsValid() && ActiveModel->GetModel() && ActiveModel->GetModel()->GetSkeletalMesh())
+						{
+							return true;
+						}
+						return false; 
+					}),
+					FIsActionChecked::CreateLambda([this] { return IsPaintModeActive(); })
+				),
+				NAME_None,
+				LOCTEXT("PaintMode", "Paint Mode"),
+				LOCTEXT("PaintModeTooltip", "Enable or disable vertex attribute painting."),
+				FSlateIcon(),
+				EUserInterfaceActionType::ToggleButton
+			);
+		}
 	}
 
 	EVisibility FMLDeformerEditorToolkit::GetDebuggingVisibility() const
@@ -570,7 +629,7 @@ namespace UE::MLDeformer
 			}
 			else
 			{
-				return EVisibility::Hidden;
+				return EVisibility::Collapsed;
 			}
 		}
 		return EVisibility::Visible;
@@ -827,14 +886,14 @@ namespace UE::MLDeformer
 				{
 					ActiveModel->ClampCurrentTrainingFrameIndex();
 					const int32 CurrentFrameNumber = ActiveModel->GetModel()->GetVizSettings()->GetTrainingFrameNumber();
-					ActiveModel->SetTrainingFrame(CurrentFrameNumber);
+					ActiveModel->SetTrainingFrame(CurrentFrameNumber, false);
 					ActiveModel->SampleDeltas();
 				}
 				else if (VizSettings->GetVisualizationMode() == EMLDeformerVizMode::TestData)
 				{
 					ActiveModel->ClampCurrentTestFrameIndex();
 					const int32 CurrentFrameNumber = ActiveModel->GetModel()->GetVizSettings()->GetTestingFrameNumber();
-					ActiveModel->SetTestFrame(CurrentFrameNumber);
+					ActiveModel->SetTestFrame(CurrentFrameNumber, false);
 				}
 
 				if (GetVizSettingsDetailsView())
@@ -1069,6 +1128,116 @@ namespace UE::MLDeformer
 		FScopeLock Lock(&ExtendersMutex);
 		return ToolsMenuExtenders;
 	}
+
+	void FMLDeformerEditorToolkit::EnablePaintMode()
+	{
+		GetEditorModeManager().DeactivateMode(FMLDeformerEditorMode::ModeName);
+		GetEditorModeManager().ActivateMode(UMLDeformerPaintMode::Id);
+
+		UEdMode* ActiveMode = GetEditorModeManager().GetActiveScriptableMode(UMLDeformerPaintMode::Id);
+		if (UMLDeformerPaintMode* PaintMode = CastChecked<UMLDeformerPaintMode>(ActiveMode))
+		{
+			PaintMode->SetMLDeformerEditor(this);
+
+			if (ActiveModel.IsValid())
+			{
+				ActiveModel->UpdatePaintModePose();
+			}
+
+			// Disable the viz settings,details and timeline tabs, as we don't want to be able to modify those settings while in Paint mode.
+			TSharedPtr<SDockTab> VizTab = TabManager->FindExistingLiveTab(FMLDeformerVizSettingsTabSummoner::TabID);
+			if (VizTab.IsValid())
+			{
+				VizTab->SetEnabled(false);
+				VizTab->GetContent()->SetEnabled(false);
+			}
+
+			TSharedPtr<SDockTab> DetailsTab = TabManager->FindExistingLiveTab(FPersonaTabs::DetailsID);
+			if (DetailsTab.IsValid())
+			{
+				DetailsTab->SetEnabled(false);
+				DetailsTab->GetContent()->SetEnabled(false);
+			}
+		}
+
+		if (ActiveModel.IsValid())
+		{
+			ActiveModel->UpdateActorVisibility();
+			ActiveModel->UpdateLabels();
+		}
+	}
+
+	void FMLDeformerEditorToolkit::DisablePaintMode()
+	{
+		GetEditorModeManager().DeactivateMode(UMLDeformerPaintMode::Id);
+		GetEditorModeManager().ActivateMode(FMLDeformerEditorMode::ModeName);
+
+		// Enable some tabs again after Paint mode ended.
+		TSharedPtr<SDockTab> VizTab = TabManager->FindExistingLiveTab(FMLDeformerVizSettingsTabSummoner::TabID);
+		if (VizTab.IsValid())
+		{
+			VizTab->SetEnabled(true);
+			VizTab->GetContent()->SetEnabled(true);
+		}
+
+		TSharedPtr<SDockTab> DetailsTab = TabManager->FindExistingLiveTab(FPersonaTabs::DetailsID);
+		if (DetailsTab.IsValid())
+		{
+			DetailsTab->SetEnabled(true);
+			DetailsTab->GetContent()->SetEnabled(true);
+		}
+
+		// Update the actor visibility as the mode changes that.
+		if (ActiveModel.IsValid())
+		{
+			ActiveModel->UpdateActorVisibility();
+			ActiveModel->UpdateLabels();
+		}
+	}
+
+	bool FMLDeformerEditorToolkit::IsDefaultModeActive() const
+	{
+		if (GetPersonaToolkitPointer() == nullptr)
+		{
+			return true;
+		}
+		return GetEditorModeManager().IsModeActive(FMLDeformerEditorMode::ModeName);
+	}
+
+	bool FMLDeformerEditorToolkit::IsPaintModeActive() const
+	{
+		if (GetPersonaToolkitPointer() == nullptr)
+		{
+			return true;
+		}
+		return GetEditorModeManager().IsModeActive(UMLDeformerPaintMode::Id);
+	}
+
+	void FMLDeformerEditorToolkit::AddViewportOverlayWidget(TSharedRef<SWidget> InViewportOverlayWidget, int32 ZOrder)
+	{
+		if (GetViewport().IsValid())
+		{
+			GetViewport()->AddOverlayWidget(InViewportOverlayWidget, ZOrder);
+		}	
+	}
+
+	void FMLDeformerEditorToolkit::RemoveViewportOverlayWidget(TSharedRef<SWidget> InViewportOverlayWidget)
+	{
+		if (GetViewport().IsValid())
+		{
+			GetViewport()->RemoveOverlayWidget(InViewportOverlayWidget);
+		}	
+	}
+
+	void FMLDeformerEditorToolkit::Tick(float DeltaTime)
+	{
+		if (bNeedsPaintModeDisable)
+		{
+			DisablePaintMode();
+			bNeedsPaintModeDisable = false;
+		}		
+	}
+
 }	// namespace UE::MLDeformer
 
 #undef LOCTEXT_NAMESPACE
diff --git a/Engine/Plugins/Animation/MLDeformer/MLDeformerFramework/Source/MLDeformerFrameworkEditor/Private/MLDeformerGeomCacheEditorModel.cpp b/Engine/Plugins/Animation/MLDeformer/MLDeformerFramework/Source/MLDeformerFrameworkEditor/Private/MLDeformerGeomCacheEditorModel.cpp
--- a/Engine/Plugins/Animation/MLDeformer/MLDeformerFramework/Source/MLDeformerFrameworkEditor/Private/MLDeformerGeomCacheEditorModel.cpp
+++ b/Engine/Plugins/Animation/MLDeformer/MLDeformerFramework/Source/MLDeformerFrameworkEditor/Private/MLDeformerGeomCacheEditorModel.cpp
@@ -320,7 +320,7 @@ namespace UE::MLDeformer
 		if (GeometryCacheComponent)
 		{
 			GeometryCacheComponent->SetGeometryCache(GetActiveGeometryCache());
-			GeometryCacheComponent->SetLooping(false);
+			GeometryCacheComponent->SetLooping(true);
 			GeometryCacheComponent->SetManualTick(true);
 			GeometryCacheComponent->SetPlaybackSpeed(TestAnimSpeed);
 			GeometryCacheComponent->Play();
diff --git a/Engine/Plugins/Animation/MLDeformer/MLDeformerFramework/Source/MLDeformerFrameworkEditor/Private/MLDeformerMorphModelDetails.cpp b/Engine/Plugins/Animation/MLDeformer/MLDeformerFramework/Source/MLDeformerFrameworkEditor/Private/MLDeformerMorphModelDetails.cpp
--- a/Engine/Plugins/Animation/MLDeformer/MLDeformerFramework/Source/MLDeformerFrameworkEditor/Private/MLDeformerMorphModelDetails.cpp
+++ b/Engine/Plugins/Animation/MLDeformer/MLDeformerFramework/Source/MLDeformerFrameworkEditor/Private/MLDeformerMorphModelDetails.cpp
@@ -175,8 +175,25 @@ namespace UE::MLDeformer
 
 		IDetailGroup& MaskGroup = MorphTargetCategoryBuilder->AddGroup("Mask", LOCTEXT("MorphMaskGroupLabel", "Masking"), false, false);
 		MaskGroup.AddPropertyRow(DetailLayoutBuilder->GetProperty(UMLDeformerMorphModel::GetMaskChannelPropertyName(), UMLDeformerMorphModel::StaticClass()));
+		MaskGroup.AddPropertyRow(DetailLayoutBuilder->GetProperty(UMLDeformerMorphModel::GetGlobalMaskAttributePropertyName(), UMLDeformerMorphModel::StaticClass()));
 		MaskGroup.AddPropertyRow(DetailLayoutBuilder->GetProperty(UMLDeformerMorphModel::GetInvertMaskChannelPropertyName(), UMLDeformerMorphModel::StaticClass()));
 
+		if (MorphModel && MorphModel->GetMaskChannel() == EMLDeformerMaskChannel::VertexAttribute && !EditorModel->FindVertexAttributes(MorphModel->GetGlobalMaskAttributeName()).IsValid())
+		{
+			const FText MaskErrorText = LOCTEXT("MorphGlobalMaskWeightMapError", "The weight map attribute you specified does not exist on the skeletal mesh.");
+			FDetailWidgetRow& MaskErrorRow = MorphTargetCategoryBuilder->AddCustomRow(FText::FromString("MorphGlobalMaskError"))
+				.WholeRowContent()
+				[
+					SNew(SBox)
+					.Padding(FMargin(0.0f, 4.0f))
+					[
+						SNew(SWarningOrErrorBox)
+						.MessageStyle(EMessageStyle::Warning)
+						.Message(MaskErrorText)
+					]
+				];
+		}
+
 		if (MorphModel && MorphModel->HasRawMorph() && !MorphModel->CanDynamicallyUpdateMorphTargets())
 		{
 			const FText DeltaCountMismatchErrorText = LOCTEXT("MorphDeltaCountMismatch", "Dynamic morph target updates disabled until retrained. This is because the vertex count changed after the model was trained.");
diff --git a/Engine/Plugins/Animation/MLDeformer/MLDeformerFramework/Source/MLDeformerFrameworkEditor/Private/MLDeformerMorphModelEditorModel.cpp b/Engine/Plugins/Animation/MLDeformer/MLDeformerFramework/Source/MLDeformerFrameworkEditor/Private/MLDeformerMorphModelEditorModel.cpp
--- a/Engine/Plugins/Animation/MLDeformer/MLDeformerFramework/Source/MLDeformerFrameworkEditor/Private/MLDeformerMorphModelEditorModel.cpp
+++ b/Engine/Plugins/Animation/MLDeformer/MLDeformerFramework/Source/MLDeformerFrameworkEditor/Private/MLDeformerMorphModelEditorModel.cpp
@@ -22,6 +22,8 @@
 #include "Engine/SkeletalMesh.h"
 #include "Misc/ScopedSlowTask.h"
 #include "BoneWeights.h"
+#include "MeshAttributeArray.h"
+#include "Async/ParallelFor.h"
 
 #define LOCTEXT_NAMESPACE "MLDeformerMorphModelEditorModel"
 
@@ -79,7 +81,9 @@ namespace UE::MLDeformer
 			Property->GetFName() == UMLDeformerMorphModel::GetMorphCompressionLevelPropertyName() ||
 			Property->GetFName() == UMLDeformerMorphModel::GetIncludeMorphTargetNormalsPropertyName() ||
 			Property->GetFName() == UMLDeformerMorphModel::GetMaskChannelPropertyName() ||
-			Property->GetFName() == UMLDeformerMorphModel::GetInvertMaskChannelPropertyName())
+			Property->GetFName() == UMLDeformerMorphModel::GetInvertMaskChannelPropertyName() ||
+			Property->GetFName() == UMLDeformerMorphModel::GetGlobalMaskAttributePropertyName() ||
+			Property->GetFName() == UMLDeformerMorphModel::GetSkeletalMeshPropertyName())
 		{
 			if (PropertyChangedEvent.ChangeType == EPropertyChangeType::ValueSet)
 			{
@@ -87,6 +91,11 @@ namespace UE::MLDeformer
 				{
 					InitEngineMorphTargets(GetMorphModel()->GetMorphTargetDeltas());
 				}
+
+				if (Property->GetFName() == UMLDeformerMorphModel::GetGlobalMaskAttributePropertyName() || Property->GetFName() == UMLDeformerMorphModel::GetMaskChannelPropertyName())
+				{
+					GetEditor()->GetModelDetailsView()->ForceRefresh();
+				}
 			}
 		}
 		else if (Property->GetFName() == UMLDeformerMorphModelVizSettings::GetMorphTargetNumberPropertyName())
@@ -312,6 +321,9 @@ namespace UE::MLDeformer
 		const int32 NumRenderVertices = RenderData->LODRenderData[LOD].GetNumVertices();
 		const FColorVertexBuffer& ColorBuffer = RenderData->LODRenderData[LOD].StaticVertexBuffers.ColorVertexBuffer;
 
+		const TArray<int32>& VertexMap = Model->GetVertexMap();
+		TArray<float> GlobalMaskWeights = CalcGlobalMaskWeights(VertexMap, ColorBuffer, EMLDeformerMaskChannel::Disabled, false);
+
 		CalcMorphTargetNormals(
 			LOD,
 			SkelMesh,
@@ -323,6 +335,7 @@ namespace UE::MLDeformer
 			ColorBuffer,
 			EMLDeformerMaskChannel::Disabled,
 			false,
+			GlobalMaskWeights,
 			OutDeltaNormals);
 	}
 
@@ -338,6 +351,27 @@ namespace UE::MLDeformer
 		EMLDeformerMaskChannel MaskChannel,
 		bool bInvertGlobalMaskChannel,
 		TArray<FVector3f>& OutDeltaNormals)
+	{
+		const TArray<int32>& VertexMap = Model->GetVertexMap();
+		TArray<float> GlobalMaskWeights = CalcGlobalMaskWeights(VertexMap, ColorBuffer, MaskChannel, bInvertGlobalMaskChannel);
+
+		CalcMorphTargetNormals(LOD, SkelMesh, MorphTargetIndex, Deltas, BaseVertexPositions, BaseNormals, 
+			ImportedVertexToRenderVertexMapping, ColorBuffer, MaskChannel, bInvertGlobalMaskChannel, GlobalMaskWeights, OutDeltaNormals);
+	}
+
+	void FMLDeformerMorphModelEditorModel::CalcMorphTargetNormals(
+		int32 LOD,
+		const USkeletalMesh* SkelMesh,
+		int32 MorphTargetIndex,
+		const TArrayView<const FVector3f> Deltas,
+		const TArrayView<const FVector3f> BaseVertexPositions,
+		const TArrayView<const FVector3f> BaseNormals,
+		const TArrayView<const int32> ImportedVertexToRenderVertexMapping,
+		const FColorVertexBuffer& ColorBuffer,
+		EMLDeformerMaskChannel MaskChannel,
+		bool bInvertGlobalMaskChannel,
+		const TArray<float>& GlobalMaskWeights,
+		TArray<FVector3f>& OutDeltaNormals)
 	{
 		const FSkeletalMeshModel* ImportedModel = SkelMesh->GetImportedModel();
 		if (ImportedModel == nullptr || !ImportedModel->LODModels.IsValidIndex(LOD))
@@ -362,7 +396,7 @@ namespace UE::MLDeformer
 			const int32 DeltaIndex = (MorphTargetIndex * NumBaseMeshVerts) + VertexIndex;
 			const FVector3f RawDelta = Deltas[DeltaIndex];
 			const float MorphMaskWeight = MorphMask.IsEmpty() ? 1.0f : MorphMask[VertexIndex];
-			const float GlobalMaskWeight = (RenderVertexIndex != INDEX_NONE) ? CalcGlobalMaskWeight(RenderVertexIndex, ColorBuffer, MaskChannel, bInvertGlobalMaskChannel) : 1.0f;
+			const float GlobalMaskWeight = (RenderVertexIndex != INDEX_NONE) ? GlobalMaskWeights[RenderVertexIndex] : 1.0f;
 			FVector3f ScaledDelta;
 			FVector3f DummyNormal;
 			ProcessVertexDelta(ScaledDelta, DummyNormal, RawDelta, FVector3f::ZeroVector, 0.0f, MorphMaskWeight, GlobalMaskWeight);
@@ -393,18 +427,103 @@ namespace UE::MLDeformer
 		return (OutScaledDelta.Length() >= DeltaThreshold);
 	}
 
+	TArray<float> FMLDeformerMorphModelEditorModel::CalcGlobalMaskWeights(const TArray<int32>& VertexMap, const FColorVertexBuffer& ColorBuffer, EMLDeformerMaskChannel MaskChannel, bool bInvertMaskChannel) const
+	{
+		TArray<float> OutWeights;
+		OutWeights.SetNumUninitialized(VertexMap.Num());
+
+		const int32 NumRenderVertices = VertexMap.Num();
+
+		// If we use vertex colors.
+		if (MaskChannel == EMLDeformerMaskChannel::VertexColorRed || 
+			MaskChannel == EMLDeformerMaskChannel::VertexColorGreen || 
+			MaskChannel == EMLDeformerMaskChannel::VertexColorBlue || 
+			MaskChannel == EMLDeformerMaskChannel::VertexColorAlpha)
+		{
+			const int32 NumColorVerts = ColorBuffer.GetNumVertices();
+			for (int32 RenderVertexIndex = 0; RenderVertexIndex < NumRenderVertices; ++RenderVertexIndex)
+			{
+				float OutWeight = 1.0f;
+				if (RenderVertexIndex != INDEX_NONE && NumColorVerts > 0)
+				{
+					const int32 ImportedVertexNumber = VertexMap[RenderVertexIndex];
+					if (ImportedVertexNumber != INDEX_NONE)
+					{
+						const FLinearColor VertexColor = ColorBuffer.VertexColor(RenderVertexIndex);
+						switch (MaskChannel)
+						{
+							case EMLDeformerMaskChannel::VertexColorRed:	{ OutWeight = VertexColor.R; break; }
+							case EMLDeformerMaskChannel::VertexColorGreen:	{ OutWeight = VertexColor.G; break; }
+							case EMLDeformerMaskChannel::VertexColorBlue:	{ OutWeight = VertexColor.B; break; }
+							case EMLDeformerMaskChannel::VertexColorAlpha:	{ OutWeight = VertexColor.A; break; }
+							default: 
+								checkf(false, TEXT("Unexpected mask channel value."));
+								break;
+						};
+					}
+				}
+
+				OutWeights[RenderVertexIndex] = OutWeight;
+			}
+
+		}
+		else if (MaskChannel == EMLDeformerMaskChannel::VertexAttribute)	// Using a weight map as setup inside the skeletal mesh editor.
+		{
+			TVertexAttributesConstRef<float> WeightMapAttributes = FindVertexAttributes(GetMorphModel()->GetGlobalMaskAttributeName());
+			if (WeightMapAttributes.IsValid())
+			{
+				for (int32 RenderVertexIndex = 0; RenderVertexIndex < NumRenderVertices; ++RenderVertexIndex)
+				{
+					const float VertexWeight = (VertexMap[RenderVertexIndex] != INDEX_NONE) ? WeightMapAttributes.Get(VertexMap[RenderVertexIndex]) : 1.0f;
+					OutWeights[RenderVertexIndex] = FMath::Clamp(VertexWeight, 0.0f, 1.0f);
+				}
+			}
+			else
+			{
+				for (int32 RenderVertexIndex = 0; RenderVertexIndex < NumRenderVertices; ++RenderVertexIndex)
+				{
+					OutWeights[RenderVertexIndex] = 1.0f;
+				}
+			}
+		}
+		else if (MaskChannel == EMLDeformerMaskChannel::Disabled)	// We disabled the mask, fill with 1.
+		{
+			for (int32 RenderVertexIndex = 0; RenderVertexIndex < NumRenderVertices; ++RenderVertexIndex)
+			{
+				OutWeights[RenderVertexIndex] = 1.0f;
+			}
+		}
+
+		// Invert the weights if desired.
+		if (bInvertMaskChannel)
+		{
+			for (int32 Index = 0; Index < OutWeights.Num(); ++Index)
+			{
+				const float Weight = OutWeights[Index];
+				OutWeights[Index] = FMath::Clamp<float>(1.0f - Weight, 0.0f, 1.0f);
+			}
+		}
+
+		return MoveTemp(OutWeights);
+	}
+
 	float FMLDeformerMorphModelEditorModel::CalcGlobalMaskWeight(int32 RenderVertexIndex, const FColorVertexBuffer& ColorBuffer, EMLDeformerMaskChannel MaskChannel, bool bInvertMaskChannel) const
 	{
 		float VertexWeight = 1.0f;
 		if (ColorBuffer.GetNumVertices() != 0 && MaskChannel != EMLDeformerMaskChannel::Disabled && RenderVertexIndex != INDEX_NONE)
 		{
-			const FLinearColor& VertexColor = ColorBuffer.VertexColor(RenderVertexIndex);
+			const FLinearColor VertexColor = ColorBuffer.VertexColor(RenderVertexIndex);
 			switch (MaskChannel)
 			{
 				case EMLDeformerMaskChannel::VertexColorRed:	{ VertexWeight = VertexColor.R; break; }
 				case EMLDeformerMaskChannel::VertexColorGreen:	{ VertexWeight = VertexColor.G; break; }
 				case EMLDeformerMaskChannel::VertexColorBlue:	{ VertexWeight = VertexColor.B; break; }
 				case EMLDeformerMaskChannel::VertexColorAlpha:	{ VertexWeight = VertexColor.A; break; }
+				case EMLDeformerMaskChannel::VertexAttribute:
+				{ 
+					VertexWeight = 1.0f;
+					break; 
+				}
 				default: 
 					checkf(false, TEXT("Unexpected mask channel value."));
 					break;
@@ -441,8 +560,17 @@ namespace UE::MLDeformer
 			return;
 		}
 
+		if (Sampler->GetSkeletalMeshComponent()->GetSkeletalMeshAsset() != Model->GetSkeletalMesh())
+		{
+			return;
+		}
+
 		const int32 NumBaseMeshVerts = Model->GetNumBaseMeshVerts();
-		check(Deltas.Num() % NumBaseMeshVerts == 0);
+		if (Deltas.Num() % NumBaseMeshVerts != 0)
+		{
+			return;
+		}
+
 		const int32 NumMorphTargets = Deltas.Num() / NumBaseMeshVerts;
 		check((Deltas.Num() / NumMorphTargets) == NumBaseMeshVerts);
 		check(!Model->GetVertexMap().IsEmpty());
@@ -459,17 +587,29 @@ namespace UE::MLDeformer
 		// Calculate the normals for the base mesh.
 		const FSkeletalMeshModel* ImportedModel = SkelMesh->GetImportedModel();
 		const TArrayView<const uint32> IndexArray = ImportedModel->LODModels[LOD].IndexBuffer;
-		const TArrayView<const int32> VertexMap = ImportedModel->LODModels[LOD].MeshToImportVertexMap;
+		const TArray<int32>& VertexMap = ImportedModel->LODModels[LOD].MeshToImportVertexMap;
 		const TArrayView<const FVector3f> BaseVertexPositions = Sampler->GetUnskinnedVertexPositions();
 		const FColorVertexBuffer& ColorBuffer = RenderData->LODRenderData[LOD].StaticVertexBuffers.ColorVertexBuffer;
 
 		// Precalc an array that maps imported vertices to a render vertex.
 		TArray<int32> ImportedVertexToRenderVertexMapping;
 		ImportedVertexToRenderVertexMapping.SetNumUninitialized(NumBaseMeshVerts);
-		for (int32 Index = 0; Index < NumBaseMeshVerts; ++Index)
+		const int32 BatchSize = 500;
+		const int32 NumBatches = (NumBaseMeshVerts / BatchSize) + 1;
+		ParallelFor(NumBatches, [&](int32 BatchIndex)
 		{
-			ImportedVertexToRenderVertexMapping[Index] = VertexMap.Find(Index);
-		}
+			const int32 StartVertex = BatchIndex * BatchSize;
+			if (StartVertex >= NumBaseMeshVerts)
+			{
+				return;
+			}
+
+			const int32 NumVertsInBatch = (StartVertex + BatchSize) < NumBaseMeshVerts ? BatchSize : FMath::Max(NumBaseMeshVerts - StartVertex, 0);
+			for (int32 VertexIndex = StartVertex; VertexIndex < StartVertex + NumVertsInBatch; ++VertexIndex)
+			{
+				ImportedVertexToRenderVertexMapping[VertexIndex] = VertexMap.Find(VertexIndex);
+			}
+		});
 
 		TArray<FVector3f> BaseNormals;
 		if (bIncludeNormals)
@@ -477,14 +617,19 @@ namespace UE::MLDeformer
 			CalcVertexNormals(BaseVertexPositions, IndexArray, VertexMap, BaseNormals);
 		}
 
+		// Calculate the global mask weights.
+		const TArray<float> GlobalMaskWeights = CalcGlobalMaskWeights(VertexMap, ColorBuffer, MaskChannel, bInvertMaskChannel);
+
 		// Initialize an engine morph target for each model morph target.
+		bool bHasOnlyEmptyMorphs = true;
 		UE_LOG(LogMLDeformer, Display, TEXT("Initializing %d engine morph targets of %d vertices each"), NumMorphTargets, Deltas.Num() / NumMorphTargets);
 		TArray<FVector3f> DeltaNormals;
 		for (int32 MorphTargetIndex = 0; MorphTargetIndex < NumMorphTargets; ++MorphTargetIndex)
 		{
 			if (bIncludeNormals)
 			{
-				CalcMorphTargetNormals(LOD, SkelMesh, MorphTargetIndex, Deltas, BaseVertexPositions, BaseNormals, ImportedVertexToRenderVertexMapping, ColorBuffer, MaskChannel, bInvertMaskChannel, DeltaNormals);
+				CalcMorphTargetNormals(LOD, SkelMesh, MorphTargetIndex, Deltas, BaseVertexPositions, BaseNormals, ImportedVertexToRenderVertexMapping, 
+					ColorBuffer, MaskChannel, bInvertMaskChannel, GlobalMaskWeights, DeltaNormals);
 			}
 
 			const FName MorphName = *FString::Printf(TEXT("%s%.3d"), *NamePrefix, MorphTargetIndex);
@@ -513,7 +658,7 @@ namespace UE::MLDeformer
 				const int32 ImportedVertexNumber = VertexMap[VertexIndex];
 				if (ImportedVertexNumber != INDEX_NONE)
 				{
-					const float GlobalMaskWeight = CalcGlobalMaskWeight(VertexIndex, ColorBuffer, MaskChannel, bInvertMaskChannel);
+					const float GlobalMaskWeight = GlobalMaskWeights[VertexIndex];
 					const float MorphMaskWeight = MorphMask.IsEmpty() ? 1.0f : MorphMask[ImportedVertexNumber];
 					const FVector3f RawDelta = Deltas[ImportedVertexNumber + MorphTargetIndex * NumBaseMeshVerts];
 					const FVector3f RawDeltaNormal = !DeltaNormals.IsEmpty() ? DeltaNormals[ImportedVertexNumber] : FVector3f::ZeroVector;
@@ -540,8 +685,11 @@ namespace UE::MLDeformer
 			}	// for all morph targets
 
 			MorphLODModel.Vertices.Shrink();
+			bHasOnlyEmptyMorphs &= MorphLODModel.Vertices.IsEmpty();
 			Task.EnterProgressFrame();
-		}
+		} // For all morph targets.
+
+		GetMorphModel()->SetHasOnlyEmptyMorphs(bHasOnlyEmptyMorphs);
 	}
 
 	void FMLDeformerMorphModelEditorModel::CompressMorphTargets(FMorphTargetVertexInfoBuffers& OutMorphBuffers, const TArray<UMorphTarget*>& MorphTargets, int32 LOD, float MorphErrorTolerance)
@@ -829,6 +977,98 @@ namespace UE::MLDeformer
 		}
 	}
 
+	void FMLDeformerMorphModelEditorModel::ApplyMaskInfoToBuffer(const USkeletalMesh* SkeletalMesh, const FMLDeformerMaskInfo& MaskInfo, TArrayView<float> ItemMaskBuffer)
+	{
+		// Apply the bones to the mask buffer.
+		if (MaskInfo.MaskMode == EMLDeformerMaskingMode::Generated)
+		{
+			const FReferenceSkeleton& RefSkel = SkeletalMesh->GetRefSkeleton();
+			for (const FName MaskBoneName : MaskInfo.BoneNames)
+			{
+				const int32 MaskBoneIndex = RefSkel.FindBoneIndex(MaskBoneName);
+				if (MaskBoneIndex != INDEX_NONE)
+				{
+					ApplyBoneToMask(MaskBoneIndex, ItemMaskBuffer);
+				}
+				else
+				{
+					UE_LOG(LogMLDeformer, Warning, TEXT("Mask contains a bone named '%s', which cannot be found in the ref skeleton of skeletal mesh '%s'."),
+						*MaskBoneName.ToString(),
+						*SkeletalMesh->GetName());
+				}
+			}
+		}
+		else // We're using a painted mask.
+		{
+			check(MaskInfo.MaskMode == EMLDeformerMaskingMode::VertexAttribute);
+			const FName VertexAttributeName = MaskInfo.VertexAttributeName;
+			if (!VertexAttributeName.IsNone())
+			{
+				TVertexAttributesConstRef<float> AttributeValues = FindVertexAttributes(VertexAttributeName);
+				if (AttributeValues.IsValid())
+				{
+					const int32 NumAttributes = AttributeValues.GetNumElements();
+					check(NumAttributes == ItemMaskBuffer.Num());
+					for (int32 Index = 0; Index < NumAttributes; ++Index)
+					{
+						ItemMaskBuffer[Index] = AttributeValues.Get(Index);
+					}
+				}
+				else
+				{
+					UE_LOG(LogMLDeformer, Warning, TEXT("Mask references a vertex attribute '%s' which doesn't exist on skeletal mesh %s."), *VertexAttributeName.ToString(), *SkeletalMesh->GetName());
+					FillMaskValues(ItemMaskBuffer, 1.0f);
+				}
+			}
+			else
+			{
+				UE_LOG(LogMLDeformer, Warning, TEXT("Mask is set to use a vertex attribute, but none is specified."));
+				FillMaskValues(ItemMaskBuffer, 1.0f);
+			}
+		}
+	}
+
+	void FMLDeformerMorphModelEditorModel::ApplyGeneratedMaskToVertexAttributes(USkeletalMesh* SkeletalMesh, FMLDeformerMaskInfo& MaskInfo, TVertexAttributesRef<float> AttributeRef)
+	{
+		// Output some buffer with the generated mask.
+		TArray<float> GeneratedBuffer;
+		const int32 NumVerts = GetModel()->GetNumBaseMeshVerts();
+		GeneratedBuffer.SetNumZeroed(NumVerts);
+		EMLDeformerMaskingMode ModeBackup = MaskInfo.MaskMode;
+		MaskInfo.MaskMode = EMLDeformerMaskingMode::Generated;	// Force generating.
+		ApplyMaskInfoToBuffer(SkeletalMesh, MaskInfo, GeneratedBuffer);
+		MaskInfo.MaskMode = ModeBackup;
+
+		check(NumVerts == GeneratedBuffer.Num());
+
+		// Calculate the average to scale the weights.
+		// Normalizing the values doesn't really work well, as many values will have quite tiny weights then.
+		// We tried normalizing, median the average, and the average seems to give the best vertex attribute weights.
+		float AverageValue = 0.0f;
+		for (int32 Index = 0; Index < NumVerts; ++Index)
+		{
+			AverageValue += GeneratedBuffer[Index];
+		}
+
+		if (NumVerts > 0)
+		{
+			AverageValue /= static_cast<float>(NumVerts);
+			AverageValue *= 4.0f;	// Scale it a bit, to have slightly smoother edges.
+		}
+
+		if (FMath::IsNearlyZero(AverageValue))
+		{
+			AverageValue = 1.0f;
+		}
+
+		// Now apply this to the vertex attributes.
+		for (int32 Index = 0; Index < NumVerts; ++Index)
+		{
+			const float Weight = FMath::Clamp(GeneratedBuffer[Index] / AverageValue, 0.0f, 1.0f);
+			AttributeRef.Set(Index, Weight);
+		}
+	}
+
 	void FMLDeformerMorphModelEditorModel::ApplyBoneToMask(int32 SkeletonBoneIndex, TArrayView<float> MaskBuffer)
 	{
 		const int32 LOD = 0;
@@ -989,6 +1229,19 @@ namespace UE::MLDeformer
 		PRAGMA_ENABLE_DEPRECATION_WARNINGS
 	}
 
+	void FMLDeformerMorphModelEditorModel::OnObjectModified(UObject* Object)
+	{
+		if (Model->GetSkeletalMesh() == Object || Model->GetInputInfo()->GetSkeletalMesh() == Object)
+		{
+			if (GetMorphModel()->CanDynamicallyUpdateMorphTargets())
+			{
+				InitEngineMorphTargets(GetMorphModel()->GetMorphTargetDeltas());
+			}
+			bNeedsAssetReinit = true;
+		}
+
+		FMLDeformerGeomCacheEditorModel::OnObjectModified(Object);
+	}
 }	// namespace UE::MLDeformer
 
 #undef LOCTEXT_NAMESPACE
diff --git a/Engine/Plugins/Animation/MLDeformer/MLDeformerFramework/Source/MLDeformerFrameworkEditor/Private/MLDeformerPaintMode.cpp b/Engine/Plugins/Animation/MLDeformer/MLDeformerFramework/Source/MLDeformerFrameworkEditor/Private/MLDeformerPaintMode.cpp
new file mode 100644
--- /dev/null
+++ b/Engine/Plugins/Animation/MLDeformer/MLDeformerFramework/Source/MLDeformerFrameworkEditor/Private/MLDeformerPaintMode.cpp
@@ -0,0 +1,147 @@
+// Copyright Epic Games, Inc. All Rights Reserved.
+
+#include "MLDeformerPaintMode.h"
+#include "MLDeformerPaintModeToolkit.h"
+#include "MeshAttributePaintTool.h"
+#include "EdMode.h"
+#include "EdModeInteractiveToolsContext.h"
+#include "EditorModeManager.h"
+#include "ToolTargets/SkeletalMeshComponentToolTarget.h"
+#include "ModelingToolsManagerActions.h"
+#include "ToolTargetManager.h"
+#include "ToolBuilderUtil.h"
+#include "Components/SkeletalMeshComponent.h"
+#include "DynamicMesh/DynamicMeshAttributeSet.h"
+#include "DynamicMesh/DynamicVertexSkinWeightsAttribute.h"
+#include "SkeletalMeshAttributes.h"
+
+#define LOCTEXT_NAMESPACE "MLDeformerPaintMode"
+
+const FEditorModeID UMLDeformerPaintMode::Id("MLDeformerPaintMode");
+
+UMLDeformerPaintMode::UMLDeformerPaintMode() 
+{
+	Info = FEditorModeInfo(Id, LOCTEXT("MLDeformerPaintMode", "ML Deformer Paint Mode"), FSlateIcon(), false);
+}
+
+void UMLDeformerPaintMode::Enter()
+{
+	UEdMode::Enter();
+
+	UEditorInteractiveToolsContext* InteractiveToolsContext = GetInteractiveToolsContext();
+	InteractiveToolsContext->TargetManager->AddTargetFactory(NewObject<USkeletalMeshComponentToolTargetFactory>(InteractiveToolsContext->TargetManager));
+
+	const FModelingToolsManagerCommands& ToolManagerCommands = FModelingToolsManagerCommands::Get();
+	RegisterTool(ToolManagerCommands.BeginMeshAttributePaintTool, TEXT("BeginMeshAttributePaintTool"), NewObject<UMeshAttributePaintToolBuilder>());
+
+	GetInteractiveToolsContext()->ToolManager->SelectActiveToolType(EToolSide::Left, TEXT("BeginMeshAttributePaintTool"));
+	GetInteractiveToolsContext()->ToolManager->ActivateTool(EToolSide::Left);
+
+	// Try to get the bind pose vertex positions.
+	UInteractiveTool* Tool = GetInteractiveToolsContext()->ToolManager->GetActiveTool(EToolSide::Left);
+	UMeshAttributePaintTool* PaintTool = Cast<UMeshAttributePaintTool>(Tool);
+	if (PaintTool)
+	{
+		TObjectPtr<UPreviewMesh> PreviewMesh = PaintTool->GetPreviewMesh();
+		if (PreviewMesh)
+		{
+			const UE::Geometry::FDynamicMesh3* Mesh = PreviewMesh->GetMesh();
+			if (Mesh)
+			{
+				BindPosePositions.Reset();
+				BindPosePositions.SetNumZeroed(Mesh->VertexCount());
+
+				for (const int32 VertexID : Mesh->VertexIndicesItr())
+				{
+					BindPosePositions[VertexID] = FVector3f(Mesh->GetVertex(VertexID));
+				}						
+			}
+		}
+	}
+}
+
+void UMLDeformerPaintMode::CreateToolkit()
+{
+	Toolkit = MakeShareable(new UE::MLDeformer::FMLDeformerPaintModeToolkit());
+}
+
+void UMLDeformerPaintMode::UpdatePose(USkeletalMeshComponent* SkeletalMeshComponent, bool bFullUpdate)
+{
+	using namespace UE::MLDeformer;
+
+	UInteractiveTool* Tool = GetInteractiveToolsContext()->ToolManager->GetActiveTool(EToolSide::Left);
+	UMeshAttributePaintTool* PaintTool = Cast<UMeshAttributePaintTool>(Tool);
+	if (!PaintTool || !SkeletalMeshComponent || !SkeletalMeshComponent->GetSkeletalMeshAsset() || BindPosePositions.IsEmpty())
+	{
+		return;
+	}
+
+	TObjectPtr<UPreviewMesh> PreviewMesh = PaintTool->GetPreviewMesh();
+	if (!PreviewMesh)
+	{
+		return;
+	}
+
+	// Get the current bone transforms.
+	TArray<FMatrix44f> BoneMatrices;
+	SkeletalMeshComponent->RefreshBoneTransforms();
+	SkeletalMeshComponent->CacheRefToLocalMatrices(BoneMatrices);
+	
+	// Modify the vertex positions of the dynamic mesh.
+	// We perform linear blend skinning using the bone matrices we just extracted.
+	PreviewMesh->DeferredEditMesh([this, &BoneMatrices](FDynamicMesh3& Mesh)
+	{
+		const UE::Geometry::FDynamicMeshVertexSkinWeightsAttribute* SkinAttribute = Mesh.Attributes()->GetSkinWeightsAttribute(FSkeletalMeshAttributes::DefaultSkinWeightProfileName);
+		if (!SkinAttribute)
+		{
+			return;
+		}
+
+		UE::AnimationCore::FBoneWeights BoneInfluences;
+		for (const int32 VertexID : Mesh.VertexIndicesItr())
+		{
+			// Get the bone weights and bone indices.
+			SkinAttribute->GetValue(VertexID, BoneInfluences);
+
+			// Perform linear skinning.
+			FVector3f SkinnedPos = FVector3f::ZeroVector;
+			const FVector3f BindPoseVertex = BindPosePositions[VertexID];
+			for (int32 SkinWeightIndex = 0; SkinWeightIndex < BoneInfluences.Num(); ++SkinWeightIndex)
+			{				
+				const UE::AnimationCore::FBoneWeight& BoneInfluence = BoneInfluences[SkinWeightIndex];
+				const int32 BoneIndex = BoneInfluence.GetBoneIndex();
+				SkinnedPos += BoneMatrices[BoneIndex].TransformPosition(BindPoseVertex) * BoneInfluence.GetWeight();
+			}
+
+			Mesh.SetVertex(VertexID, FVector(SkinnedPos), false);
+		}
+	}, false);
+
+	// Update the render mesh, otherwise we don't see the changes visually.
+	PreviewMesh->NotifyDeferredEditCompleted(UPreviewMesh::ERenderUpdateMode::FastUpdate, EMeshRenderAttributeFlags::Positions, /*RebuildSpatial*/bFullUpdate);
+
+	if (bFullUpdate)
+	{
+		// Update the octree, otherwise the painting doesn't work properly.
+		// We do this by re-inserting all vertices.
+		const FDynamicMesh3* DynamicMesh = PreviewMesh->GetMesh();
+		TSet<int32> Vertices;
+		Vertices.Reserve(DynamicMesh->VertexCount());
+		for (const int32 VertexID : DynamicMesh->VertexIndicesItr())
+		{
+			Vertices.Add(VertexID);	
+		}
+		PaintTool->GetVerticesOctree().ReinsertVertices(Vertices);
+	}
+}
+
+void UMLDeformerPaintMode::SetMLDeformerEditor(UE::MLDeformer::FMLDeformerEditorToolkit* Editor)
+{ 
+	if (Toolkit.IsValid())
+	{
+		UE::MLDeformer::FMLDeformerPaintModeToolkit* PaintToolkit = static_cast<UE::MLDeformer::FMLDeformerPaintModeToolkit*>(Toolkit.Get());
+		PaintToolkit->SetMLDeformerEditor(Editor);
+	}
+}
+
+#undef LOCTEXT_NAMESPACE
diff --git a/Engine/Plugins/Animation/MLDeformer/MLDeformerFramework/Source/MLDeformerFrameworkEditor/Private/MLDeformerPaintModeToolkit.cpp b/Engine/Plugins/Animation/MLDeformer/MLDeformerFramework/Source/MLDeformerFrameworkEditor/Private/MLDeformerPaintModeToolkit.cpp
new file mode 100644
--- /dev/null
+++ b/Engine/Plugins/Animation/MLDeformer/MLDeformerFramework/Source/MLDeformerFrameworkEditor/Private/MLDeformerPaintModeToolkit.cpp
@@ -0,0 +1,298 @@
+// Copyright Epic Games, Inc. All Rights Reserved.
+
+#include "MLDeformerPaintModeToolkit.h"
+#include "MLDeformerEditorToolkit.h"
+#include "EdModeInteractiveToolsContext.h"
+#include "Framework/MultiBox/MultiBoxBuilder.h"
+#include "IDetailsView.h"
+#include "InteractiveTool.h"
+#include "InteractiveToolsContext.h"
+#include "ModelingToolsEditorModeStyle.h"
+#include "Modules/ModuleManager.h"
+#include "PropertyEditorModule.h"
+#include "Tools/UEdMode.h"
+#include "Widgets/Input/SButton.h"
+#include "Widgets/Layout/SBorder.h"
+#include "Widgets/Images/SImage.h"
+#include "Widgets/Text/STextBlock.h"
+#include "ModelingToolsManagerActions.h"
+#include "SPrimaryButton.h"
+#include "Toolkits/AssetEditorModeUILayer.h"
+#include "EditorModeManager.h"
+
+#define LOCTEXT_NAMESPACE "MLDeformerPaintModeToolkit"
+
+namespace UE::MLDeformer
+{
+	FMLDeformerPaintModeToolkit::~FMLDeformerPaintModeToolkit()
+	{
+		if (UEditorInteractiveToolsContext* Context = GetScriptableEditorMode()->GetInteractiveToolsContext(EToolsContextScope::EdMode))
+		{
+			Context->OnToolNotificationMessage.RemoveAll(this);
+			Context->OnToolWarningMessage.RemoveAll(this);
+		}
+	}
+
+	void FMLDeformerPaintModeToolkit::Init(const TSharedPtr<IToolkitHost>& InToolkitHost, TWeakObjectPtr<UEdMode> InOwningMode)
+	{
+		bUsesToolkitBuilder = true;
+
+		FModeToolkit::Init(InToolkitHost, InOwningMode);
+
+		ModeWarningArea = SNew(STextBlock)
+			.AutoWrapText(true)
+			.Font(FCoreStyle::GetDefaultFontStyle("Bold", 9))
+			.ColorAndOpacity(FSlateColor(FLinearColor(0.9f, 0.15f, 0.15f)));
+		ModeWarningArea->SetText(FText::GetEmpty());
+		ModeWarningArea->SetVisibility(EVisibility::Collapsed);
+
+		ModeHeaderArea = SNew(STextBlock)
+			.AutoWrapText(true)
+			.Font(FCoreStyle::GetDefaultFontStyle("Bold", 12));
+		ModeHeaderArea->SetText(LOCTEXT("SelectToolLabel", "Select a Tool from the Toolbar"));
+		ModeHeaderArea->SetJustification(ETextJustify::Center);
+	
+		ToolWarningArea = SNew(STextBlock)
+			.AutoWrapText(true)
+			.Font(FCoreStyle::GetDefaultFontStyle("Bold", 9))
+			.ColorAndOpacity(FSlateColor(FLinearColor(0.9f, 0.15f, 0.15f)));
+		ToolWarningArea->SetText(FText::GetEmpty());
+
+		RegisterPalettes();
+		ClearNotification();
+		ClearWarning();
+	
+		// create the toolkit widget
+		{
+			ToolkitSections->ModeWarningArea = ModeWarningArea;
+			ToolkitSections->DetailsView = ModeDetailsView;
+			ToolkitSections->ToolWarningArea = ToolWarningArea;
+
+			SAssignNew(ToolkitWidget, SBorder)
+			.HAlign(HAlign_Fill)
+			.Padding(0)
+			.BorderImage(FAppStyle::GetBrush("ToolPanel.GroupBorder"))
+			[
+				ToolkitBuilder->GenerateWidget()->AsShared()
+			];	
+		}
+	
+		ActiveToolName = FText::GetEmpty();
+		ActiveToolMessage = FText::GetEmpty();
+
+		if (UEditorInteractiveToolsContext* Context = GetScriptableEditorMode()->GetInteractiveToolsContext(EToolsContextScope::EdMode))
+		{
+			Context->OnToolNotificationMessage.AddSP(this, &FMLDeformerPaintModeToolkit::PostNotification);
+			Context->OnToolWarningMessage.AddSP(this, &FMLDeformerPaintModeToolkit::PostWarning);
+		}
+
+		// add viewport overlay widget to accept / cancel tool
+		MakeToolAcceptCancelWidget();
+	}
+
+	FName FMLDeformerPaintModeToolkit::GetToolkitFName() const
+	{
+		return FName("MLDeformerPaintModeToolkit");
+	}
+
+	FText FMLDeformerPaintModeToolkit::GetBaseToolkitName() const
+	{
+		return LOCTEXT("DisplayName", "ML Deformer Paint Mode");
+	}
+
+	void FMLDeformerPaintModeToolkit::OnToolStarted(UInteractiveToolManager* Manager, UInteractiveTool* Tool)
+	{
+		UpdateActiveToolProperties(Tool);
+
+		Tool->OnPropertySetsModified.AddSP(this, &FMLDeformerPaintModeToolkit::UpdateActiveToolProperties, Tool);
+
+		ModeHeaderArea->SetVisibility(EVisibility::Collapsed);
+		ActiveToolName = Tool->GetToolInfo().ToolDisplayName;
+
+		FString ActiveToolIdentifier = GetScriptableEditorMode()->GetToolManager()->GetActiveToolName(EToolSide::Left);
+		ActiveToolIdentifier.InsertAt(0, ".");
+		FName ActiveToolIconName = ISlateStyle::Join(FModelingToolsManagerCommands::Get().GetContextName(), TCHAR_TO_ANSI(*ActiveToolIdentifier));
+		ActiveToolIcon = FModelingToolsEditorModeStyle::Get()->GetOptionalBrush(ActiveToolIconName);
+	
+		GetToolkitHost()->AddViewportOverlayWidget(ViewportOverlayWidget.ToSharedRef());
+	}
+
+	void FMLDeformerPaintModeToolkit::OnToolEnded(UInteractiveToolManager* Manager, UInteractiveTool* Tool)
+	{
+		if (IsHosted())
+		{
+			GetToolkitHost()->RemoveViewportOverlayWidget(ViewportOverlayWidget.ToSharedRef());
+		}
+	
+		if (Tool)
+		{
+			Tool->OnPropertySetsModified.RemoveAll(this);
+		}
+
+		ModeHeaderArea->SetVisibility(EVisibility::Visible);
+		ModeDetailsView->SetObject(nullptr);
+		ActiveToolName = FText::GetEmpty();
+		ClearNotification();
+		ClearWarning();
+
+		if (MLDeformerEditor)
+		{
+			MLDeformerEditor->SetNeedsPaintModeDisable(true);
+		}
+	}
+
+	FText FMLDeformerPaintModeToolkit::GetActiveToolDisplayName() const
+	{
+		return ActiveToolName;
+	}
+
+	FText FMLDeformerPaintModeToolkit::GetActiveToolMessage() const
+	{
+		return ActiveToolMessage;
+	}
+
+	void FMLDeformerPaintModeToolkit::RegisterPalettes()
+	{
+		const FModelingToolsManagerCommands& Commands = FModelingToolsManagerCommands::Get();
+		const TSharedPtr<FUICommandList> CommandList = GetToolkitCommands();
+	
+		ToolkitSections = MakeShared<FToolkitSections>();
+
+		FToolkitBuilderArgs ToolkitBuilderArgs(GetScriptableEditorMode()->GetModeInfo().ToolbarCustomizationName);
+		ToolkitBuilderArgs.ToolkitCommandList = GetToolkitCommands();
+		ToolkitBuilderArgs.ToolkitSections = ToolkitSections;
+		ToolkitBuilderArgs.SelectedCategoryTitleVisibility = EVisibility::Collapsed;
+		ToolkitBuilder = MakeShared<FToolkitBuilder>(ToolkitBuilderArgs);
+
+		const TArray<TSharedPtr<FUICommandInfo>> SkinCommands({
+			Commands.BeginMeshAttributePaintTool
+		});
+		ToolkitBuilder->AddPalette(MakeShareable(new FToolPalette(Commands.LoadSkinTools.ToSharedRef(), SkinCommands)));
+		ToolkitBuilder->SetActivePaletteOnLoad(Commands.LoadSkinTools.Get());
+		ToolkitBuilder->UpdateWidget();
+	
+		// If selected palette changes, make sure we are showing the palette command buttons, which may be hidden by active Tool.
+		ActivePaletteChangedHandle = ToolkitBuilder->OnActivePaletteChanged.AddLambda([this]()
+		{
+			ToolkitBuilder->SetActivePaletteCommandsVisibility(EVisibility::Visible);
+		});
+	}
+
+	void FMLDeformerPaintModeToolkit::MakeToolAcceptCancelWidget()
+	{
+		SAssignNew(ViewportOverlayWidget, SHorizontalBox)
+		+SHorizontalBox::Slot()
+		.HAlign(HAlign_Center)
+		.VAlign(VAlign_Bottom)
+		.Padding(FMargin(0.0f, 0.0f, 0.f, 15.f))
+		[
+			SNew(SBorder)
+			.BorderImage(FAppStyle::Get().GetBrush("EditorViewport.OverlayBrush"))
+			.Padding(8.f)
+			[
+				SNew(SHorizontalBox)
+				+SHorizontalBox::Slot()
+				.AutoWidth()
+				.VAlign(VAlign_Center)
+				.Padding(FMargin(0.f, 0.f, 8.f, 0.f))
+				[
+					SNew(SImage)
+					.Image_Lambda([this] () { return ActiveToolIcon; })
+				]
+				+SHorizontalBox::Slot()
+				.AutoWidth()
+				.VAlign(VAlign_Center)
+				.Padding(FMargin(0.f, 0.f, 8.f, 0.f))
+				[
+					SNew(STextBlock)
+					.Text(this, &FMLDeformerPaintModeToolkit::GetActiveToolDisplayName)
+				]
+				+SHorizontalBox::Slot()
+				.AutoWidth()
+				.Padding(FMargin(0.0, 0.f, 2.f, 0.f))
+				[
+					SNew(SPrimaryButton)
+					.Text(LOCTEXT("OverlayAccept", "Accept"))
+					.ToolTipText(LOCTEXT("OverlayAcceptTooltip", "Accept/Commit the results of the active Tool [Enter]"))
+					.OnClicked_Lambda([this]() { GetScriptableEditorMode()->GetInteractiveToolsContext()->EndTool(EToolShutdownType::Accept); return FReply::Handled(); })
+					.IsEnabled_Lambda([this]() { return GetScriptableEditorMode()->GetInteractiveToolsContext()->CanAcceptActiveTool(); })
+					.Visibility_Lambda([this]() { return GetScriptableEditorMode()->GetInteractiveToolsContext()->ActiveToolHasAccept() ? EVisibility::Visible : EVisibility::Collapsed; })
+				]
+				+SHorizontalBox::Slot()
+				.AutoWidth()
+				.Padding(FMargin(2.0, 0.f, 0.f, 0.f))
+				[
+					SNew(SButton)
+					.TextStyle( FAppStyle::Get(), "DialogButtonText" )
+					.Text(LOCTEXT("OverlayCancel", "Cancel"))
+					.ToolTipText(LOCTEXT("OverlayCancelTooltip", "Cancel the active Tool [Esc]"))
+					.HAlign(HAlign_Center)
+					.OnClicked_Lambda([this]() { GetScriptableEditorMode()->GetInteractiveToolsContext()->EndTool(EToolShutdownType::Cancel); return FReply::Handled(); })
+					.IsEnabled_Lambda([this]() { return GetScriptableEditorMode()->GetInteractiveToolsContext()->CanCancelActiveTool(); })
+					.Visibility_Lambda([this]() { return GetScriptableEditorMode()->GetInteractiveToolsContext()->ActiveToolHasAccept() ? EVisibility::Visible : EVisibility::Collapsed; })
+				]
+				+SHorizontalBox::Slot()
+				.AutoWidth()
+				.Padding(FMargin(2.0, 0.f, 0.f, 0.f))
+				[
+					SNew(SPrimaryButton)
+					.Text(LOCTEXT("OverlayComplete", "Complete"))
+					.ToolTipText(LOCTEXT("OverlayCompleteTooltip", "Exit the active Tool [Enter]"))
+					.OnClicked_Lambda([this]() { GetScriptableEditorMode()->GetInteractiveToolsContext()->EndTool(EToolShutdownType::Completed); return FReply::Handled(); })
+					.IsEnabled_Lambda([this]() { return GetScriptableEditorMode()->GetInteractiveToolsContext()->CanCompleteActiveTool(); })
+					.Visibility_Lambda([this]() { return GetScriptableEditorMode()->GetInteractiveToolsContext()->CanCompleteActiveTool() ? EVisibility::Visible : EVisibility::Collapsed; })
+				]
+			]	
+		];
+	}
+
+	void FMLDeformerPaintModeToolkit::PostNotification(const FText& InMessage)
+	{
+		ClearNotification();
+	
+		ActiveToolMessage = InMessage;
+
+		if (ModeUILayer.IsValid())
+		{
+			const FName StatusBarName = ModeUILayer.Pin()->GetStatusBarName();
+			ActiveToolMessageHandle = GEditor->GetEditorSubsystem<UStatusBarSubsystem>()->PushStatusBarMessage(StatusBarName, ActiveToolMessage);
+		}
+	}
+
+	void FMLDeformerPaintModeToolkit::ClearNotification()
+	{
+		ActiveToolMessage = FText::GetEmpty();
+
+		if (ModeUILayer.IsValid())
+		{
+			const FName StatusBarName = ModeUILayer.Pin()->GetStatusBarName();
+			GEditor->GetEditorSubsystem<UStatusBarSubsystem>()->PopStatusBarMessage(StatusBarName, ActiveToolMessageHandle);
+		}
+		ActiveToolMessageHandle.Reset();
+	}
+
+	void FMLDeformerPaintModeToolkit::PostWarning(const FText& Message)
+	{
+		ToolWarningArea->SetText(Message);
+		ToolWarningArea->SetVisibility(EVisibility::Visible);
+	}
+
+	void FMLDeformerPaintModeToolkit::ClearWarning()
+	{
+		ToolWarningArea->SetText(FText());
+		ToolWarningArea->SetVisibility(EVisibility::Collapsed);
+	}
+
+	void FMLDeformerPaintModeToolkit::UpdateActiveToolProperties(UInteractiveTool* Tool)
+	{
+		UInteractiveTool* CurTool = GetScriptableEditorMode()->GetToolManager(EToolsContextScope::EdMode)->GetActiveTool(EToolSide::Left);
+		if (CurTool == nullptr)
+		{
+			return;
+		}
+		
+		ModeDetailsView->SetObjects(CurTool->GetToolProperties(true));
+	}
+} // namespace UE::MLDeformer
+
+#undef LOCTEXT_NAMESPACE
diff --git a/Engine/Plugins/Animation/MLDeformer/MLDeformerFramework/Source/MLDeformerFrameworkEditor/Private/SMLDeformerMaskConfigDialog.cpp b/Engine/Plugins/Animation/MLDeformer/MLDeformerFramework/Source/MLDeformerFrameworkEditor/Private/SMLDeformerMaskConfigDialog.cpp
new file mode 100644
--- /dev/null
+++ b/Engine/Plugins/Animation/MLDeformer/MLDeformerFramework/Source/MLDeformerFrameworkEditor/Private/SMLDeformerMaskConfigDialog.cpp
@@ -0,0 +1,217 @@
+// Copyright Epic Games, Inc. All Rights Reserved.
+
+#include "SMLDeformerMaskConfigDialog.h"
+#include "SMLDeformerNewVertexAttributeDialog.h"
+#include "MLDeformerEditorModel.h"
+#include "SkeletalMeshAttributes.h"
+#include "MeshDescription.h"
+#include "Engine/SkeletalMesh.h"
+#include "Widgets/Input/SButton.h"
+#include "Widgets/Input/STextComboBox.h"
+#include "Widgets/Text/STextBlock.h"
+
+#define LOCTEXT_NAMESPACE "MLDeformerMaskConfigDialog"
+
+namespace UE::MLDeformer
+{
+	void SMLDeformerMaskConfigDialog::Construct(const FArguments& InArgs, FMLDeformerEditorModel* InEditorModel)
+	{
+		check(InEditorModel);
+
+		MaskInfo = InArgs._InitialMaskInfo;
+		EditorModel = InEditorModel;
+		OnSetNewVertexAttributeValues = InArgs._OnSetNewVertexAttributeValues;
+
+		MaskingModeNames.Add(MakeShared<FString>(LOCTEXT("MaskModeNameGenerated", "Auto Generated").ToString()));
+		MaskingModeNames.Add(MakeShared<FString>(LOCTEXT("MaskModeNameVertexAttribute", "Mesh Vertex Attribute").ToString()));
+
+		// The widget that shows up when using vertex attributes.
+		TSharedPtr<SWidget> MeshAttributeModeWidget = CreateMeshAttributeModeWidget();
+		MeshAttributeModeWidget->SetVisibility(TAttribute<EVisibility>::CreateLambda([this]()
+		{					
+			return (MaskInfo.MaskMode == EMLDeformerMaskingMode::VertexAttribute) ? EVisibility::Visible : EVisibility::Hidden;
+		}));
+
+		UpdateAttributeNames();
+
+		SCustomDialog::Construct
+		(
+			SCustomDialog::FArguments()
+			.AutoCloseOnButtonPress(true)
+			.Title(LOCTEXT("DialogTitle", "Mask Configuration"))
+			.UseScrollBox(false)
+			.Buttons(
+			{
+				SCustomDialog::FButton(LOCTEXT("OKText", "OK"))
+				.SetPrimary(true)
+				.SetFocus(),
+				SCustomDialog::FButton(LOCTEXT("CancelText", "Cancel"))
+			})
+			.Content()
+			[
+				SNew(SBox)
+				.Padding(FMargin(10.0f, 10.0f))
+				.MinDesiredWidth(400.0f)
+				.HAlign(EHorizontalAlignment::HAlign_Fill)
+				.VAlign(EVerticalAlignment::VAlign_Fill)
+				[
+					SNew(SVerticalBox)
+					+SVerticalBox::Slot()
+					.AutoHeight()
+					.Padding(4.0f)
+					[
+						SNew(SHorizontalBox)
+						+SHorizontalBox::Slot()
+						.AutoWidth()
+						.Padding(4.0f)
+						.VAlign(EVerticalAlignment::VAlign_Center)
+						[
+							SNew(STextBlock)
+							.Text(LOCTEXT("MaskingModeLabel", "Masking Mode:"))
+						]
+						+SHorizontalBox::Slot()
+						.AutoWidth()
+						.Padding(4.0f)
+						.VAlign(EVerticalAlignment::VAlign_Center)
+						[
+							SNew(SComboBox<TSharedPtr<FString>>)
+							.ToolTipText(LOCTEXT("MaskingModeToolTip", "Specify whether you would like to use auto-generated masks or use a vertex attribute on the mesh, which can be painted."))
+							.OptionsSource(&MaskingModeNames)
+							.OnGenerateWidget_Lambda([this](const TSharedPtr<FString>& Name)
+							{
+								return SNew(STextBlock)
+								.Text(FText::FromString(*Name));
+							})
+							.OnSelectionChanged_Lambda([this](TSharedPtr<FString> Selected, ESelectInfo::Type)
+							{
+								const int32 Index = MaskingModeNames.Find(Selected);
+								check(Index != INDEX_NONE);
+								MaskInfo.MaskMode = static_cast<EMLDeformerMaskingMode>(Index);
+							})
+							[
+								SNew(STextBlock)
+								.Text_Lambda([this]() { return FText::FromString(*MaskModeEnumToString(MaskInfo.MaskMode)); }) 
+							]
+						]
+					]
+					+SVerticalBox::Slot()
+					.AutoHeight()
+					.Padding(4.0f)
+					[
+						MeshAttributeModeWidget.ToSharedRef()
+					]
+				]
+			]
+		);
+	}
+
+	void SMLDeformerMaskConfigDialog::UpdateAttributeNames()
+	{
+		check(EditorModel);
+		AttributeNames = EditorModel->GetModel()->GetVertexAttributeNames();
+	}
+
+	TSharedPtr<FString> SMLDeformerMaskConfigDialog::MaskModeEnumToString(EMLDeformerMaskingMode MaskMode) const
+	{
+		const int32 Index = static_cast<int32>(MaskMode);
+		return MaskingModeNames[Index];
+	}
+
+	TSharedPtr<SWidget> SMLDeformerMaskConfigDialog::CreateMeshAttributeModeWidget()
+	{
+		return 
+			SNew(SVerticalBox)
+			+SVerticalBox::Slot()
+			.AutoHeight()
+			.Padding(0.0f, 2.0f)
+			[
+				SNew(SHorizontalBox)
+				+SHorizontalBox::Slot()
+				.Padding(4.0f, 0.0f)
+				.AutoWidth()
+				.VAlign(EVerticalAlignment::VAlign_Center)
+				[
+					SNew(STextBlock)
+					.Text(LOCTEXT("AttributeName", "Attribute Name:"))
+				]
+				+SHorizontalBox::Slot()
+				.Padding(4.0f, 0.0f)
+				.VAlign(EVerticalAlignment::VAlign_Center)
+				[
+					SAssignNew(VertexAttributeComboWidget, SComboBox<FName>)
+					.OptionsSource(&AttributeNames)
+					.OnGenerateWidget_Lambda([this](const FName& Name)
+					{
+						return SNew(STextBlock)
+						.Text(FText::FromName(Name));
+					})
+					.OnSelectionChanged_Lambda([this](FName Selected, ESelectInfo::Type)
+					{
+						MaskInfo.VertexAttributeName = Selected;
+					})
+					[
+						SNew(STextBlock)
+						.Text_Lambda([this]() { return FText::FromName(MaskInfo.VertexAttributeName); }) 
+					]
+				]
+				+SHorizontalBox::Slot()
+				.VAlign(EVerticalAlignment::VAlign_Center)
+				.Padding(0.0f, 2.0f)
+				.AutoWidth()
+				[
+					SNew(SButton)
+					.Text(LOCTEXT("CreateButton", "Create New"))
+					.ToolTipText(LOCTEXT("CreateButtonTooltip", "Create a new vertex attribute on the Skeletal Mesh. This will modify the Skeletal Mesh, so please make sure to save it."))
+					.HAlign(EHorizontalAlignment::HAlign_Center)
+					.VAlign(EVerticalAlignment::VAlign_Center)
+					.OnClicked_Raw(this, &SMLDeformerMaskConfigDialog::OnCreateVertexAttributeClicked)
+				]
+			];
+	}
+
+	FReply SMLDeformerMaskConfigDialog::OnCreateVertexAttributeClicked()
+	{
+		USkeletalMesh* SkeletalMesh = EditorModel->GetModel()->GetSkeletalMesh();
+		if (!SkeletalMesh)
+		{
+			return FReply::Handled();
+		}
+
+		// Show the Create New VertexAttribute dialog, which allows the user to enter a name to create a new one.
+		TSharedPtr<SMLDeformerNewVertexAttributeDialog> Dialog = SNew(SMLDeformerNewVertexAttributeDialog, SkeletalMesh)
+			.bAutoCreateAttribute(true)
+			.DefaultAttributeValue(0.0f);
+
+		const SMLDeformerNewVertexAttributeDialog::EReturnCode ReturnCode = Dialog->ShowModal();
+
+		// If we created a new attribute, update the list of attributes and auto select the new one.
+		// Also initialize the vertex values to the mask as if it was generated.
+		if (ReturnCode == SMLDeformerNewVertexAttributeDialog::EReturnCode::CreatePressed)
+		{
+			// Update the combobox items.
+			UpdateAttributeNames();
+
+			// Select the newly created item directly.
+			const FName NewAttributeName = *Dialog->GetAttributeName();
+			if (VertexAttributeComboWidget.IsValid())
+			{
+				VertexAttributeComboWidget->SetSelectedItem(NewAttributeName);
+			}
+
+			// Generate the mask and set the values to the vertex attributes.
+			const int32 LODIndex = 0;
+			FMeshDescription* MeshDescription = SkeletalMesh->GetMeshDescription(LODIndex);
+			if (MeshDescription)
+			{
+				check(MeshDescription->VertexAttributes().HasAttribute(NewAttributeName));
+				TVertexAttributesRef<float> AttributeRef = MeshDescription->VertexAttributes().GetAttributesRef<float>(NewAttributeName);
+				OnSetNewVertexAttributeValues.ExecuteIfBound(AttributeRef);
+			}
+		}
+	
+		return FReply::Handled();
+	}
+
+}	// namespace UE::MLDeformer
+
+#undef LOCTEXT_NAMESPACE
diff --git a/Engine/Plugins/Animation/MLDeformer/MLDeformerFramework/Source/MLDeformerFrameworkEditor/Private/SMLDeformerNewVertexAttributeDialog.cpp b/Engine/Plugins/Animation/MLDeformer/MLDeformerFramework/Source/MLDeformerFrameworkEditor/Private/SMLDeformerNewVertexAttributeDialog.cpp
new file mode 100644
--- /dev/null
+++ b/Engine/Plugins/Animation/MLDeformer/MLDeformerFramework/Source/MLDeformerFrameworkEditor/Private/SMLDeformerNewVertexAttributeDialog.cpp
@@ -0,0 +1,192 @@
+// Copyright Epic Games, Inc. All Rights Reserved.
+
+#include "SMLDeformerNewVertexAttributeDialog.h"
+#include "SkeletalMeshAttributes.h"
+#include "MeshDescription.h"
+#include "Engine/SkeletalMesh.h"
+#include "Editor/EditorEngine.h"
+#include "Widgets/Input/SButton.h"
+#include "Widgets/Input/SEditableTextBox.h"
+#include "Widgets/Text/STextBlock.h"
+#include "Widgets/Layout/SBox.h"
+
+#define LOCTEXT_NAMESPACE "MLDeformerNewVertexAttributeDialog"
+
+namespace UE::MLDeformer
+{
+	void SMLDeformerNewVertexAttributeDialog::Construct(const FArguments& InArgs, USkeletalMesh* InSkeletalMesh)
+	{
+		SkeletalMesh = InSkeletalMesh;
+
+		bAutoCreate = InArgs._bAutoCreateAttribute;
+		DefaultAttributeValue = InArgs._DefaultAttributeValue;
+
+		TSharedPtr<SEditableTextBox> AttributeNameEditWidget;
+		SWindow::Construct
+		(
+			SWindow::FArguments()
+			.Title(LOCTEXT("EnterAttributesNameDialogTitle", "Create New Vertex Attribute"))
+			.SizingRule(ESizingRule::Autosized)
+			.SupportsMinimize(false)
+			.SupportsMaximize(false)
+			.IsTopmostWindow(true)
+			[
+				SNew(SBox)
+				.Padding(FMargin(10.0f, 10.0f))
+				.HAlign(EHorizontalAlignment::HAlign_Fill)
+				.VAlign(EVerticalAlignment::VAlign_Fill)
+				[
+					SNew(SVerticalBox)
+					+SVerticalBox::Slot()
+					.Padding(4.0f)
+					[
+						SNew(SHorizontalBox)
+						+SHorizontalBox::Slot()
+						.AutoWidth()
+						.Padding(0.0f, 0.0f, 4.0f, 0.0f)
+						.VAlign(EVerticalAlignment::VAlign_Center)
+						[
+							SNew(STextBlock)
+							.Text(LOCTEXT("CreateNewAttributeNameLabel", "Attribute Name:"))
+						]
+						+SHorizontalBox::Slot()
+						.VAlign(EVerticalAlignment::VAlign_Center)
+						[
+							SNew(SBox)
+							.MinDesiredWidth(200.0f)
+							[
+								SAssignNew(AttributeNameEditWidget, SEditableTextBox)
+								.OnKeyDownHandler(this, &SMLDeformerNewVertexAttributeDialog::OnKeyDown)
+								.OnTextChanged_Lambda([this](const FText& InText)
+								{
+									AttributeName = InText.ToString();
+								})
+							]
+						]
+					]
+					+SVerticalBox::Slot()
+					.Padding(4.0f)
+					.AutoHeight()
+					[
+						SNew(STextBlock)
+						.Text(LOCTEXT("AttributeExistsError", "There already is an attribute with that name!"))
+						.ColorAndOpacity(FLinearColor::Red)
+						.Visibility_Lambda([this]() { return HasVertexAttribute(SkeletalMesh, *AttributeName) ? EVisibility::Visible : EVisibility::Hidden; })
+					]
+					+SVerticalBox::Slot()
+					.AutoHeight()
+					.Padding(4.0f)
+					.HAlign(EHorizontalAlignment::HAlign_Right)
+					.VAlign(EVerticalAlignment::VAlign_Center)
+					[
+						SNew(SHorizontalBox)
+						+SHorizontalBox::Slot()
+						.AutoWidth()
+						.Padding(FAppStyle::Get().GetMargin("StandardDialog.SlotPadding"))
+						[
+							SAssignNew(CreateButton, SButton)
+							.Text(LOCTEXT("CreateButtonText", "Create"))
+							.ButtonStyle(&FAppStyle::Get().GetWidgetStyle<FButtonStyle>("PrimaryButton"))
+							.IsEnabled_Lambda([this]()
+							{
+								return !AttributeName.IsEmpty() && !HasVertexAttribute(SkeletalMesh, *AttributeName);
+							})
+							.OnClicked(this, &SMLDeformerNewVertexAttributeDialog::OnCreateClicked)
+						]
+						+SHorizontalBox::Slot()
+						.AutoWidth()
+						.Padding(FAppStyle::Get().GetMargin("StandardDialog.SlotPadding"))
+						[
+							SNew(SButton)
+							.Text(LOCTEXT("CancelButtonText", "Cancel"))
+							.ButtonStyle(&FAppStyle::Get().GetWidgetStyle<FButtonStyle>("Button"))
+							.OnClicked_Lambda([this]()
+							{
+								ReturnCode = EReturnCode::Canceled;
+								RequestDestroyWindow();
+								return FReply::Handled();
+							})
+						]
+					]
+				]
+			]
+		);
+
+		SetWidgetToFocusOnActivate(AttributeNameEditWidget);
+	}
+
+	FReply SMLDeformerNewVertexAttributeDialog::OnCreateClicked()
+	{
+		if (bAutoCreate)
+		{
+			check(SkeletalMesh);
+			const bool bCreateResult = CreateVertexAttribute(*SkeletalMesh, AttributeName, DefaultAttributeValue);
+			check(bCreateResult);
+		}
+		ReturnCode = EReturnCode::CreatePressed;
+		RequestDestroyWindow();
+		return FReply::Handled();
+	}
+
+	FReply SMLDeformerNewVertexAttributeDialog::OnKeyDown(const FGeometry& InGeometry, const FKeyEvent& InKeyEvent)
+	{		
+		if (InKeyEvent.GetKey() == EKeys::Escape)
+		{
+			ReturnCode = EReturnCode::Canceled;
+			RequestDestroyWindow();
+			return FReply::Handled();
+		}
+
+		if (InKeyEvent.GetKey() == EKeys::Enter && CreateButton.IsValid() && CreateButton->IsEnabled())
+		{
+			return OnCreateClicked();
+		}
+
+		return FReply::Unhandled();
+	}
+
+	bool SMLDeformerNewVertexAttributeDialog::HasVertexAttribute(const USkeletalMesh* SkeletalMesh, FName InAttributeName)
+	{
+		if (SkeletalMesh)
+		{
+			const int32 LODIndex = 0;
+			const FMeshDescription* MeshDescription = SkeletalMesh->GetMeshDescription(LODIndex);
+			if (MeshDescription)
+			{
+				return MeshDescription->VertexAttributes().HasAttribute(InAttributeName);
+			}
+		}
+
+		return false;
+	}
+
+	SMLDeformerNewVertexAttributeDialog::EReturnCode SMLDeformerNewVertexAttributeDialog::ShowModal()
+	{
+		UEditorEngine* Editor = Cast<UEditorEngine>(GEngine);
+		check(Editor);
+		Editor->EditorAddModalWindow(SharedThis(this));
+		return ReturnCode;
+	}
+
+	bool SMLDeformerNewVertexAttributeDialog::CreateVertexAttribute(USkeletalMesh& SkeletalMesh, const FString& AttributeName, float DefaultValue)
+	{
+		const int32 LODIndex = 0;
+		FMeshDescription* MeshDescription = SkeletalMesh.GetMeshDescription(LODIndex);
+		if (!MeshDescription)
+		{
+			return false;
+		}
+
+		if (MeshDescription->VertexAttributes().HasAttribute(*AttributeName))
+		{
+			return false;
+		}
+
+		MeshDescription->VertexAttributes().RegisterAttribute<float>(*AttributeName, 1, DefaultValue);
+		
+		return SkeletalMesh.CommitMeshDescription(LODIndex);
+	}
+
+}	// namespace UE::MLDeformer
+
+#undef LOCTEXT_NAMESPACE
diff --git a/Engine/Plugins/Animation/MLDeformer/MLDeformerFramework/Source/MLDeformerFrameworkEditor/Private/SMLDeformerTimeline.cpp b/Engine/Plugins/Animation/MLDeformer/MLDeformerFramework/Source/MLDeformerFrameworkEditor/Private/SMLDeformerTimeline.cpp
--- a/Engine/Plugins/Animation/MLDeformer/MLDeformerFramework/Source/MLDeformerFrameworkEditor/Private/SMLDeformerTimeline.cpp
+++ b/Engine/Plugins/Animation/MLDeformer/MLDeformerFramework/Source/MLDeformerFrameworkEditor/Private/SMLDeformerTimeline.cpp
@@ -1221,7 +1221,7 @@ namespace UE::MLDeformer
 		{
 			FFrameNumber ScrubPosition = Model.Pin()->GetTickResScrubPosition();
 			ScrubPosition.Value += Model.Pin()->GetTicksPerFrame();
-			Model.Pin()->SetScrubPosition(ScrubPosition);
+			Model.Pin()->SetScrubPosition(ScrubPosition, false);
 			return FReply::Handled();
 		}
 		return FReply::Unhandled();
@@ -1233,7 +1233,7 @@ namespace UE::MLDeformer
 		{
 			const TRange<FFrameNumber> PlaybackRange = Model.Pin()->GetPlaybackRange();
 			FFrameNumber UpperBound = PlaybackRange.GetUpperBoundValue();
-			Model.Pin()->SetScrubPosition(UpperBound);
+			Model.Pin()->SetScrubPosition(UpperBound, false);
 			return FReply::Handled();
 		}
 		return FReply::Unhandled();
@@ -1245,7 +1245,7 @@ namespace UE::MLDeformer
 		{
 			FFrameNumber ScrubPosition = Model.Pin()->GetTickResScrubPosition();
 			ScrubPosition.Value -= Model.Pin()->GetTicksPerFrame();
-			Model.Pin()->SetScrubPosition(ScrubPosition);
+			Model.Pin()->SetScrubPosition(ScrubPosition, false);
 			return FReply::Handled();
 		}
 		return FReply::Unhandled();
@@ -1257,7 +1257,7 @@ namespace UE::MLDeformer
 		{
 			const TRange<FFrameNumber> PlaybackRange = Model.Pin()->GetPlaybackRange();
 			FFrameNumber LowerBound = PlaybackRange.GetLowerBoundValue();
-			Model.Pin()->SetScrubPosition(LowerBound);
+			Model.Pin()->SetScrubPosition(LowerBound, false);
 			return FReply::Handled();
 		}
 		return FReply::Unhandled();
@@ -1905,7 +1905,7 @@ namespace UE::MLDeformer
 	{
 		if (Model.IsValid())
 		{
-			Model.Pin()->SetScrubPosition(NewScrubPosition);
+			Model.Pin()->SetScrubPosition(NewScrubPosition, bIsScrubbing);
 		}
 	}
 
@@ -1931,7 +1931,7 @@ namespace UE::MLDeformer
 		const FFrameTime FrameTime((int)InFrameTime, 0.0f);
 		if (Model.IsValid())
 		{
-			Model.Pin()->SetScrubPosition(FrameTime);
+			Model.Pin()->SetScrubPosition(FrameTime, false);
 		}
 	}
 
diff --git a/Engine/Plugins/Animation/MLDeformer/MLDeformerFramework/Source/MLDeformerFrameworkEditor/Public/MLDeformerEditorModel.h b/Engine/Plugins/Animation/MLDeformer/MLDeformerFramework/Source/MLDeformerFrameworkEditor/Public/MLDeformerEditorModel.h
--- a/Engine/Plugins/Animation/MLDeformer/MLDeformerFramework/Source/MLDeformerFrameworkEditor/Public/MLDeformerEditorModel.h
+++ b/Engine/Plugins/Animation/MLDeformer/MLDeformerFramework/Source/MLDeformerFrameworkEditor/Public/MLDeformerEditorModel.h
@@ -13,6 +13,7 @@
 #include "MLDeformerModule.h"
 #include "MLDeformerModel.h"
 #include "MLDeformerSampler.h"
+#include "MeshDescription.h"
 #include "MLDeformerEditorModel.generated.h"
 
 class UMLDeformerModel;
@@ -500,15 +501,23 @@ namespace UE::MLDeformer
 		 * Go to a given frame in the training data.
 		 * Internally this will also call OnTimeSliderScrubPositionChanged, to go to the specific frame.
 		 * @param FrameNumber The training frame to go to. This will automatically be clamped to a valid range in case it goes out of bounds.
+		 * @param bIsScrubbing This is set to true when we are scrubbing the timeline. It is set to false when the mouse button is released.
 		 */
-		virtual void SetTrainingFrame(int32 FrameNumber);
+		virtual void SetTrainingFrame(int32 FrameNumber, bool bIsScrubbing);
+
+		UE_DEPRECATED(5.5, "Please use the SetTrainingFrame that takes a bIsScrubbing parameter.")
+		virtual void SetTrainingFrame(int32 FrameNumber) { SetTrainingFrame(FrameNumber, false); }
 
 		/**
 		 * Go to a given frame in the test data.
 		 * Internally this will also call OnTimeSliderScrubPositionChanged, to go to the specific frame.
 		 * @param FrameNumber The test frame to go to. This will automatically be clamped to a valid range in case it goes out of bounds.
+		 * @param bIsScrubbing This is set to true when we are scrubbing the timeline. It is set to false when the mouse button is released.
 		 */
-		virtual void SetTestFrame(int32 FrameNumber);
+		virtual void SetTestFrame(int32 FrameNumber, bool bIsScrubbing);
+
+		UE_DEPRECATED(5.5, "Please use the SetTestFrame that takes a bIsScrubbing parameter.")
+		virtual void SetTestFrame(int32 FrameNumber) { SetTestFrame(FrameNumber, false); }
 
 		/**
 		 * Render additional (debug) things in the viewport.
@@ -691,6 +700,17 @@ namespace UE::MLDeformer
 		 */
 		virtual void ApplyDebugActorTransforms(const TArray<FTransform>& DebugActorComponentSpaceTransforms);
 
+		/**
+		 * Update the character's pose in paint mode.
+		 * This will check whether we are in training or testing mode and choose the editor actor based on that.
+		 * The paint mode will then use this specific editor actor's skeletal mesh component to calculate skinned vertex positions at its current pose.
+		 * After that it will update the dynamic mesh in the paint mode in order to paint in this same pose.
+		 * This method is called when you start painting, or scrub the timeline.
+		 * @param bFullUpdate When set to true, this can internally update the acceleration structures, which is slow. When set to false it can quickly update the visual mesh
+		 *                    It is still important to once call it with bFullUpdate set to true though, before you start painting.
+		 */
+		virtual void UpdatePaintModePose(bool bFullUpdate=true);
+
 		/**
 		 * Debug draw helpers inside the PIE viewport that highlight debuggable actors.
 		 * On default this will draw a bounding box around the actors that can be debugged by this model.
@@ -736,10 +756,14 @@ namespace UE::MLDeformer
 		float GetScrubTime() const;
 
 		/** Set the current scrub position. */
-		void SetScrubPosition(FFrameTime NewScrubPostion);
+		UE_DEPRECATED(5.5, "Please use the SetScrubPosition which takes a bIsScrubbing parameter.")
+		void SetScrubPosition(FFrameTime NewScrubPosition);
+		void SetScrubPosition(FFrameTime NewScrubPosition, bool bIsScrubbing);
 
 		/** Set the current scrub position. */
-		void SetScrubPosition(FFrameNumber NewScrubPostion);
+		UE_DEPRECATED(5.5, "Please use the SetScrubPosition which takes a bIsScrubbing parameter.")
+		void SetScrubPosition(FFrameNumber NewScrubPosition);
+		void SetScrubPosition(FFrameNumber NewScrubPosition, bool bIsScrubbing);
 
 		/** Set if frames are displayed. */
 		void SetDisplayFrames(bool bDisplayFrames);
@@ -903,7 +927,7 @@ namespace UE::MLDeformer
 		/**
 		 * Get the currently desired training frame number.
 		 * You can call CheckTrainingFrameChanged in order to make the desired frame the actual frame.
-		 * @return The currenty desired training frame number, which might not the the same as the real training frame number.
+		 * @return The currently desired training frame number, which might not the the same as the real training frame number.
 		 */
 		int32 GetCurrentTrainingFrame() const { return CurrentTrainingFrame; }
 
@@ -953,6 +977,16 @@ namespace UE::MLDeformer
 		/** Mark the deltas to be updated on next Tick. */
 		void InvalidateDeltas();
 
+		/** 
+		 * Find the float based vertex attributes on the skeletal map that have a given name.
+		 * You can check whether the attribute data has been found using ReturnedValue.IsValid().
+		 */
+		TVertexAttributesConstRef<float> FindVertexAttributes(FName AttributeName) const;
+
+		/** Are we scrubbing on the timeline? */
+		bool IsScrubbingTimeline() const	{ return bIsScrubbingTimeline; }
+
+
 	protected:
 		virtual void CreateSamplers();
 
@@ -1177,6 +1211,9 @@ namespace UE::MLDeformer
 
 		/** The training input animation that is selected in the timeline. */
 		int32 ActiveTrainingInputAnimIndex = INDEX_NONE;
+
+		/** Are we currently scrubbing the timeline? */
+		bool bIsScrubbingTimeline = false;
 	};
 
 	/**
diff --git a/Engine/Plugins/Animation/MLDeformer/MLDeformerFramework/Source/MLDeformerFrameworkEditor/Public/MLDeformerEditorToolkit.h b/Engine/Plugins/Animation/MLDeformer/MLDeformerFramework/Source/MLDeformerFrameworkEditor/Public/MLDeformerEditorToolkit.h
--- a/Engine/Plugins/Animation/MLDeformer/MLDeformerFramework/Source/MLDeformerFrameworkEditor/Public/MLDeformerEditorToolkit.h
+++ b/Engine/Plugins/Animation/MLDeformer/MLDeformerFramework/Source/MLDeformerFrameworkEditor/Public/MLDeformerEditorToolkit.h
@@ -72,6 +72,8 @@ namespace UE::MLDeformer
 		virtual FText GetToolkitName() const override;
 		virtual FLinearColor GetWorldCentricTabColorScale() const override;
 		virtual FString GetWorldCentricTabPrefix() const override;
+		virtual void AddViewportOverlayWidget(TSharedRef<SWidget> InViewportOverlayWidget, int32 ZOrder = INDEX_NONE) override;
+		virtual void RemoveViewportOverlayWidget(TSharedRef<SWidget> InViewportOverlayWidget) override;
 		// ~END FAssetEditorToolkit overrides.
 
 		// FGCObject overrides.
@@ -80,7 +82,7 @@ namespace UE::MLDeformer
 		// ~END FGCObject overrides.
 
 		// FTickableEditorObject overrides.
-		virtual void Tick(float DeltaTime) override {};
+		virtual void Tick(float DeltaTime) override;
 		virtual ETickableTickType GetTickableTickType() const override				{ return ETickableTickType::Always; }
 		virtual TStatId GetStatId() const override;
 		// ~END FTickableEditorObject overrides.
@@ -105,11 +107,20 @@ namespace UE::MLDeformer
 		FMLDeformerApplicationMode* GetApplicationMode() const						{ return ApplicationMode; }
 		TSharedPtr<IPersonaViewport> GetViewport() const							{ return PersonaViewport; }
 
+		void SetNeedsPaintModeDisable(bool bNeedsDisable)							{ bNeedsPaintModeDisable = bNeedsDisable; }
+
 		double CalcTimelinePosition() const;
 		void OnTimeSliderScrubPositionChanged(double NewScrubTime, bool bIsScrubbing);
 		void UpdateTimeSliderRange();
 		void SetTimeSliderRange(double StartTime, double EndTime);
 
+		void EnablePaintMode();
+		void DisablePaintMode();
+		bool IsDefaultModeActive() const;
+		bool IsPaintModeActive() const;
+
+		bool IsInitialized() const { return bIsInitialized; }
+
 		/**
 		 * Switch the editor to a given model type.
 		 * @param ModelType The model type you want to switch to, for example something like: UNeuralMorphModel::StaticClass().
@@ -175,6 +186,7 @@ namespace UE::MLDeformer
 		TSharedRef<SWidget> GenerateVizModeButtonContents(TSharedRef<FUICommandList> InCommandList);
 		TSharedRef<SWidget> GenerateToolsMenuContents(TSharedRef<FUICommandList> InCommandList);
 
+
 	private:
 		/** The persona toolkit. */	
 		TSharedPtr<IPersonaToolkit> PersonaToolkit;
@@ -209,6 +221,13 @@ namespace UE::MLDeformer
 		/** Are we currently in a training process? */
 		bool bIsTraining = false;
 
+		/** 
+		 * When set to true, the paint mode will be disabled automatically.
+		 * This is used to defer the disable of the paint mode to a later stage in the frame, as sometimes this
+		 * can cause issues by deactivating or activating modes while in a loop over all modes.
+		 */
+		bool bNeedsPaintModeDisable = false;
+
 		/** Extenders for Tools menu */
 		static TArray<TUniquePtr<FToolsMenuExtender>> ToolsMenuExtenders;
 
diff --git a/Engine/Plugins/Animation/MLDeformer/MLDeformerFramework/Source/MLDeformerFrameworkEditor/Public/MLDeformerMorphModelEditorModel.h b/Engine/Plugins/Animation/MLDeformer/MLDeformerFramework/Source/MLDeformerFrameworkEditor/Public/MLDeformerMorphModelEditorModel.h
--- a/Engine/Plugins/Animation/MLDeformer/MLDeformerFramework/Source/MLDeformerFrameworkEditor/Public/MLDeformerMorphModelEditorModel.h
+++ b/Engine/Plugins/Animation/MLDeformer/MLDeformerFramework/Source/MLDeformerFrameworkEditor/Public/MLDeformerMorphModelEditorModel.h
@@ -6,6 +6,7 @@
 #include "MLDeformerGeomCacheEditorModel.h"
 #include "MLDeformerGeomCacheSampler.h"
 #include "MLDeformerEditorActor.h"
+#include "MLDeformerMasking.h"
 #include "Rendering/ColorVertexBuffer.h"
 
 class UMLDeformerMorphModel;
@@ -35,6 +36,7 @@ namespace UE::MLDeformer
 		virtual void OnPreTraining() override;
 		virtual void OnPostTraining(ETrainingResult TrainingResult, bool bUsePartiallyTrainedWhenAborted) override;
 		virtual void OnMaxNumLODsChanged() override;
+		virtual void OnObjectModified(UObject* Object) override;
 		// ~END FMLDeformerEditorModel overrides.
 
 		/**
@@ -142,6 +144,25 @@ namespace UE::MLDeformer
 		UE_DEPRECATED(5.4, "This method will be removed.")
 		TArray<int32> BuildVirtualParentTable(const FReferenceSkeleton& RefSkel, const TArray<FName>& IncludedBoneNames) const;
 
+		/**
+		 * Apply the mask info to the mask buffer of floats.
+		 * @param SkeletalMesh The skeletal mesh that the mask relates to.
+		 * @param MaskInfo The mask info object to apply to the float mask buffer.
+		 * @param ItemMaskBuffer The mask buffer for this item. This should be an array view of size NumBaseMeshVerts.
+		 */
+		void ApplyMaskInfoToBuffer(const USkeletalMesh* SkeletalMesh, const FMLDeformerMaskInfo& MaskInfo, TArrayView<float> ItemMaskBuffer);
+
+		/**
+		 * Apply a generated mask to a specific vertex attribute.
+		 * This copy the generated values one to one to the vertex attribute, but it applies some scaling to the attributes to give better results.
+		 * If we don't do this scaling then the values often look too small and the vertex attribute visually looks very light, rather than clearly defining the area.
+		 * This will also ensure the mask gets generated, even if the MaskInfo masking mode is set to something non-generated.
+		 * @param SkeletalMesh The skeletal mesh that the vertex attribute is on.
+		 * @param MaskInfo The mask info object that we generate the mask from.
+		 * @param AttributeRef The output vertex attributes to write to.
+		 */
+		void ApplyGeneratedMaskToVertexAttributes(USkeletalMesh* SkeletalMesh, FMLDeformerMaskInfo& MaskInfo, TVertexAttributesRef<float> AttributeRef);
+
 		// Helpers.
 		UMLDeformerMorphModel* GetMorphModel() const;
 		UMLDeformerMorphModelVizSettings* GetMorphModelVizSettings() const;
@@ -184,6 +205,20 @@ namespace UE::MLDeformer
 		UE_DEPRECATED(5.3, "Please use the CalcMorphTargetNormals method that takes more parameters.")
 		virtual void CalcMorphTargetNormals(int32 LOD, USkeletalMesh* SkelMesh, int32 MorphTargetIndex, TArrayView<const FVector3f> Deltas, TArrayView<const FVector3f> BaseVertexPositions, TArrayView<FVector3f> BaseNormals, TArray<FVector3f>& OutDeltaNormals);
 
+		UE_DEPRECATED(5.5, "Please use the CalcMorphTaretNormals that takes an array named GlobalMaskWeights.")
+		virtual void CalcMorphTargetNormals(
+			int32 LOD,
+			const USkeletalMesh* SkelMesh,
+			int32 MorphTargetIndex,
+			const TArrayView<const FVector3f> Deltas,
+			const TArrayView<const FVector3f> BaseVertexPositions,
+			const TArrayView<const FVector3f> BaseNormals,
+			const TArrayView<const int32> ImportedVertexToRenderVertexMapping,
+			const FColorVertexBuffer& ColorBuffer,
+			EMLDeformerMaskChannel MaskChannel,
+			bool bInvertGlobalMaskChannel,
+			TArray<FVector3f>& OutDeltaNormals);
+
 		/**
 		 * Calculate the delta normals for a given morph target.
 		 * @param LOD The LOD level.
@@ -196,6 +231,7 @@ namespace UE::MLDeformer
 		 * @param ColorBuffer The color buffer to use when calculating the global mask.
 		 * @param MaskChannel The global mask channel, for example which color channel to use from the color buffer.
 		 * @param bInvertGlobalMask Set to true when you want the global mask to be inverted.
+		 * @param GlobalMaskWeights The array of weights, one for each render vertex.
 		 * @param OutDeltaNormals The array that we will write the generated normals to. This will automatically be resized by this method.
 		 */
 		virtual void CalcMorphTargetNormals(
@@ -209,6 +245,7 @@ namespace UE::MLDeformer
 			const FColorVertexBuffer& ColorBuffer,
 			EMLDeformerMaskChannel MaskChannel,
 			bool bInvertGlobalMaskChannel,
+			const TArray<float>& GlobalMaskWeights,
 			TArray<FVector3f>& OutDeltaNormals);
 
 		/**
@@ -283,15 +320,18 @@ namespace UE::MLDeformer
 		 */
 		bool ProcessVertexDelta(FVector3f& OutScaledDelta, FVector3f& OutScaledDeltaNormal, const FVector3f RawDelta, const FVector3f RawDeltaNormal, float DeltaThreshold, float MorphMaskWeight, float GlobalMaskWeight) const;
 
+		UE_DEPRECATED(5.5, "Please use the CalcGlobalMaskWeights method instead")
+		float CalcGlobalMaskWeight(int32 RenderVertexIndex, const FColorVertexBuffer& ColorBuffer, EMLDeformerMaskChannel MaskChannel, bool bInvertMaskChannel) const;
+
 		/**
 		 * Calculate the global mask weight for a specific render vertex, using a color buffer and specified mask channel settings.
-		 * @param RenderVertexIndex The render vertex index to calculate the global deformer mask weight for.
+		 * @param VertexMap The render vertex to import vertex map.
 		 * @param ColorBuffer The color buffer which we can grab color values from, which are used to calculate the mask weight.
 		 * @param MaskChannel The channel to grab the weight value from. This specifieds whether the mask is enabled and if so, from what channel (r, g, b, a) we grab the weight value.
 		 * @param bInvertMaskChannel Set to true when we should invert the weight value.
-		 * @return Returns the mask weight. A value of 1.0 is returned when the mask channel is set to disable the mask. The return value is always between 0 and 1.
+		 * @return An array with a weight per render vertex. A value of 1.0 is returned when the mask channel is set to disable the mask. The return value is always between 0 and 1.
 		 */
-		float CalcGlobalMaskWeight(int32 RenderVertexIndex, const FColorVertexBuffer& ColorBuffer, EMLDeformerMaskChannel MaskChannel, bool bInvertMaskChannel) const;
+		TArray<float> CalcGlobalMaskWeights(const TArray<int32>& VertexMap, const FColorVertexBuffer& ColorBuffer, EMLDeformerMaskChannel MaskChannel, bool bInvertMaskChannel) const;
 
 	protected:
 		/**
diff --git a/Engine/Plugins/Animation/MLDeformer/MLDeformerFramework/Source/MLDeformerFrameworkEditor/Public/MLDeformerPaintMode.h b/Engine/Plugins/Animation/MLDeformer/MLDeformerFramework/Source/MLDeformerFrameworkEditor/Public/MLDeformerPaintMode.h
new file mode 100644
--- /dev/null
+++ b/Engine/Plugins/Animation/MLDeformer/MLDeformerFramework/Source/MLDeformerFrameworkEditor/Public/MLDeformerPaintMode.h
@@ -0,0 +1,40 @@
+// Copyright Epic Games, Inc. All Rights Reserved.
+
+#pragma once
+
+#include "CoreTypes.h"
+#include "EdMode.h"
+#include "MLDeformerPaintMode.generated.h"
+
+class USkeletalMeshComponent;
+
+namespace UE::MLDeformer
+{
+	class FMLDeformerEditorToolkit;
+}
+
+UCLASS()
+class MLDEFORMERFRAMEWORKEDITOR_API UMLDeformerPaintMode
+	: public UBaseLegacyWidgetEdMode
+{
+	GENERATED_BODY()
+
+public:
+	const static FEditorModeID Id;	
+
+	UMLDeformerPaintMode();
+
+	// UEdMode overrides
+	virtual void Enter() override;
+	virtual void CreateToolkit() override;
+	virtual bool UsesTransformWidget() const override	{ return false; };
+	virtual bool UsesPropertyWidgets() const override	{ return false; };
+	virtual bool UsesToolkits() const override			{ return true; }
+	// ~END UEdMode overrides
+
+	void SetMLDeformerEditor(UE::MLDeformer::FMLDeformerEditorToolkit* Editor);
+	void UpdatePose(USkeletalMeshComponent* SkeletalMeshComponent, bool bFullUpdate=true);
+
+private:
+	TArray<FVector3f> BindPosePositions;
+};
diff --git a/Engine/Plugins/Animation/MLDeformer/MLDeformerFramework/Source/MLDeformerFrameworkEditor/Public/MLDeformerPaintModeToolkit.h b/Engine/Plugins/Animation/MLDeformer/MLDeformerFramework/Source/MLDeformerFrameworkEditor/Public/MLDeformerPaintModeToolkit.h
new file mode 100644
--- /dev/null
+++ b/Engine/Plugins/Animation/MLDeformer/MLDeformerFramework/Source/MLDeformerFrameworkEditor/Public/MLDeformerPaintModeToolkit.h
@@ -0,0 +1,67 @@
+// Copyright Epic Games, Inc. All Rights Reserved.
+
+#pragma once
+
+#include "CoreTypes.h"
+#include "StatusBarSubsystem.h"
+#include "Toolkits/BaseToolkit.h"
+
+class IToolkitHost;
+class STextBlock;
+class SButton;
+
+namespace UE::MLDeformer
+{
+	class FMLDeformerEditorToolkit;
+
+	class MLDEFORMERFRAMEWORKEDITOR_API FMLDeformerPaintModeToolkit
+		: public FModeToolkit
+	{
+	public:
+		virtual ~FMLDeformerPaintModeToolkit() override;
+	
+		// IToolkit overrides
+		virtual void Init(const TSharedPtr<IToolkitHost>& InToolkitHost, TWeakObjectPtr<UEdMode> InOwningMode) override;
+		virtual FName GetToolkitFName() const override;
+		virtual FText GetBaseToolkitName() const override;
+		virtual TSharedPtr<SWidget> GetInlineContent() const override			{ return ToolkitWidget; }
+
+		// FModeToolkit overrides
+		virtual void OnToolStarted(UInteractiveToolManager* Manager, UInteractiveTool* Tool) override;
+		virtual void OnToolEnded(UInteractiveToolManager* Manager, UInteractiveTool* Tool) override;
+		virtual FText GetActiveToolDisplayName() const override;
+		virtual FText GetActiveToolMessage() const override;
+
+		void SetMLDeformerEditor(FMLDeformerEditorToolkit* Editor)				{ MLDeformerEditor = Editor; }
+		UE::MLDeformer::FMLDeformerEditorToolkit* GetMLDeformerEditor() const	{ return MLDeformerEditor; }
+
+	private:
+		void PostNotification(const FText& Message);
+		void ClearNotification();
+
+		void PostWarning(const FText& Message);
+		void ClearWarning();
+
+		void UpdateActiveToolProperties(UInteractiveTool* Tool);
+
+		void RegisterPalettes();
+		FDelegateHandle ActivePaletteChangedHandle;
+
+		void MakeToolAcceptCancelWidget();
+
+		FText ActiveToolName;
+		FText ActiveToolMessage;
+		FStatusBarMessageHandle ActiveToolMessageHandle;
+
+		TSharedPtr<SWidget> ToolkitWidget;
+
+		FMLDeformerEditorToolkit* MLDeformerEditor = nullptr;
+
+		TSharedPtr<SWidget> ViewportOverlayWidget;
+		const FSlateBrush* ActiveToolIcon = nullptr;
+	
+		TSharedPtr<STextBlock> ModeWarningArea;
+		TSharedPtr<STextBlock> ModeHeaderArea;
+		TSharedPtr<STextBlock> ToolWarningArea;
+	};
+} // namespace UE::MLDeformer
diff --git a/Engine/Plugins/Animation/MLDeformer/MLDeformerFramework/Source/MLDeformerFrameworkEditor/Public/SMLDeformerMaskConfigDialog.h b/Engine/Plugins/Animation/MLDeformer/MLDeformerFramework/Source/MLDeformerFrameworkEditor/Public/SMLDeformerMaskConfigDialog.h
new file mode 100644
--- /dev/null
+++ b/Engine/Plugins/Animation/MLDeformer/MLDeformerFramework/Source/MLDeformerFrameworkEditor/Public/SMLDeformerMaskConfigDialog.h
@@ -0,0 +1,56 @@
+// Copyright Epic Games, Inc. All Rights Reserved.
+
+#pragma once
+
+#include "CoreTypes.h"
+#include "Dialog/SCustomDialog.h"
+#include "MLDeformerModel.h"
+#include "MLDeformerMasking.h"
+#include "MeshDescription.h"
+#include "Widgets/Input/STextComboBox.h"
+#include "Widgets/DeclarativeSyntaxSupport.h"
+
+namespace UE::MLDeformer
+{
+	class FMLDeformerEditorModel;
+
+	DECLARE_DELEGATE_OneParam(FOnSetNewVertexAttributeValues, TVertexAttributesRef<float>);
+
+	/**
+	 * The dialog that allows the user to choose a masking mode (Generated or using Skeletal Mesh Vertex Attribute).
+	 * When using a vertex attribute, you can pick which attribute to use, or create a new one.
+	 * Creating a new attribute will pop up the SMLDeformerNewVertexAttributeDialog.
+	 */
+	class MLDEFORMERFRAMEWORKEDITOR_API SMLDeformerMaskConfigDialog
+		: public SCustomDialog
+	{
+	public:
+		SLATE_BEGIN_ARGS(SMLDeformerMaskConfigDialog) {}
+		SLATE_ARGUMENT(FMLDeformerMaskInfo, InitialMaskInfo)
+		SLATE_EVENT(FOnSetNewVertexAttributeValues, OnSetNewVertexAttributeValues)
+		SLATE_END_ARGS()
+
+		void Construct(const FArguments& InArgs, FMLDeformerEditorModel* InEditorModel);
+
+		/**
+		 * This returns the mask info object how it was setup in the dialog.
+		 * You should only use this when the user didn't cancel the dialog.
+		 */
+		const FMLDeformerMaskInfo& GetMaskInfo() const { return MaskInfo; }
+
+	private:
+		void UpdateAttributeNames();
+		TSharedPtr<SWidget> CreateMeshAttributeModeWidget();
+		TSharedPtr<FString> MaskModeEnumToString(EMLDeformerMaskingMode MaskMode) const;
+		FReply OnCreateVertexAttributeClicked();
+
+	private:
+		TArray<FName> AttributeNames;
+		TArray<TSharedPtr<FString>> MaskingModeNames;
+		TSharedPtr<SComboBox<FName>> VertexAttributeComboWidget;
+		FMLDeformerEditorModel* EditorModel = nullptr;
+		FMLDeformerMaskInfo MaskInfo;
+		FOnSetNewVertexAttributeValues OnSetNewVertexAttributeValues;
+	};
+
+}	// namespace UE::MLDeformer
diff --git a/Engine/Plugins/Animation/MLDeformer/MLDeformerFramework/Source/MLDeformerFrameworkEditor/Public/SMLDeformerNewVertexAttributeDialog.h b/Engine/Plugins/Animation/MLDeformer/MLDeformerFramework/Source/MLDeformerFrameworkEditor/Public/SMLDeformerNewVertexAttributeDialog.h
new file mode 100644
--- /dev/null
+++ b/Engine/Plugins/Animation/MLDeformer/MLDeformerFramework/Source/MLDeformerFrameworkEditor/Public/SMLDeformerNewVertexAttributeDialog.h
@@ -0,0 +1,116 @@
+// Copyright Epic Games, Inc. All Rights Reserved.
+
+#pragma once
+
+#include "CoreTypes.h"
+#include "Widgets/SWindow.h"
+#include "Widgets/Input/SButton.h"
+
+class USkeletalMesh;
+
+namespace UE::MLDeformer
+{
+	/**
+	 * The dialog that asks the user to enter a vertex attribute name.
+	 * On default the dialog creates the actual vertex attribute on the mesh as well, when the user pressed Create.
+	 * You can disable this behavior using the bAutoCreateAttribute Slate argument.
+	 * A default value for all vertices can be set as well using the DefaultValue attribute.
+	 * 
+	 * Example usage:
+	 * 
+	 * <code>
+	 * TSharedPtr<SMLDeformerNewVertexAttributeDialog> Dialog = SNew(SMLDeformerNewVertexAttributeDialog, SkeletalMesh)
+	 *		.bAutoCreateAttribute(true)
+	 *		.DefaultValue(1.0f);
+	 * 
+	 * if (Dialog.ShowModal() == SMLDeformerNewVertexAttributeDialog::EReturnCode::CreatePressed)
+	 * {
+	 *		// Get the name of the attribute that was created.
+	 *		const FString& AttributeName = Dialog.GetAttributeName();
+	 * 
+	 *		// ...
+	 * }
+	 * </code>
+	 */
+	class SMLDeformerNewVertexAttributeDialog
+		: public SWindow
+	{
+	public:
+		/** The result of the dialog, which determines whether the user pressed the Create button, or if the user would like to cancel the creation. */
+		enum class EReturnCode : uint8
+		{
+			/** The user pressed the Create button. */
+			CreatePressed,
+
+			/** The user canceled the creation. */
+			Canceled
+		};
+
+		SLATE_BEGIN_ARGS(SMLDeformerNewVertexAttributeDialog)
+			: _bAutoCreateAttribute(true)
+			, _DefaultAttributeValue(0.0f)
+		{}
+
+			/** Should we automatically create the vertex attribute? The default is true. */
+			SLATE_ARGUMENT(bool, bAutoCreateAttribute)
+
+			/** What value should the vertex attributes have for all vertices after creation? This only is used when bAutoCreateAttribute is set to true. The default value is 0. */
+			SLATE_ARGUMENT(float, DefaultAttributeValue)
+
+		SLATE_END_ARGS()
+
+		void Construct(const FArguments& InArgs, USkeletalMesh* InSkeletalMesh);
+
+		/** Shows the dialog. The return code tells you whether the user canceled or pressed create. */
+		EReturnCode ShowModal();
+
+		/** 
+		 * When the user pressed Create, which you can check using the return code of ShowModal, then this is the attribute name they entered.
+		 * It is guaranteed that this attribute is one with a unique name if no other has been created in some other thread or so after the dialog has been created.
+		 */
+		const FString& GetAttributeName() const { return AttributeName; }
+
+
+	private:
+		/** Override the OnKeyDown to detect when the Escape key has been pressed, to cancel and close the dialog. */
+		FReply OnKeyDown(const FGeometry& InGeometry, const FKeyEvent& InKeyEvent) override;
+
+		/** Perform actions when the create button has been clicked. */
+		FReply OnCreateClicked();
+
+		/**
+		 * Check whether a given skeletal mesh has a given vertex attribute already.
+		 * @return Returns True when the attribute already exists, or false when not.
+		 */
+		static bool HasVertexAttribute(const USkeletalMesh* SkeletalMesh, FName InAttributeName);
+
+		/**
+		 * Create a float based vertex attribute with a given name, using a specific default value.
+		 * When there already is an attribute with the specified name, then false is returned and nothing is changed.
+		 * When the attribute name is unique, the new attribute will be created and all vertices will have the default value.
+		 * This will create the attribute on LOD 0.
+		 * @result Returns true on success, or false when failed. A fail can occur when there is already an attribute with the given name or when no Mesh Description can be found.
+		 */
+		static bool CreateVertexAttribute(USkeletalMesh& SkeletalMesh, const FString& AttributeName, float DefaultValue = 0.0f);
+
+
+	private:
+		/** The Skeletal Mesh that we wish to create the attribute on. */
+		TObjectPtr<USkeletalMesh> SkeletalMesh;
+
+		/** The return code, specifying whether the user pressed create or cancel. */
+		EReturnCode ReturnCode = EReturnCode::Canceled;
+
+		/** The name of the attribute that the user entered. */
+		FString AttributeName;
+
+		/** The create button. */
+		TSharedPtr<SButton> CreateButton;
+
+		/** Should we automatically create the new attribute when the Create button has been pressed? */
+		bool bAutoCreate = true;
+
+		/** The default value of all vertices for this vertex attribute. Only used when bAutoCreate we end up creating the attribute. */
+		float DefaultAttributeValue = 0.0f;
+	};
+}	// namespace UE::MLDeformer
diff --git a/Engine/Plugins/Animation/MLDeformer/NeuralMorphModel/Source/NeuralMorphModel/Private/NeuralMorphModel.cpp b/Engine/Plugins/Animation/MLDeformer/NeuralMorphModel/Source/NeuralMorphModel/Private/NeuralMorphModel.cpp
--- a/Engine/Plugins/Animation/MLDeformer/NeuralMorphModel/Source/NeuralMorphModel/Private/NeuralMorphModel.cpp
+++ b/Engine/Plugins/Animation/MLDeformer/NeuralMorphModel/Source/NeuralMorphModel/Private/NeuralMorphModel.cpp
@@ -77,7 +77,7 @@ namespace
 		return BoneIndex;
 	}
 
-	void ModifyMaskInfosForBackwardCompatibility(const FReferenceSkeleton& RefSkel, const TArray<FBoneReference>& BoneIncludeList, TMap<FName, FNeuralMorphMaskInfo>& MaskInfos)
+	void ModifyMaskInfosForBackwardCompatibility(const FReferenceSkeleton& RefSkel, const TArray<FBoneReference>& BoneIncludeList, TMap<FName, FMLDeformerMaskInfo>& MaskInfos)
 	{
 		// If we are not using some masks, just skip it.
 		if (MaskInfos.IsEmpty())
@@ -107,7 +107,7 @@ namespace
 		TArray<FName> RequiredBones;
 		for (auto& MaskInfoEntry : MaskInfos)
 		{
-			FNeuralMorphMaskInfo& MaskInfo = MaskInfoEntry.Value;
+			FMLDeformerMaskInfo& MaskInfo = MaskInfoEntry.Value;
 			RequiredBones.Reset();
 			for (const FName MaskBoneName : MaskInfo.BoneNames)
 			{
@@ -141,6 +141,17 @@ namespace
 			}
 		}
 	}
+
+	PRAGMA_DISABLE_DEPRECATION_WARNINGS;
+	FMLDeformerMaskInfo ConvertMaskInfo(const FNeuralMorphMaskInfo& NeuralMaskInfo)
+	{
+		FMLDeformerMaskInfo Result;
+		Result.BoneNames = NeuralMaskInfo.BoneNames;
+		Result.MaskMode = NeuralMaskInfo.MaskMode;
+		Result.VertexAttributeName = NeuralMaskInfo.VertexAttributeName;
+		return Result;
+	}
+	PRAGMA_ENABLE_DEPRECATION_WARNINGS;
 } // Anonymous namespace.
 
 void UNeuralMorphModel::Serialize(FArchive& Archive)
@@ -162,8 +173,8 @@ void UNeuralMorphModel::Serialize(FArchive& Archive)
 		}
 
 		// Strip the mask data in the cooked asset.
-		BoneMaskInfos.Empty();
-		BoneGroupMaskInfos.Empty();
+		BoneMaskInfoMap.Empty();
+		BoneGroupMaskInfoMap.Empty();
 	}
 
 	// Convert the UMLDeformerInputInfo object into a UNeuralMorphInputInfo object for backward compatiblity.
@@ -178,10 +189,31 @@ void UNeuralMorphModel::Serialize(FArchive& Archive)
 		}		
  	}
 
-	// Since 5.4 (when LOD was added), we modified the way the bone masks are generated.
-	// In UE 5.3 we added some extra bones to the mask. To keep the same behavior, we will upgrade the mask to include those bones again to ensure we get the same results when training.
-	// Newer assets will not do this.
 	#if WITH_EDITORONLY_DATA
+		// Since 5.5 we removed the FNeuralMorphMaskInfo and use the FMLDeformerMaskInfo type instead.
+		// We need to convert the masks.
+		if (Archive.IsLoading() && 
+			Archive.CustomVer(UE::MLDeformer::FMLDeformerObjectVersion::GUID) < UE::MLDeformer::FMLDeformerObjectVersion::MaskInfoMovedToFramework)
+		{
+			PRAGMA_DISABLE_DEPRECATION_WARNINGS;
+				for (const auto& MaskItemPair : BoneMaskInfos_DEPRECATED)
+				{
+					BoneMaskInfoMap.Add(MaskItemPair.Key, ConvertMaskInfo(MaskItemPair.Value));
+				}
+
+				for (const auto& MaskItemPair : BoneGroupMaskInfos_DEPRECATED)
+				{
+					BoneGroupMaskInfoMap.Add(MaskItemPair.Key, ConvertMaskInfo(MaskItemPair.Value));
+				}
+
+				BoneMaskInfos_DEPRECATED.Empty();
+				BoneGroupMaskInfos_DEPRECATED.Empty();
+			PRAGMA_ENABLE_DEPRECATION_WARNINGS;
+		}
+
+		// Since 5.4 (when LOD was added), we modified the way the bone masks are generated.
+		// In UE 5.3 we added some extra bones to the mask. To keep the same behavior, we will upgrade the mask to include those bones again to ensure we get the same results when training.
+		// Newer assets will not do this.
 		if (Archive.IsLoading() && 
 			Archive.CustomVer(UE::MLDeformer::FMLDeformerObjectVersion::GUID) < UE::MLDeformer::FMLDeformerObjectVersion::LODSupportAdded)
 		{
@@ -189,8 +221,8 @@ void UNeuralMorphModel::Serialize(FArchive& Archive)
 			if (SkelMesh)
 			{
 				const FReferenceSkeleton& RefSkeleton = SkelMesh->GetRefSkeleton();
-				ModifyMaskInfosForBackwardCompatibility(RefSkeleton, GetBoneIncludeList(), BoneMaskInfos);
-				ModifyMaskInfosForBackwardCompatibility(RefSkeleton, GetBoneIncludeList(), BoneGroupMaskInfos);
+				ModifyMaskInfosForBackwardCompatibility(RefSkeleton, GetBoneIncludeList(), BoneMaskInfoMap);
+				ModifyMaskInfosForBackwardCompatibility(RefSkeleton, GetBoneIncludeList(), BoneGroupMaskInfoMap);
 			}
 		}
 	#endif
diff --git a/Engine/Plugins/Animation/MLDeformer/NeuralMorphModel/Source/NeuralMorphModel/Public/NeuralMorphModel.h b/Engine/Plugins/Animation/MLDeformer/NeuralMorphModel/Source/NeuralMorphModel/Public/NeuralMorphModel.h
--- a/Engine/Plugins/Animation/MLDeformer/NeuralMorphModel/Source/NeuralMorphModel/Public/NeuralMorphModel.h
+++ b/Engine/Plugins/Animation/MLDeformer/NeuralMorphModel/Source/NeuralMorphModel/Public/NeuralMorphModel.h
@@ -94,6 +94,16 @@ public:
 	UPROPERTY()
 	TArray<FNeuralMorphCurveGroup> CurveGroups;
 
+	PRAGMA_DISABLE_DEPRECATION_WARNINGS
+	UE_DEPRECATED(5.5, "This property has been deprecated, please use BoneMaskInfoMap instead.")
+	UPROPERTY(meta = (DeprecatedProperty, DeprecationMessage = "This property has been deprecated, please use BoneMaskInfoMap instead."))
+	TMap<FName, FNeuralMorphMaskInfo> BoneMaskInfos_DEPRECATED;
+
+	UE_DEPRECATED(5.5, "This property has been deprecated, please use BoneGroupMaskInfoMap instead.")
+	UPROPERTY(meta = (DeprecatedProperty, DeprecationMessage = "This property has been deprecated, please use BoneGroupMaskInfoMap instead."))
+	TMap<FName, FNeuralMorphMaskInfo> BoneGroupMaskInfos_DEPRECATED;
+	PRAGMA_ENABLE_DEPRECATION_WARNINGS
+
 	/**
 	 * Information needed to generate a mask for each bone.
 	 * Each mask info object contains a list of bones who's skinning influence regions should be included in the final mask for the specific bone.
@@ -101,7 +111,7 @@ public:
 	 * The FName map key represents the bone name.
 	 */
 	UPROPERTY()
-	TMap<FName, FNeuralMorphMaskInfo> BoneMaskInfos;
+	TMap<FName, FMLDeformerMaskInfo> BoneMaskInfoMap;
 
 	/**
 	 * Information needed to generate a mask for each bone group.
@@ -110,7 +120,7 @@ public:
 	 * This information (and bone masking in general) is used inside the ENeuralMorphMode::Local mode.
 	 */
 	UPROPERTY()
-	TMap<FName, FNeuralMorphMaskInfo> BoneGroupMaskInfos;
+	TMap<FName, FMLDeformerMaskInfo> BoneGroupMaskInfoMap;
 
 	/**
 	 * The mode that the neural network will operate in. 
diff --git a/Engine/Plugins/Animation/MLDeformer/NeuralMorphModel/Source/NeuralMorphModel/Public/NeuralMorphTypes.h b/Engine/Plugins/Animation/MLDeformer/NeuralMorphModel/Source/NeuralMorphModel/Public/NeuralMorphTypes.h
--- a/Engine/Plugins/Animation/MLDeformer/NeuralMorphModel/Source/NeuralMorphModel/Public/NeuralMorphTypes.h
+++ b/Engine/Plugins/Animation/MLDeformer/NeuralMorphModel/Source/NeuralMorphModel/Public/NeuralMorphTypes.h
@@ -4,6 +4,7 @@
 #include "CoreTypes.h"
 #include "Containers/Array.h"
 #include "BoneContainer.h"
+#include "MLDeformerMasking.h"
 #include "MLDeformerCurveReference.h"
 #include "NeuralMorphTypes.generated.h"
 
@@ -99,7 +100,7 @@ public:
  * There will be an array of these structs, one for each bone.
  */
 USTRUCT()
-struct NEURALMORPHMODEL_API FNeuralMorphMaskInfo
+struct UE_DEPRECATED(5.5, "This class has been deprecated, please use FMLDeformerMaskInfo instead.") NEURALMORPHMODEL_API FNeuralMorphMaskInfo
 {
 	GENERATED_BODY()
 
@@ -107,4 +108,12 @@ public:
 	/** The list of bone names that should be included in the mask generation. */
 	UPROPERTY()
 	TArray<FName> BoneNames;
+
+	/** The masking mode. */
+	UPROPERTY()
+	EMLDeformerMaskingMode MaskMode = EMLDeformerMaskingMode::Generated;
+
+	/** If the masking mode is set to VertexAttribute then we can check which attribute to use using this member. */
+	UPROPERTY()
+	FName VertexAttributeName;
 };
diff --git a/Engine/Plugins/Animation/MLDeformer/NeuralMorphModel/Source/NeuralMorphModelEditor/Private/NeuralMorphEditorModel.cpp b/Engine/Plugins/Animation/MLDeformer/NeuralMorphModel/Source/NeuralMorphModelEditor/Private/NeuralMorphEditorModel.cpp
--- a/Engine/Plugins/Animation/MLDeformer/NeuralMorphModel/Source/NeuralMorphModelEditor/Private/NeuralMorphEditorModel.cpp
+++ b/Engine/Plugins/Animation/MLDeformer/NeuralMorphModel/Source/NeuralMorphModelEditor/Private/NeuralMorphEditorModel.cpp
@@ -75,6 +75,7 @@ namespace UE::NeuralMorphModel
 		{
 			if (PropertyChangedEvent.ChangeType == EPropertyChangeType::ValueSet)
 			{
+				SetResamplingInputOutputsNeeded(true);
 				UpdateIsReadyForTrainingState();
 				GetEditor()->GetModelDetailsView()->ForceRefresh();
 				if (GetNeuralMorphModel()->GetModelMode() == ENeuralMorphMode::Local)
@@ -283,27 +284,6 @@ namespace UE::NeuralMorphModel
 		return MorphInputInfo->GetMaskForItem(ItemIndex);
 	}
 
-	void FNeuralMorphEditorModel::ApplyMaskInfoToMaskBuffer(const USkeletalMesh* SkeletalMesh, const FNeuralMorphMaskInfo& MaskInfo, TArrayView<float> ItemMaskBuffer)	
-	{
-		const FReferenceSkeleton& RefSkel = SkeletalMesh->GetRefSkeleton();
-
-		// Apply the bones to the mask buffer.
-		for (const FName MaskBoneName : MaskInfo.BoneNames)
-		{
-			const int32 MaskBoneIndex = RefSkel.FindBoneIndex(MaskBoneName);
-			if (MaskBoneIndex != INDEX_NONE)
-			{
-				ApplyBoneToMask(MaskBoneIndex, ItemMaskBuffer);
-			}
-			else
-			{
-				UE_LOG(LogNeuralMorphModel, Warning, TEXT("Mask contains a bone named '%s', which cannot be found in the ref skeleton of skeletal mesh '%s'."),
-					*MaskBoneName.ToString(),
-					*SkeletalMesh->GetName());
-			}
-		}
-	}
-
 	void FNeuralMorphEditorModel::BuildMaskBuffer(TArray<float>& OutMaskBuffer)
 	{
 		if (Model->GetSkeletalMesh() == nullptr)
@@ -345,14 +325,14 @@ namespace UE::NeuralMorphModel
 		{
 			TArrayView<float> ItemMaskBuffer(&OutMaskBuffer[MaskOffset], NumBaseMeshVerts);
 			const FName BoneName = BoneNames[Index];
-			FNeuralMorphMaskInfo* MaskInfo = NeuralMorphModel->BoneMaskInfos.Find(BoneName);
+			FMLDeformerMaskInfo* MaskInfo = NeuralMorphModel->BoneMaskInfoMap.Find(BoneName);
 			if (MaskInfo == nullptr)
 			{
-				MaskInfo = &NeuralMorphModel->BoneMaskInfos.Add(BoneName, FNeuralMorphMaskInfo());
+				MaskInfo = &NeuralMorphModel->BoneMaskInfoMap.Add(BoneName, FMLDeformerMaskInfo());
 				GenerateBoneMaskInfo(Index, HierarchyDepth);
 			}
 
-			ApplyMaskInfoToMaskBuffer(SkeletalMesh, *MaskInfo, ItemMaskBuffer);
+			ApplyMaskInfoToBuffer(SkeletalMesh, *MaskInfo, ItemMaskBuffer);
 			MaskOffset += NumBaseMeshVerts;
 		}
 
@@ -370,14 +350,14 @@ namespace UE::NeuralMorphModel
 			TArrayView<float> ItemMaskBuffer(&OutMaskBuffer[MaskOffset], NumBaseMeshVerts);
 			
 			const FName GroupName = NeuralMorphInputInfo->GetBoneGroups()[Index].GroupName;
-			FNeuralMorphMaskInfo* MaskInfo = NeuralMorphModel->BoneGroupMaskInfos.Find(GroupName);
+			FMLDeformerMaskInfo* MaskInfo = NeuralMorphModel->BoneGroupMaskInfoMap.Find(GroupName);
 			if (MaskInfo == nullptr)
 			{
-				MaskInfo = &NeuralMorphModel->BoneGroupMaskInfos.Add(GroupName, FNeuralMorphMaskInfo());
+				MaskInfo = &NeuralMorphModel->BoneGroupMaskInfoMap.Add(GroupName, FMLDeformerMaskInfo());
 				GenerateBoneGroupMaskInfo(Index, HierarchyDepth);
 			}
 
-			ApplyMaskInfoToMaskBuffer(SkeletalMesh, *MaskInfo, ItemMaskBuffer);
+			ApplyMaskInfoToBuffer(SkeletalMesh, *MaskInfo, ItemMaskBuffer);
 			MaskOffset += NumBaseMeshVerts;
 		}
 
@@ -450,12 +430,12 @@ namespace UE::NeuralMorphModel
 
 	void FNeuralMorphEditorModel::ResetBoneMaskInfos()
 	{
-		GetNeuralMorphModel()->BoneMaskInfos.Empty();
+		GetNeuralMorphModel()->BoneMaskInfoMap.Empty();
 	}
 
 	void FNeuralMorphEditorModel::ResetBoneGroupMaskInfos()
 	{
-		GetNeuralMorphModel()->BoneGroupMaskInfos.Empty();
+		GetNeuralMorphModel()->BoneGroupMaskInfoMap.Empty();
 	}
 
 	void FNeuralMorphEditorModel::AddTwistBones(const FReferenceSkeleton& RefSkel, TArray<int32>& SkelBoneIndices)
@@ -520,10 +500,10 @@ namespace UE::NeuralMorphModel
 		AddTwistBones(RefSkel, BonesAdded);
 
 		// Now that we know which bones we need, add them to the mask.
-		FNeuralMorphMaskInfo* MaskInfo = NeuralMorphModel->BoneMaskInfos.Find(BoneName);
+		FMLDeformerMaskInfo* MaskInfo = NeuralMorphModel->BoneMaskInfoMap.Find(BoneName);
 		if (MaskInfo == nullptr)
 		{
-			MaskInfo = &NeuralMorphModel->BoneMaskInfos.Add(BoneName, FNeuralMorphMaskInfo());
+			MaskInfo = &NeuralMorphModel->BoneMaskInfoMap.Add(BoneName, FMLDeformerMaskInfo());
 		}
 		MaskInfo->BoneNames.Reset();
 		for (int32 BoneIndex : BonesAdded)
@@ -569,10 +549,10 @@ namespace UE::NeuralMorphModel
 		}
 
 		// Now that we know which bones we need, add them to the mask.
-		FNeuralMorphMaskInfo* MaskInfo = NeuralMorphModel->BoneGroupMaskInfos.Find(BoneGroup.GroupName);
+		FMLDeformerMaskInfo* MaskInfo = NeuralMorphModel->BoneGroupMaskInfoMap.Find(BoneGroup.GroupName);
 		if (MaskInfo == nullptr)
 		{
-			MaskInfo = &NeuralMorphModel->BoneGroupMaskInfos.Add(BoneGroup.GroupName, FNeuralMorphMaskInfo());
+			MaskInfo = &NeuralMorphModel->BoneGroupMaskInfoMap.Add(BoneGroup.GroupName, FMLDeformerMaskInfo());
 		}
 		MaskInfo->BoneNames.Reset();
 		for (int32 BoneIndex : BonesAdded)
@@ -629,7 +609,7 @@ namespace UE::NeuralMorphModel
 
 		// Find all mask infos for bones that don't exist anymore.
 		TArray<FName> BoneMaskInfosToRemove;
-		for (const auto& MaskInfo : NeuralMorphModel->BoneMaskInfos)
+		for (const auto& MaskInfo : NeuralMorphModel->BoneMaskInfoMap)
 		{
 			const FName BoneName = MaskInfo.Key;
 			if (!NeuralMorphModel->GetBoneIncludeList().Contains(BoneName))
@@ -641,12 +621,12 @@ namespace UE::NeuralMorphModel
 		// Remove those mask infos.
 		for (const FName BoneName : BoneMaskInfosToRemove)
 		{
-			NeuralMorphModel->BoneMaskInfos.Remove(BoneName);
+			NeuralMorphModel->BoneMaskInfoMap.Remove(BoneName);
 		}
 
 		// Find all bone groups that don't exist anymore.
 		BoneMaskInfosToRemove.Reset();
-		for (const auto& MaskInfo : NeuralMorphModel->BoneGroupMaskInfos)
+		for (const auto& MaskInfo : NeuralMorphModel->BoneGroupMaskInfoMap)
 		{
 			const FName GroupName = MaskInfo.Key;
 
@@ -669,7 +649,7 @@ namespace UE::NeuralMorphModel
 		// Remove those group mask infos.
 		for (const FName GroupName : BoneMaskInfosToRemove)
 		{
-			NeuralMorphModel->BoneGroupMaskInfos.Remove(GroupName);
+			NeuralMorphModel->BoneGroupMaskInfoMap.Remove(GroupName);
 		}
 	}
 
@@ -692,12 +672,26 @@ namespace UE::NeuralMorphModel
 		const int32 NumMaskItems = InputInfo->GetBoneNames().Num() + InputInfo->GetCurveNames().Num() + InputInfo->GetBoneGroups().Num() + InputInfo->GetCurveGroups().Num();
 		const int32 FinalMaskItemIndex = FMath::Clamp<int32>(MaskItemIndex, 0, NumMaskItems - 1);
 		const TArrayView<const float> MaskBuffer = InputInfo->GetMaskForItem(FinalMaskItemIndex);
+
+		// Find the mask info.
+		FMLDeformerMaskInfo* MaskInfo = nullptr;
+		if (FinalMaskItemIndex < InputInfo->GetBoneNames().Num())	// If we have a bone.
+		{
+			MaskInfo = GetNeuralMorphModel()->BoneMaskInfoMap.Find(InputInfo->GetBoneNames()[FinalMaskItemIndex]);
+		}
+		else // It's a bone group.
+		{
+			const int32 GroupIndex = FinalMaskItemIndex - (InputInfo->GetBoneNames().Num() + InputInfo->GetCurveNames().Num());
+			const FName GroupName = InputInfo->GetBoneGroups()[GroupIndex].GroupName;
+			MaskInfo = GetNeuralMorphModel()->BoneGroupMaskInfoMap.Find(GroupName);
+		}
+
 		if (!MaskBuffer.IsEmpty())
 		{
 			check(MaskBuffer.Num() == NumVerts);
 			check(NumVerts == UnskinnedPositions.Num());
 
-			FLinearColor OrgIncludedColor = FMLDeformerEditorStyle::Get().GetColor("MLDeformer.Morphs.MaskIncludedVertexColor");
+			const FLinearColor OrgIncludedColor = (MaskInfo && MaskInfo->MaskMode == EMLDeformerMaskingMode::VertexAttribute) ? FMLDeformerEditorStyle::Get().GetColor("MLDeformer.Morphs.MaskIncludedVertexColorPainted") : FMLDeformerEditorStyle::Get().GetColor("MLDeformer.Morphs.MaskIncludedVertexColor");
 			const FLinearColor ExcludedColor = FMLDeformerEditorStyle::Get().GetColor("MLDeformer.Morphs.MaskExcludedVertexColor");
 			for (int32 VertexIndex = 0; VertexIndex < NumVerts; ++VertexIndex)
 			{
@@ -744,7 +738,6 @@ namespace UE::NeuralMorphModel
 			}
 		}
 	}
-
 }	// namespace UE::NeuralMorphModel
 
 #undef LOCTEXT_NAMESPACE
diff --git a/Engine/Plugins/Animation/MLDeformer/NeuralMorphModel/Source/NeuralMorphModelEditor/Private/SNeuralMorphBoneGroupsWidget.cpp b/Engine/Plugins/Animation/MLDeformer/NeuralMorphModel/Source/NeuralMorphModelEditor/Private/SNeuralMorphBoneGroupsWidget.cpp
--- a/Engine/Plugins/Animation/MLDeformer/NeuralMorphModel/Source/NeuralMorphModelEditor/Private/SNeuralMorphBoneGroupsWidget.cpp
+++ b/Engine/Plugins/Animation/MLDeformer/NeuralMorphModel/Source/NeuralMorphModelEditor/Private/SNeuralMorphBoneGroupsWidget.cpp
@@ -190,7 +190,7 @@ namespace UE::NeuralMorphModel
 					else
 					{
 						NeuralMorphModel->BoneGroups[ParentGroup->GroupIndex].BoneNames.Remove(BoneItem->Name);
-						FNeuralMorphMaskInfo* MaskInfo = NeuralMorphModel->BoneGroupMaskInfos.Find(ParentGroup->Name);
+						FMLDeformerMaskInfo* MaskInfo = NeuralMorphModel->BoneGroupMaskInfoMap.Find(ParentGroup->Name);
 						if (MaskInfo)
 						{
 							MaskInfo->BoneNames.Remove(BoneItem->Name);
@@ -208,7 +208,7 @@ namespace UE::NeuralMorphModel
 			// Remove the items, back to front.
 			for (const int32 Index : GroupsToRemove)
 			{
-				NeuralMorphModel->BoneGroupMaskInfos.Remove(NeuralMorphModel->BoneGroups[Index].GroupName);
+				NeuralMorphModel->BoneGroupMaskInfoMap.Remove(NeuralMorphModel->BoneGroups[Index].GroupName);
 				NeuralMorphModel->BoneGroups.RemoveAt(Index);
 			}
 		}
@@ -221,7 +221,7 @@ namespace UE::NeuralMorphModel
 	{
 		UNeuralMorphModel* NeuralMorphModel = EditorModel->GetNeuralMorphModel();
 		NeuralMorphModel->BoneGroups.Empty();
-		NeuralMorphModel->BoneGroupMaskInfos.Empty();
+		NeuralMorphModel->BoneGroupMaskInfoMap.Empty();
 		RefreshTree(true);
 		EditorModel->RebuildEditorMaskInfo();
 	}
diff --git a/Engine/Plugins/Animation/MLDeformer/NeuralMorphModel/Source/NeuralMorphModelEditor/Private/SNeuralMorphInputWidget.cpp b/Engine/Plugins/Animation/MLDeformer/NeuralMorphModel/Source/NeuralMorphModelEditor/Private/SNeuralMorphInputWidget.cpp
--- a/Engine/Plugins/Animation/MLDeformer/NeuralMorphModel/Source/NeuralMorphModelEditor/Private/SNeuralMorphInputWidget.cpp
+++ b/Engine/Plugins/Animation/MLDeformer/NeuralMorphModel/Source/NeuralMorphModelEditor/Private/SNeuralMorphInputWidget.cpp
@@ -23,7 +23,9 @@
 #include "SMLDeformerInputBonesWidget.h"
 #include "SMLDeformerBonePickerDialog.h"
 #include "SMLDeformerCurvePickerDialog.h"
+#include "SMLDeformerMaskConfigDialog.h"
 #include "Engine/SkeletalMesh.h"
+#include "MeshDescription.h"
 #include "IDetailsView.h"
 
 #define LOCTEXT_NAMESPACE "NeuralMorphInputWidget"
@@ -52,6 +54,9 @@ namespace UE::NeuralMorphModel
 		UI_COMMAND(ResetAllBoneGroupMasks, "Reset All Bone Group Masks", "Reset all masks for every bone group.", EUserInterfaceActionType::Button, FInputChord());
 		UI_COMMAND(ResetSelectedBoneGroupMasks, "Reset Selected Bone Group Masks", "Reset the masks for all selected bone groups.", EUserInterfaceActionType::Button, FInputChord());
 		UI_COMMAND(ExpandSelectedBoneGroupMasks, "Edit Group Mask", "Specify which bones to include inside the mask for this group.", EUserInterfaceActionType::Button, FInputChord());
+
+		UI_COMMAND(ConfigureBoneMask, "Configure Mask", "Configure settings specific to the bone masking.", EUserInterfaceActionType::Button, FInputChord());
+		UI_COMMAND(ConfigureBoneGroupMask, "Configure Mask", "Configure settings specific to the bone group masking.", EUserInterfaceActionType::Button, FInputChord());
 	}
 
 	void SNeuralMorphInputWidget::BindCommands()
@@ -63,6 +68,7 @@ namespace UE::NeuralMorphModel
 			BonesCommandList->MapAction(Commands.ResetAllBoneMasks, FExecuteAction::CreateSP(this, &SNeuralMorphInputWidget::ResetAllBoneMasks));
 			BonesCommandList->MapAction(Commands.ResetSelectedBoneMasks, FExecuteAction::CreateSP(this, &SNeuralMorphInputWidget::ResetSelectedBoneMasks));
 			BonesCommandList->MapAction(Commands.ExpandSelectedBoneMasks, FExecuteAction::CreateSP(this, &SNeuralMorphInputWidget::ExpandBoneMasks));
+			BonesCommandList->MapAction(Commands.ConfigureBoneMask, FExecuteAction::CreateSP(this, &SNeuralMorphInputWidget::ConfigureBoneMask));
 		}
 
 		if (BoneGroupsWidget.IsValid())
@@ -70,6 +76,7 @@ namespace UE::NeuralMorphModel
 			BoneGroupsCommandList->MapAction(Commands.ResetAllBoneGroupMasks, FExecuteAction::CreateSP(this, &SNeuralMorphInputWidget::ResetAllBoneGroupMasks));
 			BoneGroupsCommandList->MapAction(Commands.ResetSelectedBoneGroupMasks, FExecuteAction::CreateSP(this, &SNeuralMorphInputWidget::ResetSelectedBoneGroupMasks));
 			BoneGroupsCommandList->MapAction(Commands.ExpandSelectedBoneGroupMasks, FExecuteAction::CreateSP(this, &SNeuralMorphInputWidget::ExpandBoneGroupMasks));
+			BoneGroupsCommandList->MapAction(Commands.ConfigureBoneGroupMask, FExecuteAction::CreateSP(this, &SNeuralMorphInputWidget::ConfigureBoneGroupMask));
 		}
 	}
 
@@ -184,7 +191,7 @@ namespace UE::NeuralMorphModel
 		{
 			check(SelectedItems[0].IsValid());
 			const FName BoneName = SelectedItems[0]->Name;
-			const FNeuralMorphMaskInfo* MaskInfo = NeuralEditorModel->GetNeuralMorphModel()->BoneMaskInfos.Find(BoneName);
+			const FMLDeformerMaskInfo* MaskInfo = NeuralEditorModel->GetNeuralMorphModel()->BoneMaskInfoMap.Find(BoneName);
 			if (MaskInfo)
 			{
 				HighlightedBones = MaskInfo->BoneNames;
@@ -215,10 +222,10 @@ namespace UE::NeuralMorphModel
 				// Add the picked bone name to the mask info of the selected bone.
 				const int32 BoneIndex = EditorInputInfo->GetBoneNames().Find(SelectedInputBoneName);
 				check(BoneIndex != INDEX_NONE);
-				FNeuralMorphMaskInfo* MaskInfo = NeuralEditorModel->GetNeuralMorphModel()->BoneMaskInfos.Find(SelectedInputBoneName);
+				FMLDeformerMaskInfo* MaskInfo = NeuralEditorModel->GetNeuralMorphModel()->BoneMaskInfoMap.Find(SelectedInputBoneName);
 				if (MaskInfo == nullptr)
 				{
-					MaskInfo = &NeuralEditorModel->GetNeuralMorphModel()->BoneMaskInfos.Add(SelectedInputBoneName, FNeuralMorphMaskInfo());
+					MaskInfo = &NeuralEditorModel->GetNeuralMorphModel()->BoneMaskInfoMap.Add(SelectedInputBoneName, FMLDeformerMaskInfo());
 				}
 
 				check(MaskInfo);
@@ -323,7 +330,7 @@ namespace UE::NeuralMorphModel
 			check(LastSelectedGroup->IsGroup());
 			const int32 GroupIndex = LastSelectedGroup->GroupIndex;
 			check(GroupIndex!= INDEX_NONE);
-			const FNeuralMorphMaskInfo* MaskInfo = NeuralEditorModel->GetNeuralMorphModel()->BoneGroupMaskInfos.Find(LastSelectedGroup->Name);
+			const FMLDeformerMaskInfo* MaskInfo = NeuralEditorModel->GetNeuralMorphModel()->BoneGroupMaskInfoMap.Find(LastSelectedGroup->Name);
 			if (MaskInfo)
 			{
 				HighlightedBones = MaskInfo->BoneNames;
@@ -356,10 +363,10 @@ namespace UE::NeuralMorphModel
 				}
 
 				// Add the picked bone name to the mask info of the selected group.
-				FNeuralMorphMaskInfo* MaskInfo = NeuralEditorModel->GetNeuralMorphModel()->BoneGroupMaskInfos.Find(Item->Name);
+				FMLDeformerMaskInfo* MaskInfo = NeuralEditorModel->GetNeuralMorphModel()->BoneGroupMaskInfoMap.Find(Item->Name);
 				if (MaskInfo == nullptr)
 				{
-					MaskInfo = &NeuralEditorModel->GetNeuralMorphModel()->BoneGroupMaskInfos.Add(Item->Name, FNeuralMorphMaskInfo());
+					MaskInfo = &NeuralEditorModel->GetNeuralMorphModel()->BoneGroupMaskInfoMap.Add(Item->Name, FMLDeformerMaskInfo());
 				}
 				check(MaskInfo);
 				MaskInfo->BoneNames.Reset();
@@ -387,17 +394,45 @@ namespace UE::NeuralMorphModel
 		MenuBuilder.BeginSection("BoneMaskActions", LOCTEXT("BoneMaskActionsHeading", "Bone Masks"));
 		{
 			const FNeuralMorphInputWidgetCommands& Commands = FNeuralMorphInputWidgetCommands::Get();
+
+			// Check if we have only generated masks or not.
+			bool bHasNonGeneratedMask = false;
 			const int32 NumSelectedItems = InputBonesWidget->GetTreeWidget()->GetNumItemsSelected();
+			const TArray<TSharedPtr<FMLDeformerInputBoneTreeElement>> SelectedItems = InputBonesWidget->GetTreeWidget()->GetSelectedItems();
+			for (const TSharedPtr<FMLDeformerInputBoneTreeElement>& Item : SelectedItems)
+			{
+				if (!Item.IsValid())
+				{
+					continue;
+				}
+
+				FMLDeformerMaskInfo* MaskInfo = NeuralMorphModel->BoneMaskInfoMap.Find(Item->Name);
+				if (MaskInfo && MaskInfo->MaskMode != EMLDeformerMaskingMode::Generated)
+				{
+					bHasNonGeneratedMask = true;
+				}
+			}
+
 			if (NumSelectedItems > 0)
 			{
-				if (NumSelectedItems == 1)
+				if (!bHasNonGeneratedMask)
 				{
 					MenuBuilder.AddMenuEntry(Commands.ExpandSelectedBoneMasks);
+					MenuBuilder.AddMenuEntry(Commands.ResetSelectedBoneMasks);
+				}
+
+				if (NumSelectedItems == 1)
+				{
+					if (!bHasNonGeneratedMask)
+					{
+						MenuBuilder.AddMenuSeparator();
+					}
+					MenuBuilder.AddMenuEntry(Commands.ConfigureBoneMask);
 				}
-				MenuBuilder.AddMenuEntry(Commands.ResetSelectedBoneMasks);
 			}
+			MenuBuilder.EndSection();
 		}
-		MenuBuilder.EndSection();
+
 	}
 
 	void SNeuralMorphInputWidget::AddInputBonesPlusIconMenuItems(FMenuBuilder& MenuBuilder)
@@ -424,8 +459,8 @@ namespace UE::NeuralMorphModel
 	{
 		FNeuralMorphEditorModel* NeuralEditorModel = static_cast<FNeuralMorphEditorModel*>(EditorModel);
 		UNeuralMorphModel* NeuralMorphModel = NeuralEditorModel->GetNeuralMorphModel();
-		NeuralMorphModel->BoneMaskInfos.Empty();
-		NeuralMorphModel->BoneGroupMaskInfos.Empty();
+		NeuralMorphModel->BoneMaskInfoMap.Empty();
+		NeuralMorphModel->BoneGroupMaskInfoMap.Empty();
 		NeuralEditorModel->RebuildEditorMaskInfo();
 	}
 
@@ -436,10 +471,10 @@ namespace UE::NeuralMorphModel
 
 		for (const FName Name : Names)
 		{
-			NeuralMorphModel->BoneMaskInfos.Remove(Name);
+			NeuralMorphModel->BoneMaskInfoMap.Remove(Name);
 
 			// Remove the bone from any bone group masks.
-			for (auto& MaskInfo : NeuralMorphModel->BoneGroupMaskInfos)
+			for (auto& MaskInfo : NeuralMorphModel->BoneGroupMaskInfoMap)
 			{
 				MaskInfo.Value.BoneNames.Remove(Name);
 			}
@@ -468,11 +503,34 @@ namespace UE::NeuralMorphModel
 			const FNeuralMorphInputWidgetCommands& Commands = FNeuralMorphInputWidgetCommands::Get();
 			if (BoneGroupsWidget->GetNumSelectedGroups() > 0)
 			{
+				TSharedPtr<FNeuralMorphBoneGroupsTreeElement> LastSelectedGroup;
+				for (const TSharedPtr<FNeuralMorphBoneGroupsTreeElement>& Item : BoneGroupsWidget->GetSelectedItems())
+				{
+					if (Item->IsGroup())
+					{
+						LastSelectedGroup = Item;
+					}
+				}
+				check(LastSelectedGroup.IsValid());
+
+				bool bIsNonGeneratedMask = false;
+				FMLDeformerMaskInfo* MaskInfo = NeuralMorphModel->BoneGroupMaskInfoMap.Find(LastSelectedGroup->Name);
+				if (MaskInfo && MaskInfo->MaskMode != EMLDeformerMaskingMode::Generated)
+				{
+					bIsNonGeneratedMask = true;
+				}
+
 				if (BoneGroupsWidget->GetNumSelectedGroups() == 1)
 				{
-					MenuBuilder.AddMenuEntry(Commands.ExpandSelectedBoneGroupMasks);
+					if (!bIsNonGeneratedMask)
+					{
+						MenuBuilder.AddMenuEntry(Commands.ExpandSelectedBoneGroupMasks);
+						MenuBuilder.AddMenuEntry(Commands.ResetSelectedBoneGroupMasks);
+						MenuBuilder.AddMenuSeparator();
+					}
+
+					MenuBuilder.AddMenuEntry(Commands.ConfigureBoneGroupMask);
 				}
-				MenuBuilder.AddMenuEntry(Commands.ResetSelectedBoneGroupMasks);
 			}
 		}
 		MenuBuilder.EndSection();
@@ -774,6 +832,144 @@ namespace UE::NeuralMorphModel
 		FNeuralMorphEditorModel* NeuralEditorModel = static_cast<FNeuralMorphEditorModel*>(EditorModel);
 		NeuralEditorModel->RebuildEditorMaskInfo();
 	}
+
+	void SNeuralMorphInputWidget::ConfigureBoneMask()
+	{
+		using namespace UE::MLDeformer;
+
+		// Get the selected input bones.
+		check(InputBonesWidget.IsValid());
+		check(InputBonesWidget->GetTreeWidget().IsValid());
+		TArray<TSharedPtr<FMLDeformerInputBoneTreeElement>> SelectedItems = InputBonesWidget->GetTreeWidget()->GetSelectedItems();
+		check(SelectedItems.Num() == 1);	// This only works for single selection.
+
+		const FName BoneName = SelectedItems[0]->Name;
+		UNeuralMorphModel* NeuralModel = Cast<UNeuralMorphModel>(EditorModel->GetModel());
+		FMLDeformerMaskInfo NeuralMaskInfo;
+		if (NeuralModel->BoneMaskInfoMap.Contains(BoneName))
+		{
+			NeuralMaskInfo = *NeuralModel->BoneMaskInfoMap.Find(BoneName);
+		}
+
+		// Show the mask configuration dialog.
+		TSharedPtr<SMLDeformerMaskConfigDialog> Dialog = SNew(SMLDeformerMaskConfigDialog, EditorModel)
+			.InitialMaskInfo(NeuralMaskInfo)
+			.OnSetNewVertexAttributeValues_Lambda([this, BoneName, NeuralModel](TVertexAttributesRef<float> AttributeRef)
+			{
+				// When a new vertex attribute is created in the mask config dialog, we want to initialize it with a generated mask.
+				FNeuralMorphEditorModel* NeuralEditorModel = static_cast<FNeuralMorphEditorModel*>(EditorModel);
+				UNeuralMorphInputInfo* InputInfo = Cast<UNeuralMorphInputInfo>(NeuralEditorModel->GetEditorInputInfo());
+				USkeletalMesh* SkeletalMesh = NeuralModel->GetSkeletalMesh();
+				if (InputInfo && SkeletalMesh)
+				{
+					FMLDeformerMaskInfo* MaskInfo = NeuralModel->BoneMaskInfoMap.Find(BoneName);
+					check(MaskInfo);
+					NeuralEditorModel->ApplyGeneratedMaskToVertexAttributes(SkeletalMesh, *MaskInfo, AttributeRef);
+				}
+			});
+				
+		const int32 ReturnCode = Dialog->ShowModal();
+		if (ReturnCode == 0) // OK was pressed.
+		{
+			// Create the mask info if needed.
+			if (!NeuralModel->BoneMaskInfoMap.Contains(BoneName))
+			{
+				NeuralModel->BoneMaskInfoMap.Add(BoneName, Dialog->GetMaskInfo());
+			}
+			else
+			{
+				*NeuralModel->BoneMaskInfoMap.Find(BoneName) = Dialog->GetMaskInfo();
+			}
+
+			FNeuralMorphEditorModel* NeuralEditorModel = static_cast<FNeuralMorphEditorModel*>(EditorModel);
+			NeuralEditorModel->RebuildEditorMaskInfo();
+		}
+	}
+
+	void SNeuralMorphInputWidget::ConfigureBoneGroupMask()
+	{
+		using namespace UE::MLDeformer;
+
+		USkeletalMesh* SkelMesh = EditorModel->GetModel()->GetSkeletalMesh();
+		if (!SkelMesh)
+		{
+			return;
+		}
+
+		// Get the selected bone groups.
+		check(BoneGroupsWidget.IsValid());
+		TArray<TSharedPtr<FNeuralMorphBoneGroupsTreeElement>> SelectedItems = BoneGroupsWidget->GetSelectedItems();
+
+		// Calculate the number of selected groups.
+		int32 NumSelectedGroups = 0;
+		TSharedPtr<FNeuralMorphBoneGroupsTreeElement> LastSelectedGroup;
+		for (const TSharedPtr<FNeuralMorphBoneGroupsTreeElement>& Item : SelectedItems)
+		{
+			if (Item->IsGroup())
+			{
+				LastSelectedGroup = Item;
+				NumSelectedGroups++;
+			}
+		}
+
+		// If we only selected one item, we can highlight the bones already in the mask info.
+		FNeuralMorphEditorModel* NeuralEditorModel = static_cast<FNeuralMorphEditorModel*>(EditorModel);
+		UNeuralMorphInputInfo* EditorInputInfo = Cast<UNeuralMorphInputInfo>(NeuralEditorModel->GetEditorInputInfo());
+		check(EditorInputInfo);
+
+		FMLDeformerMaskInfo* GroupMaskInfo = nullptr;
+		if (LastSelectedGroup.IsValid() && NumSelectedGroups == 1)
+		{
+			check(LastSelectedGroup->IsGroup());
+			const int32 GroupIndex = LastSelectedGroup->GroupIndex;
+			check(GroupIndex!= INDEX_NONE);
+			GroupMaskInfo = NeuralEditorModel->GetNeuralMorphModel()->BoneGroupMaskInfoMap.Find(LastSelectedGroup->Name);
+			if (!GroupMaskInfo)
+			{
+				return;
+			}
+		}
+		else
+		{
+			return;
+		}
+
+		check(GroupMaskInfo);
+		const FName GroupName = LastSelectedGroup->Name;
+
+		// Show the mask configuration dialog.
+		TSharedPtr<SMLDeformerMaskConfigDialog> Dialog = SNew(SMLDeformerMaskConfigDialog, EditorModel)
+			.InitialMaskInfo(*GroupMaskInfo)
+			.OnSetNewVertexAttributeValues_Lambda([this, GroupMaskInfo](TVertexAttributesRef<float> AttributeRef)
+			{
+				// When a new vertex attribute is created in the mask config dialog, we want to initialize it with a generated mask.
+				FNeuralMorphEditorModel* NeuralEditorModel = static_cast<FNeuralMorphEditorModel*>(EditorModel);
+				UNeuralMorphInputInfo* InputInfo = Cast<UNeuralMorphInputInfo>(NeuralEditorModel->GetEditorInputInfo());
+				USkeletalMesh* SkeletalMesh = EditorModel->GetModel()->GetSkeletalMesh();
+				if (InputInfo && SkeletalMesh)
+				{
+					NeuralEditorModel->ApplyGeneratedMaskToVertexAttributes(SkeletalMesh, *GroupMaskInfo, AttributeRef);
+				}
+			});
+				
+		const int32 ReturnCode = Dialog->ShowModal();
+		if (ReturnCode == 0) // OK was pressed.
+		{
+			// Create the mask info if needed.
+			UNeuralMorphModel* NeuralModel = Cast<UNeuralMorphModel>(EditorModel->GetModel());
+			if (!NeuralModel->BoneGroupMaskInfoMap.Contains(GroupName))
+			{
+				NeuralModel->BoneGroupMaskInfoMap.Add(GroupName, Dialog->GetMaskInfo());
+			}
+			else
+			{
+				*NeuralModel->BoneGroupMaskInfoMap.Find(GroupName) = Dialog->GetMaskInfo();
+			}
+
+			NeuralEditorModel->RebuildEditorMaskInfo();
+		}
+	}
+
 }	// namespace UE::NeuralMorphModel
 
 #undef LOCTEXT_NAMESPACE
diff --git a/Engine/Plugins/Animation/MLDeformer/NeuralMorphModel/Source/NeuralMorphModelEditor/Public/NeuralMorphEditorModel.h b/Engine/Plugins/Animation/MLDeformer/NeuralMorphModel/Source/NeuralMorphModelEditor/Public/NeuralMorphEditorModel.h
--- a/Engine/Plugins/Animation/MLDeformer/NeuralMorphModel/Source/NeuralMorphModelEditor/Public/NeuralMorphEditorModel.h
+++ b/Engine/Plugins/Animation/MLDeformer/NeuralMorphModel/Source/NeuralMorphModelEditor/Public/NeuralMorphEditorModel.h
@@ -127,13 +127,10 @@ namespace UE::NeuralMorphModel
 		 */
 		void GenerateBoneGroupMaskInfo(int32 InputInfoBoneGroupIndex, int32 HierarchyDepth);
 
-		/**
-		 * Apply the mask info to the mask buffer of floats.
-		 * @param SkeletalMesh The skeletal mesh that the mask relates to.
-		 * @param MaskInfo The mask info object to apply to the float mask buffer.
-		 * @param ItemMaskBuffer The mask buffer for this item. This should be an array view of size NumBaseMeshVerts.
-		 */
-		void ApplyMaskInfoToMaskBuffer(const USkeletalMesh* SkeletalMesh, const FNeuralMorphMaskInfo& MaskInfo, TArrayView<float> ItemMaskBuffer);
+		PRAGMA_DISABLE_DEPRECATION_WARNINGS;
+		UE_DEPRECATED(5.5, "Please use the FMLDeformerMorphModelEditorModel::ApplyMaskInfoToBuffer that takes an FMLDeformerMaskInfo as mask info type.")
+		void ApplyMaskInfoToMaskBuffer(const USkeletalMesh* SkeletalMesh, const FNeuralMorphMaskInfo& MaskInfo, TArrayView<float> ItemMaskBuffer) {}
+		PRAGMA_ENABLE_DEPRECATION_WARNINGS;
 
 		/**
 		 * Set the mask visualization item index. This specifies which item (bone, curve, bone group, curve group) to visualize the mask for.
diff --git a/Engine/Plugins/Animation/MLDeformer/NeuralMorphModel/Source/NeuralMorphModelEditor/Public/SNeuralMorphInputWidget.h b/Engine/Plugins/Animation/MLDeformer/NeuralMorphModel/Source/NeuralMorphModelEditor/Public/SNeuralMorphInputWidget.h
--- a/Engine/Plugins/Animation/MLDeformer/NeuralMorphModel/Source/NeuralMorphModelEditor/Public/SNeuralMorphInputWidget.h
+++ b/Engine/Plugins/Animation/MLDeformer/NeuralMorphModel/Source/NeuralMorphModelEditor/Public/SNeuralMorphInputWidget.h
@@ -5,9 +5,12 @@
 #include "CoreTypes.h"
 #include "Templates/SharedPointer.h"
 #include "SMLDeformerInputWidget.h"
+#include "MeshDescription.h"
 #include "Framework/Commands/Commands.h"
 
 class FMenuBuilder;
+class USkeletalMesh;
+struct FMLDeformerMaskInfo;
 
 namespace UE::MLDeformer
 {
@@ -36,6 +39,9 @@ namespace UE::NeuralMorphModel
 		TSharedPtr<FUICommandInfo> ResetAllBoneGroupMasks;
 		TSharedPtr<FUICommandInfo> ResetSelectedBoneGroupMasks;
 		TSharedPtr<FUICommandInfo> ExpandSelectedBoneGroupMasks;
+
+		TSharedPtr<FUICommandInfo> ConfigureBoneMask;
+		TSharedPtr<FUICommandInfo> ConfigureBoneGroupMask;
 	};
 
 
@@ -89,6 +95,9 @@ namespace UE::NeuralMorphModel
 		void ResetSelectedBoneGroupMasks();
 		void ExpandBoneGroupMasks();
 
+		void ConfigureBoneMask();
+		void ConfigureBoneGroupMask();
+
 		void BindCommands();
 		FReply ShowBoneGroupsManageContextMenu();
 		FReply ShowCurveGroupsManageContextMenu();
diff --git a/Engine/Plugins/Experimental/MeshModelingToolsetExp/Source/MeshModelingToolsExp/Public/DynamicMeshBrushTool.h b/Engine/Plugins/Experimental/MeshModelingToolsetExp/Source/MeshModelingToolsExp/Public/DynamicMeshBrushTool.h
--- a/Engine/Plugins/Experimental/MeshModelingToolsetExp/Source/MeshModelingToolsExp/Public/DynamicMeshBrushTool.h
+++ b/Engine/Plugins/Experimental/MeshModelingToolsetExp/Source/MeshModelingToolsExp/Public/DynamicMeshBrushTool.h
@@ -36,6 +36,8 @@ public:
 
 	virtual bool HitTest(const FRay& Ray, FHitResult& OutHit) override;
 
+	TObjectPtr<UPreviewMesh> GetPreviewMesh() { return PreviewMesh; }
+
 protected:
 	// subclasses can override these to customize behavior
 	virtual void OnShutdown(EToolShutdownType ShutdownType) {}
diff --git a/Engine/Plugins/Experimental/MeshModelingToolsetExp/Source/MeshModelingToolsExp/Public/MeshAttributePaintTool.h b/Engine/Plugins/Experimental/MeshModelingToolsetExp/Source/MeshModelingToolsExp/Public/MeshAttributePaintTool.h
--- a/Engine/Plugins/Experimental/MeshModelingToolsetExp/Source/MeshModelingToolsExp/Public/MeshAttributePaintTool.h
+++ b/Engine/Plugins/Experimental/MeshModelingToolsetExp/Source/MeshModelingToolsExp/Public/MeshAttributePaintTool.h
@@ -232,6 +232,8 @@ public:
 		return true;
 	}
 
+	UE::Geometry::TDynamicVerticesOctree3<FDynamicMesh3>& GetVerticesOctree() { return VerticesOctree; }
+
 protected:
 	virtual void ApplyStamp(const FBrushStampData& Stamp);
 
diff --git a/Engine/Source/Runtime/GeometryCore/Public/DynamicMesh/DynamicVertexSkinWeightsAttribute.h b/Engine/Source/Runtime/GeometryCore/Public/DynamicMesh/DynamicVertexSkinWeightsAttribute.h
--- a/Engine/Source/Runtime/GeometryCore/Public/DynamicMesh/DynamicVertexSkinWeightsAttribute.h
+++ b/Engine/Source/Runtime/GeometryCore/Public/DynamicMesh/DynamicVertexSkinWeightsAttribute.h
@@ -320,8 +320,8 @@ public:
 
 		const int32 NumEntries = BoneWeights.Num();
 
-		OutBones.SetNum(NumEntries);
-		OutWeights.SetNum(NumEntries);
+		OutBones.SetNum(NumEntries, EAllowShrinking::No);
+		OutWeights.SetNum(NumEntries, EAllowShrinking::No);
 
 		for (int32 BoneIdx = 0; BoneIdx < NumEntries; ++BoneIdx)
 		{
