[MLDeformer] Fixed an issue where ground truth heatmaps were incorrect on some frames. Disabled showing ground truth heatmaps while playing the animation as well as that just doesn't work well due to frame differences. Directly added some extra error handling in the Neural Morph Model and logging some info about the Cuda hardware used.

diff --git a/Engine/Plugins/Animation/MLDeformer/MLDeformerFramework/Source/MLDeformerFramework/Private/MLDeformerGeomCacheHelpers.cpp b/Engine/Plugins/Animation/MLDeformer/MLDeformerFramework/Source/MLDeformerFramework/Private/MLDeformerGeomCacheHelpers.cpp
--- a/Engine/Plugins/Animation/MLDeformer/MLDeformerFramework/Source/MLDeformerFramework/Private/MLDeformerGeomCacheHelpers.cpp
+++ b/Engine/Plugins/Animation/MLDeformer/MLDeformerFramework/Source/MLDeformerFramework/Private/MLDeformerGeomCacheHelpers.cpp
@@ -294,7 +294,7 @@ namespace UE::MLDeformer
 	void SampleGeomCachePositions(
 		int32 InLODIndex,
 		float InSampleTime,
-		const TArray<UE::MLDeformer::FMLDeformerGeomCacheMeshMapping>& InMeshMappings,
+		const TArray<FMLDeformerGeomCacheMeshMapping>& InMeshMappings,
 		const USkeletalMesh* InSkelMesh,
 		const UGeometryCache* InGeometryCache,
 		const FTransform& InAlignmentTransform,
@@ -318,7 +318,7 @@ namespace UE::MLDeformer
 
 		const FSkeletalMeshConstAttributes MeshAttributes(*MeshDescription);
 		const FSkeletalMeshAttributes::FSourceGeometryPartVertexOffsetAndCountConstRef GeoPartOffsetAndCounts = MeshAttributes.GetSourceGeometryPartVertexOffsetAndCounts();
-		
+
 		const uint32 NumVertices = MeshDescription->Vertices().Num();
 		OutPositions.Reset(NumVertices);
 		OutPositions.AddZeroed(NumVertices);
@@ -330,7 +330,7 @@ namespace UE::MLDeformer
 			TArrayView<const int32> GeoPartInfo = GeoPartOffsetAndCounts.Get(MeshMapping.MeshIndex);
 			const int32 StartImportedVertex = GeoPartInfo[0];
 			const int32 NumImportedVertices = GeoPartInfo[1];
-			
+
 			UGeometryCacheTrack* Track = InGeometryCache->Tracks[MeshMapping.TrackIndex];
 
 			FGeometryCacheMeshData GeomCacheMeshData;
@@ -352,6 +352,67 @@ namespace UE::MLDeformer
 		}
 	}
 
+	void SampleGeomCachePositionsAtFrame(
+		int32 InLODIndex,
+		int32 FrameIndex,
+		const TArray<UE::MLDeformer::FMLDeformerGeomCacheMeshMapping>& InMeshMappings,
+		const USkeletalMesh* InSkelMesh,
+		const UGeometryCache* InGeometryCache,
+		const FTransform& InAlignmentTransform,
+		TArray<FVector3f>& OutPositions)
+	{
+		if (InGeometryCache == nullptr)
+		{
+			return;
+		}
+
+		if (!ensure(InSkelMesh != nullptr))
+		{
+			return;
+		}
+
+		const FMeshDescription* MeshDescription = InSkelMesh->GetMeshDescription(InLODIndex);
+		if (!ensure(MeshDescription))
+		{
+			return;
+		}
+
+		const FSkeletalMeshConstAttributes MeshAttributes(*MeshDescription);
+		const FSkeletalMeshAttributes::FSourceGeometryPartVertexOffsetAndCountConstRef GeoPartOffsetAndCounts = MeshAttributes.GetSourceGeometryPartVertexOffsetAndCounts();
+
+		const uint32 NumVertices = MeshDescription->Vertices().Num();
+		OutPositions.Reset(NumVertices);
+		OutPositions.AddZeroed(NumVertices);
+
+		// For all mesh mappings we found.
+		for (int32 MeshMappingIndex = 0; MeshMappingIndex < InMeshMappings.Num(); ++MeshMappingIndex)
+		{
+			const UE::MLDeformer::FMLDeformerGeomCacheMeshMapping& MeshMapping = InMeshMappings[MeshMappingIndex];
+			TArrayView<const int32> GeoPartInfo = GeoPartOffsetAndCounts.Get(MeshMapping.MeshIndex);
+			const int32 StartImportedVertex = GeoPartInfo[0];
+			const int32 NumImportedVertices = GeoPartInfo[1];
+
+			UGeometryCacheTrack* Track = InGeometryCache->Tracks[MeshMapping.TrackIndex];
+
+			FGeometryCacheMeshData GeomCacheMeshData;
+			if (!Track->GetMeshDataAtSampleIndex(FrameIndex, GeomCacheMeshData))
+			{
+				continue;
+			}
+
+			for (int32 VertexIndex = 0; VertexIndex < NumImportedVertices; ++VertexIndex)
+			{
+				const int32 SkinnedVertexIndex = StartImportedVertex + VertexIndex;
+				const int32 GeomCacheVertexIndex = MeshMapping.SkelMeshToTrackVertexMap[VertexIndex];
+				if (GeomCacheVertexIndex != INDEX_NONE && GeomCacheMeshData.Positions.IsValidIndex(GeomCacheVertexIndex))
+				{
+					const FVector3f GeomCacheVertexPos = (FVector3f)InAlignmentTransform.TransformPosition((FVector)GeomCacheMeshData.Positions[GeomCacheVertexIndex]);
+					OutPositions[SkinnedVertexIndex] = GeomCacheVertexPos;
+				}
+			}
+		}
+	}
+
 	int32 ExtractNumImportedGeomCacheVertices(UGeometryCache* GeometryCache)
 	{
 		if (GeometryCache == nullptr)
diff --git a/Engine/Plugins/Animation/MLDeformer/MLDeformerFramework/Source/MLDeformerFramework/Private/MLDeformerGeomCacheModel.cpp b/Engine/Plugins/Animation/MLDeformer/MLDeformerFramework/Source/MLDeformerFramework/Private/MLDeformerGeomCacheModel.cpp
--- a/Engine/Plugins/Animation/MLDeformer/MLDeformerFramework/Source/MLDeformerFramework/Private/MLDeformerGeomCacheModel.cpp
+++ b/Engine/Plugins/Animation/MLDeformer/MLDeformerFramework/Source/MLDeformerFramework/Private/MLDeformerGeomCacheModel.cpp
@@ -121,6 +121,35 @@ bool UMLDeformerGeomCacheModel::HasTrainingGroundTruth() const
 	return false;
 }
 
+void UMLDeformerGeomCacheModel::SampleGroundTruthPositionsAtFrame(int32 FrameIndex, TArray<FVector3f>& OutPositions)
+{
+	const UMLDeformerGeomCacheVizSettings* GeomCacheVizSettings = GetGeomCacheVizSettings();
+	check(GeomCacheVizSettings);
+
+	UGeometryCache* GeomCache = GeomCacheVizSettings->GetTestGroundTruth();
+	if (GeomCache == nullptr)
+	{
+		OutPositions.Reset();
+		return;
+	}
+
+	if (MeshMappings.IsEmpty())
+	{
+		TArray<FString> FailedImportedMeshNames;
+		TArray<FString> VertexMisMatchNames;
+		UE::MLDeformer::GenerateGeomCacheMeshMappings(GetSkeletalMesh(), GeomCache, MeshMappings, FailedImportedMeshNames, VertexMisMatchNames, /**bSuppressLog*/true);
+	}
+
+	UE::MLDeformer::SampleGeomCachePositionsAtFrame(
+		0,
+		FrameIndex,
+		MeshMappings,
+		GetSkeletalMesh(),
+		GeomCache,
+		GetAlignmentTransform(),
+		OutPositions);
+}
+
 void UMLDeformerGeomCacheModel::SampleGroundTruthPositions(float SampleTime, TArray<FVector3f>& OutPositions)
 {
 	const UMLDeformerGeomCacheVizSettings* GeomCacheVizSettings = GetGeomCacheVizSettings();
@@ -140,6 +169,7 @@ void UMLDeformerGeomCacheModel::SampleGroundTruthPositions(float SampleTime, TAr
 		UE::MLDeformer::GenerateGeomCacheMeshMappings(GetSkeletalMesh(), GeomCache, MeshMappings, FailedImportedMeshNames, VertexMisMatchNames, /**bSuppressLog*/true);
 	}
 
+	PRAGMA_DISABLE_DEPRECATION_WARNINGS;
 	UE::MLDeformer::SampleGeomCachePositions(
 		0,
 		SampleTime,
@@ -148,7 +178,9 @@ void UMLDeformerGeomCacheModel::SampleGroundTruthPositions(float SampleTime, TAr
 		GeomCache,
 		GetAlignmentTransform(),
 		OutPositions);
+	PRAGMA_ENABLE_DEPRECATION_WARNINGS;
 }
+
 #endif
 
 #undef LOCTEXT_NAMESPACE
diff --git a/Engine/Plugins/Animation/MLDeformer/MLDeformerFramework/Source/MLDeformerFramework/Private/MLDeformerGraphDebugDataInterface.cpp b/Engine/Plugins/Animation/MLDeformer/MLDeformerFramework/Source/MLDeformerFramework/Private/MLDeformerGraphDebugDataInterface.cpp
--- a/Engine/Plugins/Animation/MLDeformer/MLDeformerFramework/Source/MLDeformerFramework/Private/MLDeformerGraphDebugDataInterface.cpp
+++ b/Engine/Plugins/Animation/MLDeformer/MLDeformerFramework/Source/MLDeformerFramework/Private/MLDeformerGraphDebugDataInterface.cpp
@@ -133,8 +133,8 @@ FComputeDataProviderRenderProxy* UMLDeformerGraphDebugDataProvider::GetRenderPro
 		UMLDeformerModel* Model = DeformerAsset->GetModel();
 		if (Model->GetVizSettings()->GetHeatMapMode() == EMLDeformerHeatMapMode::GroundTruth)
 		{
-			const float SampleTime = DeformerComponent->GetModelInstance()->GetSkeletalMeshComponent()->GetPosition();
-			Model->SampleGroundTruthPositions(SampleTime, Proxy->GetGroundTruthPositions());
+			const int32 GroundTruthFrame = DeformerComponent->GetModelInstance()->GetGroundTruthFrameIndex();
+			Model->SampleGroundTruthPositionsAtFrame(GroundTruthFrame, Proxy->GetGroundTruthPositions());
 		}
 		Proxy->HandleZeroGroundTruthPositions();
 		return Proxy;
diff --git a/Engine/Plugins/Animation/MLDeformer/MLDeformerFramework/Source/MLDeformerFramework/Private/MLDeformerModelInstance.cpp b/Engine/Plugins/Animation/MLDeformer/MLDeformerFramework/Source/MLDeformerFramework/Private/MLDeformerModelInstance.cpp
--- a/Engine/Plugins/Animation/MLDeformer/MLDeformerFramework/Source/MLDeformerFramework/Private/MLDeformerModelInstance.cpp
+++ b/Engine/Plugins/Animation/MLDeformer/MLDeformerFramework/Source/MLDeformerFramework/Private/MLDeformerModelInstance.cpp
@@ -444,4 +444,14 @@ void UMLDeformerModelInstance::Tick(float DeltaTime, float ModelWeight)
 			CopyDataFromDebugActor(DebugActor, DebugMLDeformerComponent);
 		}
 	}
+
+	void UMLDeformerModelInstance::SetGroundTruthFrameIndex(int32 FrameIndex)
+	{ 
+		GroundTruthFrameIndex = FrameIndex;
+	}
+
+	int32 UMLDeformerModelInstance::GetGroundTruthFrameIndex() const
+	{ 
+		return GroundTruthFrameIndex;
+	}
 #endif
diff --git a/Engine/Plugins/Animation/MLDeformer/MLDeformerFramework/Source/MLDeformerFramework/Public/MLDeformerGeomCacheHelpers.h b/Engine/Plugins/Animation/MLDeformer/MLDeformerFramework/Source/MLDeformerFramework/Public/MLDeformerGeomCacheHelpers.h
--- a/Engine/Plugins/Animation/MLDeformer/MLDeformerFramework/Source/MLDeformerFramework/Public/MLDeformerGeomCacheHelpers.h
+++ b/Engine/Plugins/Animation/MLDeformer/MLDeformerFramework/Source/MLDeformerFramework/Public/MLDeformerGeomCacheHelpers.h
@@ -78,20 +78,30 @@ namespace UE::MLDeformer
 	 */
 	MLDEFORMERFRAMEWORK_API void GenerateGeomCacheMeshMappings(USkeletalMesh* SkelMesh, UGeometryCache* GeomCache, TArray<FMLDeformerGeomCacheMeshMapping>& OutMeshMappings, TArray<FString>& OutFailedImportedMeshNames, TArray<FString>& OutVertexMisMatchNames, bool bSuppressLog = false);
 
+	UE_DEPRECATED(5.5, "Please use SampleGeomCachePositionsAtFrame instead.")
+	MLDEFORMERFRAMEWORK_API void SampleGeomCachePositions(
+		int32 InLODIndex,
+		float InSampleTime,
+		const TArray<FMLDeformerGeomCacheMeshMapping>& InMeshMappings,
+		const USkeletalMesh* InSkelMesh,
+		const UGeometryCache* InGeometryCache,
+		const FTransform& InAlignmentTransform,
+		TArray<FVector3f>& OutPositions);
+
 	/**
 	 * Sample the vertex position data of a geometry cache, at a given time stamp.
 	 * This basically allows you to sample ground truth or training target vertex positions.
 	 * @param InLODIndex The LOD level to sample for.
-	 * @param InSampleTime The time to take the sample at, in seconds.
-	 * @param InMeshMappings The geometry cache to skeletal mesh mappings, which can be generated using the GenerateGeomCacheMeshMapings method.
+	 * @param FrameIndex The frame number to sample.
+	 * @param InMeshMappings The geometry cache to skeletal mesh mappings, which can be generated using the GenerateGeomCacheMeshMappings method.
 	 * @param InSkelMesh The skeletal mesh object.
 	 * @param InGeometryCache The geometry cache object.
 	 * @param InAlignmentTransform This is the transformation that will be applied on the sampled positions, as post process. This can be used to rotate or scale the position data.
-	 * @param OutPosiitons The resulting positions, as sampled at the specified time. This array will be resized internally.
+	 * @param OutPositions The resulting positions, as sampled at the specified time. This array will be resized internally.
 	 */
-	MLDEFORMERFRAMEWORK_API void SampleGeomCachePositions(
+	MLDEFORMERFRAMEWORK_API void SampleGeomCachePositionsAtFrame(
 		int32 InLODIndex,
-		float InSampleTime,
+		int32 FrameIndex,
 		const TArray<FMLDeformerGeomCacheMeshMapping>& InMeshMappings,
 		const USkeletalMesh* InSkelMesh,
 		const UGeometryCache* InGeometryCache,
diff --git a/Engine/Plugins/Animation/MLDeformer/MLDeformerFramework/Source/MLDeformerFramework/Public/MLDeformerGeomCacheModel.h b/Engine/Plugins/Animation/MLDeformer/MLDeformerFramework/Source/MLDeformerFramework/Public/MLDeformerGeomCacheModel.h
--- a/Engine/Plugins/Animation/MLDeformer/MLDeformerFramework/Source/MLDeformerFramework/Public/MLDeformerGeomCacheModel.h
+++ b/Engine/Plugins/Animation/MLDeformer/MLDeformerFramework/Source/MLDeformerFramework/Public/MLDeformerGeomCacheModel.h
@@ -38,7 +38,8 @@ public:
 	// UMLDeformerModel overrides.
 #if WITH_EDITORONLY_DATA
 	virtual bool HasTrainingGroundTruth() const override;
-	virtual void SampleGroundTruthPositions(float SampleTime, TArray<FVector3f>& OutPositions) override;
+	virtual void SampleGroundTruthPositionsAtFrame(int32 FrameIndex, TArray<FVector3f>& OutPositions) override;
+	virtual void SampleGroundTruthPositions(float SampleTime, TArray<FVector3f>& OutPositions) override;	// Deprecated in base class, will be removed soon.
 #endif
 #if WITH_EDITOR
 	virtual void UpdateNumTargetMeshVertices() override;
diff --git a/Engine/Plugins/Animation/MLDeformer/MLDeformerFramework/Source/MLDeformerFramework/Public/MLDeformerModel.h b/Engine/Plugins/Animation/MLDeformer/MLDeformerFramework/Source/MLDeformerFramework/Public/MLDeformerModel.h
--- a/Engine/Plugins/Animation/MLDeformer/MLDeformerFramework/Source/MLDeformerFramework/Public/MLDeformerModel.h
+++ b/Engine/Plugins/Animation/MLDeformer/MLDeformerFramework/Source/MLDeformerFramework/Public/MLDeformerModel.h
@@ -249,12 +249,15 @@ public:
 	 */
 	virtual bool HasTrainingGroundTruth() const				{ return false; }
 
+	UE_DEPRECATED(5.5, "Please use SampleGroundTruthPositionAtFrame instead.")
+	virtual void SampleGroundTruthPositions(float SampleTime, TArray<FVector3f>& OutPositions) {}
+
 	/**
-	 * Sample the positions from the target (ground truth) mesh, at a specific time (in seconds).
-	 * @param SampleTime The time to sample the positions at, in seconds.
+	 * Sample the positions from the target (ground truth) mesh, at a specific frame.
+	 * @param FrameIndex The frame number to sample the positions at.
 	 * @param OutPositions The array that will receive the resulting vertex positions. This array will automatically be resized internally.
 	 */
-	virtual void SampleGroundTruthPositions(float SampleTime, TArray<FVector3f>& OutPositions) {}
+	virtual void SampleGroundTruthPositionsAtFrame(int32 FrameIndex, TArray<FVector3f>& OutPositions) {}
 #endif
 
 #if WITH_EDITOR
diff --git a/Engine/Plugins/Animation/MLDeformer/MLDeformerFramework/Source/MLDeformerFramework/Public/MLDeformerModelInstance.h b/Engine/Plugins/Animation/MLDeformer/MLDeformerFramework/Source/MLDeformerFramework/Public/MLDeformerModelInstance.h
--- a/Engine/Plugins/Animation/MLDeformer/MLDeformerFramework/Source/MLDeformerFramework/Public/MLDeformerModelInstance.h
+++ b/Engine/Plugins/Animation/MLDeformer/MLDeformerFramework/Source/MLDeformerFramework/Public/MLDeformerModelInstance.h
@@ -112,6 +112,18 @@ public:
 	 * @param DebugMLDeformerComponent The debug ML Deformer component to copy data from.
 	 */
 	virtual void CopyDataFromDebugActor(const AActor* DebugActor, const UMLDeformerComponent* DebugMLDeformerComponent) {}
+
+	/**
+	 * Set the frame index that we use to sample the ground truth vertex positions when showing the heat map in ground truth mode.
+	 * @param FrameIndex The frame number in the geometry cache that we sample the ground truth positions from.
+	 */
+	void SetGroundTruthFrameIndex(int32 FrameIndex);
+
+	/**
+	 * Get the frame index that we use to sample the ground truth vertex positions when showing the heat map in ground truth mode.
+	 * @return The frame number inside the geometry cache that we sample the ground truth positions from.
+	 */
+	int32 GetGroundTruthFrameIndex() const;
 #endif
 
 	/**
@@ -263,4 +275,9 @@ protected:
 
 	/** Has this instance had a successful UMLDeformerModel::PostMLDeformerComponentInit call? */
 	bool bHasPostInitialized = false;
+
+#if WITH_EDITOR
+	/** The current ground truth frame index. This is used for displaying heatmaps versus the ground truth accurately. */
+	int32 GroundTruthFrameIndex = 0;
+#endif
 };
diff --git a/Engine/Plugins/Animation/MLDeformer/MLDeformerFramework/Source/MLDeformerFrameworkEditor/Private/MLDeformerEditorModel.cpp b/Engine/Plugins/Animation/MLDeformer/MLDeformerFramework/Source/MLDeformerFrameworkEditor/Private/MLDeformerEditorModel.cpp
--- a/Engine/Plugins/Animation/MLDeformer/MLDeformerFramework/Source/MLDeformerFrameworkEditor/Private/MLDeformerEditorModel.cpp
+++ b/Engine/Plugins/Animation/MLDeformer/MLDeformerFramework/Source/MLDeformerFrameworkEditor/Private/MLDeformerEditorModel.cpp
@@ -1003,6 +1003,14 @@ namespace UE::MLDeformer
 				if (EditorActor && EditorActor->IsTestActor())
 				{
 					EditorActor->SetPlayPosition(PlayOffset);
+					UMLDeformerComponent* MLDComponent = EditorActor->GetMLDeformerComponent();
+					if (MLDComponent)
+					{
+						if (MLDComponent->GetModelInstance())
+						{
+							MLDComponent->GetModelInstance()->SetGroundTruthFrameIndex(TargetFrame);
+						}
+					}
 				}
 			}
 			VizSettings->SetTestingFrameNumber(TargetFrame);
@@ -1311,6 +1319,18 @@ namespace UE::MLDeformer
 				EditorActor->Pause(bMustPause);
 			}
 		}
+
+		if (VizSettings->GetShowHeatMap() && VizSettings->GetHeatMapMode() == EMLDeformerHeatMapMode::GroundTruth)
+		{
+			if (!bMustPause)
+			{
+				SetHeatMapMaterialEnabled(false);
+			}
+			else
+			{
+				SetHeatMapMaterialEnabled(VizSettings->GetShowHeatMap());
+			}
+		}
 	}
 
 	bool FMLDeformerEditorModel::IsPlayingAnim() const
diff --git a/Engine/Plugins/Animation/MLDeformer/MLDeformerFramework/Source/MLDeformerFrameworkEditor/Private/MLDeformerEditorToolkit.cpp b/Engine/Plugins/Animation/MLDeformer/MLDeformerFramework/Source/MLDeformerFrameworkEditor/Private/MLDeformerEditorToolkit.cpp
--- a/Engine/Plugins/Animation/MLDeformer/MLDeformerFramework/Source/MLDeformerFrameworkEditor/Private/MLDeformerEditorToolkit.cpp
+++ b/Engine/Plugins/Animation/MLDeformer/MLDeformerFramework/Source/MLDeformerFrameworkEditor/Private/MLDeformerEditorToolkit.cpp
@@ -747,6 +747,14 @@ namespace UE::MLDeformer
 			}
 			break;
 
+			// There is an error another error, some exception got thrown while training.
+			case ETrainingResult::Other:
+			{
+				GEditor->PlayEditorSound(TEXT("/Engine/EditorSounds/Notifications/CompileFailed_Cue.CompileFailed_Cue"));
+				WindowMessage = LOCTEXT("TrainingOtherError", "Training failed!\nPlease check the output log.\n");
+			}
+			break;
+
 			// Unhandled error codes.
 			default:
 				checkf(false, TEXT("Unknown error code"));
diff --git a/Engine/Plugins/Animation/MLDeformer/MLDeformerFramework/Source/MLDeformerFrameworkEditor/Public/MLDeformerEditorModel.h b/Engine/Plugins/Animation/MLDeformer/MLDeformerFramework/Source/MLDeformerFrameworkEditor/Public/MLDeformerEditorModel.h
--- a/Engine/Plugins/Animation/MLDeformer/MLDeformerFramework/Source/MLDeformerFrameworkEditor/Public/MLDeformerEditorModel.h
+++ b/Engine/Plugins/Animation/MLDeformer/MLDeformerFramework/Source/MLDeformerFrameworkEditor/Public/MLDeformerEditorModel.h
@@ -52,7 +52,10 @@ enum class ETrainingResult : uint8
 	FailOnData,
 
 	/** The python script has some error (see output log). Or it's missing some required model training class. */
-	FailPythonError
+	FailPythonError,
+
+	/** Any other error, for example when some exception got thrown during training. */
+	Other
 };
 
 namespace UE::MLDeformer
