[MLDeformer] Allow specific training data to focus on specific areas of the mesh. This enables users to train different body parts at higher quality. One example of this is training data for the hand. The user can now provide a vertex attribute mask for the hand training data, which allows the training to focus on the area around the hand, which results in higher quality deformations in the hand region.

Also changed the Neural Morph Model's default values for regularization factor and smoothing beta to 0.

diff --git a/Engine/Plugins/Animation/MLDeformer/MLDeformerFramework/Source/MLDeformerFramework/Public/MLDeformerTrainingInputAnim.h b/Engine/Plugins/Animation/MLDeformer/MLDeformerFramework/Source/MLDeformerFramework/Public/MLDeformerTrainingInputAnim.h
--- a/Engine/Plugins/Animation/MLDeformer/MLDeformerFramework/Source/MLDeformerFramework/Public/MLDeformerTrainingInputAnim.h
+++ b/Engine/Plugins/Animation/MLDeformer/MLDeformerFramework/Source/MLDeformerFramework/Public/MLDeformerTrainingInputAnim.h
@@ -39,21 +39,24 @@ public:
 	void SetUseCustomRange(bool bCustomRange)					{ bUseCustomRange = bCustomRange; }
 	void SetStartFrame(int32 FrameIndex)						{ StartFrame = FrameIndex; }
 	void SetEndFrame(int32 FrameIndex)							{ EndFrame = FrameIndex; }
+	void SetVertexMask(FName MaskName)							{ VertexMask = MaskName; }
 
 	bool IsEnabled() const										{ return bEnabled; }
 	bool GetUseCustomRange() const								{ return bUseCustomRange; }
 	int32 GetStartFrame() const									{ return StartFrame; }
 	int32 GetEndFrame() const									{ return EndFrame; }
+	FName GetVertexMask() const									{ return VertexMask; }
 	const UAnimSequence* GetAnimSequence() const				{ return AnimSequence.LoadSynchronous(); }
 	UAnimSequence* GetAnimSequence()							{ return AnimSequence.LoadSynchronous(); }
 	const TSoftObjectPtr<UAnimSequence>& GetAnimSequenceSoftObjectPtr() const	{ return AnimSequence; } 
-	TSoftObjectPtr<UAnimSequence>& GetAnimSequenceSoftObjectPtr()				{ return AnimSequence; } 
+	TSoftObjectPtr<UAnimSequence>& GetAnimSequenceSoftObjectPtr()				{ return AnimSequence; }
 
 	static FName GetAnimSequencePropertyName()					{ return GET_MEMBER_NAME_CHECKED(FMLDeformerTrainingInputAnim, AnimSequence); }
 	static FName GetEnabledPropertyName()						{ return GET_MEMBER_NAME_CHECKED(FMLDeformerTrainingInputAnim, bEnabled); }
 	static FName GetUseCustomRangePropertyName()				{ return GET_MEMBER_NAME_CHECKED(FMLDeformerTrainingInputAnim, bUseCustomRange); }
 	static FName GetStartFramePropertyName()					{ return GET_MEMBER_NAME_CHECKED(FMLDeformerTrainingInputAnim, StartFrame); }
 	static FName GetEndFramePropertyName()						{ return GET_MEMBER_NAME_CHECKED(FMLDeformerTrainingInputAnim, EndFrame); }
+	static FName GetVertexMaskPropertyName()					{ return GET_MEMBER_NAME_CHECKED(FMLDeformerTrainingInputAnim, VertexMask); }
 
 private:
 	/** The animation sequence. */
@@ -72,6 +75,13 @@ private:
 	UPROPERTY(EditAnywhere, Category = "Frame Settings", meta = (ClampMin="0", EditCondition="bUseCustomRange"))
 	int32 EndFrame = 0;
 
+	/**
+	 * The optional vertex attribute mask to use for this training data. This can be used to focus training on a specific area on the mesh.
+	 * For example if you have training data for a hand, you can specify a mask that contains just the hand and maybe forearm.
+	 */
+	UPROPERTY(EditAnywhere, Category = "Settings")
+	FName VertexMask;
+
 	/** Is this animation enabled? If not, it is excluded from the training process. */
 	UPROPERTY(EditAnywhere, Category = "Settings", meta = (DisplayPriority = MAX_int32))	// Show as last property.
 	bool bEnabled = true;
diff --git a/Engine/Plugins/Animation/MLDeformer/MLDeformerFramework/Source/MLDeformerFrameworkEditor/Private/MLDeformerEditorModel.cpp b/Engine/Plugins/Animation/MLDeformer/MLDeformerFramework/Source/MLDeformerFrameworkEditor/Private/MLDeformerEditorModel.cpp
--- a/Engine/Plugins/Animation/MLDeformer/MLDeformerFramework/Source/MLDeformerFrameworkEditor/Private/MLDeformerEditorModel.cpp
+++ b/Engine/Plugins/Animation/MLDeformer/MLDeformerFramework/Source/MLDeformerFrameworkEditor/Private/MLDeformerEditorModel.cpp
@@ -127,6 +127,22 @@ namespace UE::MLDeformer
 		UpdateLODMappings();
 	}
 
+	void FMLDeformerEditorModel::UpdateVertexAttributeNames()
+	{
+		const TArray<FName> AttributeNames = Model->GetVertexAttributeNames();
+
+		// Now build a list of shared pointers to attribute names.
+		VertexAttributeNames.Reset();
+		VertexAttributeNames.Reserve(AttributeNames.Num() + 1);
+
+		const FName DisabledName = FName(LOCTEXT("DisabledString", "Disabled").ToString());
+		VertexAttributeNames.Emplace(MakeShared<FName>(DisabledName));
+		for (int32 Index = 0; Index < AttributeNames.Num(); ++Index)
+		{
+			VertexAttributeNames.Emplace(MakeShared<FName>(AttributeNames[Index]));
+		}
+	}
+
 	void FMLDeformerEditorModel::CreateSamplers()
 	{
 		// First get rid of the existing samplers.
diff --git a/Engine/Plugins/Animation/MLDeformer/MLDeformerFramework/Source/MLDeformerFrameworkEditor/Private/MLDeformerGeomCacheTrainingInputAnimCustomize.cpp b/Engine/Plugins/Animation/MLDeformer/MLDeformerFramework/Source/MLDeformerFrameworkEditor/Private/MLDeformerGeomCacheTrainingInputAnimCustomize.cpp
--- a/Engine/Plugins/Animation/MLDeformer/MLDeformerFramework/Source/MLDeformerFrameworkEditor/Private/MLDeformerGeomCacheTrainingInputAnimCustomize.cpp
+++ b/Engine/Plugins/Animation/MLDeformer/MLDeformerFramework/Source/MLDeformerFrameworkEditor/Private/MLDeformerGeomCacheTrainingInputAnimCustomize.cpp
@@ -13,6 +13,7 @@
 #include "Widgets/SBoxPanel.h"
 #include "Widgets/Text/STextBlock.h"
 #include "Widgets/Layout/SBox.h"
+#include "Widgets/Input/SComboBox.h"
 #include "SWarningOrErrorBox.h"
 #include "DetailWidgetRow.h"
 #include "Engine/SkeletalMesh.h"
@@ -91,12 +92,14 @@ namespace UE::MLDeformer
 		TSharedPtr<IPropertyHandle> UseCustomRangeHandle	= PropertyHandles.FindChecked(FMLDeformerGeomCacheTrainingInputAnim::GetUseCustomRangePropertyName());
 		TSharedPtr<IPropertyHandle> StartFrameHandle		= PropertyHandles.FindChecked(FMLDeformerGeomCacheTrainingInputAnim::GetStartFramePropertyName());
 		TSharedPtr<IPropertyHandle> EndFrameHandle			= PropertyHandles.FindChecked(FMLDeformerGeomCacheTrainingInputAnim::GetEndFramePropertyName());
+		TSharedPtr<IPropertyHandle> VertexMapHandle			= PropertyHandles.FindChecked(FMLDeformerGeomCacheTrainingInputAnim::GetVertexMaskPropertyName());
 		check(AnimSequenceHandle.IsValid());
 		check(GeomCacheHandle.IsValid());
 		check(EnabledHandle.IsValid());
 		check(UseCustomRangeHandle.IsValid());
 		check(StartFrameHandle.IsValid());
 		check(EndFrameHandle.IsValid());
+		check(VertexMapHandle.IsValid());
 
 		UMLDeformerGeomCacheModel* Model = FindMLDeformerModel(StructPropertyHandle);
 		check(Model);
@@ -150,6 +153,71 @@ namespace UE::MLDeformer
 		ChildBuilder.AddProperty(EndFrameHandle.ToSharedRef())
 			.EditCondition(InputAnim.IsEnabled() && InputAnim.GetUseCustomRange(), nullptr);
 
+		// Vertex Map property.
+		if (EditorModel->GetSupportsPerTrainingInputAnimVertexMask())
+		{
+			EditorModel->UpdateVertexAttributeNames();
+			IDetailPropertyRow& VertexMapRow = ChildBuilder.AddProperty(VertexMapHandle.ToSharedRef());
+			VertexMapRow.EditCondition(InputAnim.IsEnabled(), nullptr);
+			VertexMapRow.CustomWidget()
+			.NameContent()
+			[
+				VertexMapRow.GetPropertyHandle()->CreatePropertyNameWidget()					
+			]
+			.ValueContent()
+			[
+				SNew(SComboBox<TSharedPtr<FName>>)
+				.OptionsSource(&EditorModel->GetVertexAttributeNames())
+				.OnGenerateWidget_Lambda
+				(
+					[&StructCustomizationUtils](const TSharedPtr<FName> Item)
+					{
+						return SNew(STextBlock)
+							.Text(Item.IsValid() ? FText::FromName(*Item) : FText())
+							.Font(StructCustomizationUtils.GetRegularFont());
+					}
+				)
+				.OnSelectionChanged_Lambda(
+					[&InputAnim, EditorModel](TSharedPtr<FName> Item, ESelectInfo::Type)
+					{
+						const auto& VertexAttribNames = EditorModel->GetVertexAttributeNames();
+						if (!VertexAttribNames.IsEmpty() && Item == VertexAttribNames[0])
+						{
+							InputAnim.SetVertexMask(FName());
+						}
+						else
+						{
+							InputAnim.SetVertexMask(Item.IsValid() ? *Item : FName("<unknown>"));
+						}
+					}
+				)
+				.Content()
+				[
+					SNew(STextBlock)
+					.Text_Lambda(
+						[&InputAnim]()
+						{
+							return !InputAnim.GetVertexMask().IsNone() ? FText::FromName(InputAnim.GetVertexMask()) : LOCTEXT("DisabledString", "Disabled");
+						}
+					)
+					.ColorAndOpacity_Lambda(
+						[&InputAnim, EditorModel]()
+						{
+							const FName SelectedMask = InputAnim.GetVertexMask();
+							const TSharedPtr<FName>* ItemFound = EditorModel->GetVertexAttributeNames().FindByPredicate(
+								[&SelectedMask](const TSharedPtr<FName>& Item)
+								{
+									return (Item.IsValid() && *Item == SelectedMask);				
+								});
+
+							return ItemFound || SelectedMask.IsNone() ? FSlateColor::UseForeground() : FLinearColor::Red;
+						}
+					)
+					.Font(StructCustomizationUtils.GetRegularFont())
+				]
+			];
+		}
+
 		ChildBuilder.AddProperty(EnabledHandle.ToSharedRef());
 
 		// Verify the geometry cache against the anim sequence.
diff --git a/Engine/Plugins/Animation/MLDeformer/MLDeformerFramework/Source/MLDeformerFrameworkEditor/Private/MLDeformerGeomCacheTrainingModel.cpp b/Engine/Plugins/Animation/MLDeformer/MLDeformerFramework/Source/MLDeformerFrameworkEditor/Private/MLDeformerGeomCacheTrainingModel.cpp
--- a/Engine/Plugins/Animation/MLDeformer/MLDeformerFramework/Source/MLDeformerFrameworkEditor/Private/MLDeformerGeomCacheTrainingModel.cpp
+++ b/Engine/Plugins/Animation/MLDeformer/MLDeformerFramework/Source/MLDeformerFrameworkEditor/Private/MLDeformerGeomCacheTrainingModel.cpp
@@ -74,6 +74,7 @@ bool UMLDeformerGeomCacheTrainingModel::SampleNextFrame()
 	FMLDeformerSampler* Sampler = EditorModel->GetSamplerForTrainingAnim(SampleAnimIndex);
 	Sampler->SetVertexDeltaSpace(EVertexDeltaSpace::PreSkinning);
 	Sampler->Sample(CurFrameToSample);
+	MaskIndexPerSample.Add(GetMaskIndexForAnimIndex(SampleAnimIndex));
 	UE_LOG(LogMLDeformer, Display, TEXT("Sampling frame %d of anim %d"), CurFrameToSample, SampleAnimIndex);
 
 	// Copy sampled values.
diff --git a/Engine/Plugins/Animation/MLDeformer/MLDeformerFramework/Source/MLDeformerFrameworkEditor/Private/MLDeformerTrainingModel.cpp b/Engine/Plugins/Animation/MLDeformer/MLDeformerFramework/Source/MLDeformerFrameworkEditor/Private/MLDeformerTrainingModel.cpp
--- a/Engine/Plugins/Animation/MLDeformer/MLDeformerFramework/Source/MLDeformerFrameworkEditor/Private/MLDeformerTrainingModel.cpp
+++ b/Engine/Plugins/Animation/MLDeformer/MLDeformerFramework/Source/MLDeformerFrameworkEditor/Private/MLDeformerTrainingModel.cpp
@@ -5,6 +5,11 @@
 #include "MLDeformerModel.h"
 #include "MLDeformerInputInfo.h"
 #include "MLDeformerSampler.h"
+#include "MLDeformerTrainingInputAnim.h"
+#include "SkeletalMeshAttributes.h"
+#include "Engine/SkeletalMesh.h"
+
+FName UMLDeformerTrainingModel::DefaultMaskName = FName("MLD_DefaultMask");
 
 UMLDeformerModel* UMLDeformerTrainingModel::GetModel() const
 {
@@ -47,7 +52,9 @@ void UMLDeformerTrainingModel::ResetSampling()
 	NumTimesSampled.Reset();
 	NumTimesSampled.AddZeroed(EditorModel->GetNumTrainingInputAnims());
 	SampleAnimIndex = 0;
+	MaskIndexPerSample.Reset();
 	bFinishedSampling = !FindNextAnimToSample(SampleAnimIndex);
+	MaskNames = GetTrainingInputAnimMasks();
 }
 
 int32 UMLDeformerTrainingModel::GetNumberSampleDeltas() const
@@ -68,6 +75,88 @@ void UMLDeformerTrainingModel::SetNumFloatsPerCurve(int32 NumFloatsPerCurve)
 	}
 }
 
+TArray<FName> UMLDeformerTrainingModel::GetTrainingInputAnimMasks() const
+{
+	TArray<FName> ValidMasks;
+
+	ValidMasks.Add(DefaultMaskName);
+
+	const int32 NumAnims = EditorModel->GetNumTrainingInputAnims();
+	for (int32 AnimIndex = 0; AnimIndex < NumAnims; ++AnimIndex)
+	{
+		const FMLDeformerTrainingInputAnim* Anim = EditorModel->GetTrainingInputAnim(AnimIndex);
+		const FName MaskName = Anim->GetVertexMask();
+		if (MaskName.IsNone() || !Anim->IsEnabled())
+		{
+			continue;
+		}
+
+		TVertexAttributesConstRef<float> MaskData = EditorModel->FindVertexAttributes(MaskName);
+		if (MaskData.IsValid())
+		{
+			ValidMasks.Add(MaskName);
+		}
+	}
+
+	return MoveTemp(ValidMasks);
+}
+
+TArray<float> UMLDeformerTrainingModel::GetTrainingInputAnimMaskData(FName MaskName) const
+{
+	TArray<float> PerVertexValues;
+
+	// Get the imported vertex numbers.
+	const int32 LodIndex = 0;
+	USkeletalMesh* SkeletalMesh = EditorModel->GetModel()->GetSkeletalMesh();
+	FMeshDescription* MeshDescription = SkeletalMesh->GetMeshDescription(LodIndex);
+	TVertexAttributesConstRef<int32> ImportPointIndex = MeshDescription ? MeshDescription->VertexAttributes().GetAttributesRef<int32>(MeshAttribute::Vertex::ImportPointIndex) : TVertexAttributesConstRef<int32>(); 
+
+	TVertexAttributesConstRef<float> MaskData = EditorModel->FindVertexAttributes(MaskName);
+	if (ImportPointIndex.IsValid() && MaskData.IsValid())
+	{
+		const int32 NumModelVerts = EditorModel->GetModel()->GetNumBaseMeshVerts();
+		PerVertexValues.SetNumZeroed(NumModelVerts);
+
+		check(ImportPointIndex.GetNumElements() == MaskData.GetNumElements());
+		for (int32 Index = 0; Index < MaskData.GetNumElements(); ++Index)
+		{
+			PerVertexValues[ImportPointIndex[Index]] = MaskData.Get(Index);
+		}
+	}
+	else
+	{
+		if (MaskName == DefaultMaskName)
+		{
+			const int32 NumModelVerts = EditorModel->GetModel()->GetNumBaseMeshVerts();
+			PerVertexValues.SetNumUninitialized(NumModelVerts);		
+			for (int32 Index = 0; Index < PerVertexValues.Num(); ++Index)
+			{
+				PerVertexValues[Index] = 1.0f;
+			}
+		}
+	}
+
+	return MoveTemp(PerVertexValues);
+}
+
+int32 UMLDeformerTrainingModel::GetMaskIndexForAnimIndex(int32 AnimIndex) const
+{
+	const FMLDeformerTrainingInputAnim* Anim = EditorModel->GetTrainingInputAnim(AnimIndex);
+	const FName AnimMaskName = Anim->GetVertexMask();
+	if (AnimMaskName.IsNone())
+	{
+		return 0;	
+	}
+
+	const int32 MaskIndex = MaskNames.Find(AnimMaskName);
+	if (MaskIndex == INDEX_NONE)
+	{
+		return 0;
+	}
+
+	return MaskIndex;
+}
+
 bool UMLDeformerTrainingModel::SetCurrentSampleIndex(int32 Index)
 {
 	return NextSample();
diff --git a/Engine/Plugins/Animation/MLDeformer/MLDeformerFramework/Source/MLDeformerFrameworkEditor/Public/MLDeformerEditorModel.h b/Engine/Plugins/Animation/MLDeformer/MLDeformerFramework/Source/MLDeformerFrameworkEditor/Public/MLDeformerEditorModel.h
--- a/Engine/Plugins/Animation/MLDeformer/MLDeformerFramework/Source/MLDeformerFrameworkEditor/Public/MLDeformerEditorModel.h
+++ b/Engine/Plugins/Animation/MLDeformer/MLDeformerFramework/Source/MLDeformerFrameworkEditor/Public/MLDeformerEditorModel.h
@@ -722,6 +722,13 @@ namespace UE::MLDeformer
 		 */
 		virtual void DrawPIEDebugActors();
 
+		/**
+		 * Check whether this model supports a per training input animation vertex mask.
+		 * If this is enabled, each training intput animation should show something where the user can select an optional mask.
+		 * @return Returns true when vertex masking per training input is supported, false otherwise.
+		 */
+		virtual bool GetSupportsPerTrainingInputAnimVertexMask() const	{ return false; }
+
 		/** 
 		 * Update LOD levels of the actors in editor world.
 		 * This can be used to for example sync the LOD levels of the compare actors with the main actor.
@@ -987,7 +994,13 @@ namespace UE::MLDeformer
 		TVertexAttributesConstRef<float> FindVertexAttributes(FName AttributeName) const;
 
 		/** Are we scrubbing on the timeline? */
-		bool IsScrubbingTimeline() const	{ return bIsScrubbingTimeline; }
+		bool IsScrubbingTimeline() const									{ return bIsScrubbingTimeline; }
+
+		/** Get the array of vertex attribute names that can be used for UI widget elements. Refresh this list with UpdateVertexAttributeNames(). */
+		const TArray<TSharedPtr<FName>>& GetVertexAttributeNames() const	{ return VertexAttributeNames; }
+
+		/** Update the vertex attribute names array as returned by GetVertexAttributeNames(). This is some array that can be used for UI elements like combo boxes. */
+		void UpdateVertexAttributeNames();
 
 
 	protected:
@@ -1167,6 +1180,9 @@ namespace UE::MLDeformer
 		/** The heatmap deformer graph. */
 		TObjectPtr<UMeshDeformer> HeatMapDeformerGraph = nullptr;
 
+		/** A list of vertex attribute names that can be used inside combobox UI elements. */
+		TArray<TSharedPtr<FName>> VertexAttributeNames;
+
 		/** The delegate handle to the post edit property event. */
 		FDelegateHandle PostEditPropertyDelegateHandle;
 
diff --git a/Engine/Plugins/Animation/MLDeformer/MLDeformerFramework/Source/MLDeformerFrameworkEditor/Public/MLDeformerTrainingModel.h b/Engine/Plugins/Animation/MLDeformer/MLDeformerFramework/Source/MLDeformerFrameworkEditor/Public/MLDeformerTrainingModel.h
--- a/Engine/Plugins/Animation/MLDeformer/MLDeformerFramework/Source/MLDeformerFrameworkEditor/Public/MLDeformerTrainingModel.h
+++ b/Engine/Plugins/Animation/MLDeformer/MLDeformerFramework/Source/MLDeformerFrameworkEditor/Public/MLDeformerTrainingModel.h
@@ -112,6 +112,20 @@ public:
 	UFUNCTION(BlueprintCallable, Category = "Training Data")
 	void SetNumFloatsPerCurve(int32 NumFloatsPerCurve);
 
+	/** Get the names of all valid training animation masks. They are valid when they are used, the anim is enabled, and the attribute exists on the skeletal mesh. */
+	UFUNCTION(BlueprintPure, Category = "Training Data")
+    TArray<FName> GetTrainingInputAnimMasks() const;
+
+	/** Get the per vertex data for a given mask name. The mask name must be one that is present in the array returned by GetTrainingInputAnimMasks. */
+	UFUNCTION(BlueprintPure, Category = "Training Data")
+	TArray<float> GetTrainingInputAnimMaskData(FName MaskName) const;
+
+	/** 
+	 * For each sample we took, this contains the index inside the mask as returned by the GetTrainingInputAnimMasks() array.
+	 */
+	UFUNCTION(BlueprintPure, Category = "Training Data")
+	TArray<int32> GetMaskIndexPerSampleArray() const	{ return MaskIndexPerSample; }
+
 protected:
 	/** This updates the sample deltas, curves, and bone rotations. */
 	UE_DEPRECATED(5.4, "Please use UMLDeformerTrainingModel::SampleNextFrame() instead.")
@@ -135,6 +149,12 @@ protected:
 	 */
 	virtual bool FindNextAnimToSample(int32& OutNextAnimIndex) const { return false; }
 
+	/**
+	 * Get the mask index for a given input training animation.
+	 * @see GetTrainingInputAnimMasks.
+	 */
+	int32 GetMaskIndexForAnimIndex(int32 AnimIndex) const;
+
 public:
 	// The delta values per vertex for this sample. This is updated after SetCurrentSampleIndex is called. Contains an xyz (3 floats) for each vertex.
 	UPROPERTY(EditAnywhere, BlueprintReadWrite, Category = "Training Data")
@@ -148,6 +168,11 @@ public:
 	UPROPERTY(EditAnywhere, BlueprintReadWrite, Category = "Training Data")
 	TArray<float> SampleBoneRotations;
 
+	// The mask index for each sample.
+	// See GetTrainingInputAnimMasks() and GetTrainingInputAnimMaskData().
+	UPROPERTY(EditAnywhere, BlueprintReadWrite, Category = "Training Data")
+	TArray<int32> MaskIndexPerSample;
+
 protected:
 	/** A pointer to the editor model from which we use the sampler. */
 	UE::MLDeformer::FMLDeformerEditorModel* EditorModel = nullptr;
@@ -164,4 +189,9 @@ protected:
 
 	/** Did we finish sampling? This is set to true when every possible frame has been sampled. */
 	bool bFinishedSampling = false;
+
+	TArray<FName> MaskNames;
+
+	/** The default mask name. */
+	static FName DefaultMaskName;
 };
diff --git a/Engine/Plugins/Animation/MLDeformer/NeuralMorphModel/Source/NeuralMorphModel/Public/NeuralMorphModel.h b/Engine/Plugins/Animation/MLDeformer/NeuralMorphModel/Source/NeuralMorphModel/Public/NeuralMorphModel.h
--- a/Engine/Plugins/Animation/MLDeformer/NeuralMorphModel/Source/NeuralMorphModel/Public/NeuralMorphModel.h
+++ b/Engine/Plugins/Animation/MLDeformer/NeuralMorphModel/Source/NeuralMorphModel/Public/NeuralMorphModel.h
@@ -190,7 +190,7 @@ public:
 	 * A value of 0 disables the regularization, and gives the highest quality, at the cost of higher runtime memory usage.
 	 */
 	UPROPERTY(EditAnywhere, BlueprintReadWrite, AdvancedDisplay, Category = "Training Settings", meta = (ClampMin = "0.0", ClampMax = "10.0"))
-	float RegularizationFactor = 1.0f;
+	float RegularizationFactor = 0.0f;
 
 	/** 
 	 * Enable the use of per bone and bone group masks.
@@ -208,7 +208,7 @@ public:
 	 * If you see some noise in the trained results, even with large amount of samples and iterations, try increasing this value.
 	 */
 	UPROPERTY(EditAnywhere, BlueprintReadWrite, AdvancedDisplay, Category = "Training Settings", meta = (ClampMin = "0.0", ClampMax = "10.0"))
-	float SmoothLossBeta = 1.0f;
+	float SmoothLossBeta = 0.0f;
 
 	/**
 	 * The neural morph model network.
diff --git a/Engine/Plugins/Animation/MLDeformer/NeuralMorphModel/Source/NeuralMorphModelEditor/Public/NeuralMorphEditorModel.h b/Engine/Plugins/Animation/MLDeformer/NeuralMorphModel/Source/NeuralMorphModelEditor/Public/NeuralMorphEditorModel.h
--- a/Engine/Plugins/Animation/MLDeformer/NeuralMorphModel/Source/NeuralMorphModelEditor/Public/NeuralMorphEditorModel.h
+++ b/Engine/Plugins/Animation/MLDeformer/NeuralMorphModel/Source/NeuralMorphModelEditor/Public/NeuralMorphEditorModel.h
@@ -42,6 +42,7 @@ namespace UE::NeuralMorphModel
 		virtual TSharedPtr<SMLDeformerInputWidget> CreateInputWidget() override;
 		virtual void OnPostInputAssetChanged() override;
 		virtual void Render(const FSceneView* View, FViewport* Viewport, FPrimitiveDrawInterface* PDI) override;
+		virtual bool GetSupportsPerTrainingInputAnimVertexMask() const override	{ return true; }
 		// ~END FMLDeformerEditorModel overrides.
 
 		// FMLDeformerMorphModelEditorModel overrides.
