[MLDeformer]
- Enable tests for the editor framework, updated all tests to latest API etc.
- Cooking improvements, where stripped data is recovered after cook.
- Ability to select a training device. This allows you to pick a specific GPU to train on, or choose CPU.
- Added some extra Cuda related logging, to help debugging issues.
- Added checks to validate training data, checks for NaN and Inf etc in NMM.
- Add another return code from the training, and improve error logging in NMM (show call stack etc).
- Some other small typo fixes etc.

diff --git a/Engine/Plugins/Animation/MLDeformer/MLDeformerFramework/Source/MLDeformerFramework/Private/MLDeformerGeomCacheModel.cpp b/Engine/Plugins/Animation/MLDeformer/MLDeformerFramework/Source/MLDeformerFramework/Private/MLDeformerGeomCacheModel.cpp
--- a/Engine/Plugins/Animation/MLDeformer/MLDeformerFramework/Source/MLDeformerFramework/Private/MLDeformerGeomCacheModel.cpp
+++ b/Engine/Plugins/Animation/MLDeformer/MLDeformerFramework/Source/MLDeformerFramework/Private/MLDeformerGeomCacheModel.cpp
@@ -16,21 +16,61 @@
 
 void UMLDeformerGeomCacheModel::Serialize(FArchive& Archive)
 {
-	Archive.UsingCustomVersion(UE::MLDeformer::FMLDeformerObjectVersion::GUID);
-
 	#if WITH_EDITOR
-		if (Archive.IsSaving() && Archive.IsCooking())
+		bool bModifiedPropertiesForCook = false;
+		TArray<UGeometryCache*> SavedGeomCaches;
+		TArray<UAnimSequence*> SavedAnimSequences;
+		auto ModifyPropertiesForCook = [this, &bModifiedPropertiesForCook, &SavedGeomCaches, &SavedAnimSequences]()
 		{
+			bModifiedPropertiesForCook = true;
 			GeometryCache_DEPRECATED = nullptr;
+			SavedAnimSequences.Reset();
+			SavedAnimSequences.Reserve(TrainingInputAnims.Num());
+			SavedGeomCaches.Reset();
+			SavedGeomCaches.Reserve(TrainingInputAnims.Num());
 			for (FMLDeformerGeomCacheTrainingInputAnim& Anim : TrainingInputAnims)
 			{
+				SavedGeomCaches.Add(Anim.GetGeometryCache());
+				SavedAnimSequences.Add(Anim.GetAnimSequence());
 				Anim.SetGeometryCache(nullptr);
 				Anim.SetAnimSequence(nullptr);
 			}
+		};
+
+		auto RestorePropertiesForCook = [this, &bModifiedPropertiesForCook, &SavedGeomCaches, &SavedAnimSequences]()
+		{
+			if (!bModifiedPropertiesForCook)
+			{
+				return;
+			}
+			check(TrainingInputAnims.Num() == SavedGeomCaches.Num());
+			check(TrainingInputAnims.Num() == SavedAnimSequences.Num());
+			for (int32 Index = 0; Index < TrainingInputAnims.Num(); ++Index)			
+			{
+				FMLDeformerGeomCacheTrainingInputAnim& Anim = TrainingInputAnims[Index];
+				Anim.SetGeometryCache(SavedGeomCaches[Index]);
+				Anim.SetAnimSequence(SavedAnimSequences[Index]);
+			}		
+		};
+
+		ON_SCOPE_EXIT
+		{
+			if (GetRecoverStrippedDataAfterCook())
+			{
+				RestorePropertiesForCook();
+			}
+		};
+	#endif
+
+	#if WITH_EDITOR
+		if (Archive.IsSaving() && Archive.IsCooking())
+		{
+			ModifyPropertiesForCook();
 		}
 	#endif
 
 	Super::Serialize(Archive);
+	Archive.UsingCustomVersion(UE::MLDeformer::FMLDeformerObjectVersion::GUID);
 }
 
 void UMLDeformerGeomCacheModel::PostLoad()
diff --git a/Engine/Plugins/Animation/MLDeformer/MLDeformerFramework/Source/MLDeformerFramework/Private/MLDeformerModel.cpp b/Engine/Plugins/Animation/MLDeformer/MLDeformerFramework/Source/MLDeformerFramework/Private/MLDeformerModel.cpp
--- a/Engine/Plugins/Animation/MLDeformer/MLDeformerFramework/Source/MLDeformerFramework/Private/MLDeformerModel.cpp
+++ b/Engine/Plugins/Animation/MLDeformer/MLDeformerFramework/Source/MLDeformerFramework/Private/MLDeformerModel.cpp
@@ -74,17 +74,37 @@ void UMLDeformerModel::Init(UMLDeformerAsset* InDeformerAsset)
 
 void UMLDeformerModel::Serialize(FArchive& Archive)
 {
-	Archive.UsingCustomVersion(UE::MLDeformer::FMLDeformerObjectVersion::GUID);
+	UMLDeformerVizSettings* VizSettingsBackup = nullptr;
+	bool bProcessedDataOnCook = false;
+	ON_SCOPE_EXIT
+	{
+		if (bProcessedDataOnCook && GetRecoverStrippedDataAfterCook())
+		{
+			#if WITH_EDITORONLY_DATA
+				VizSettings = VizSettingsBackup;
+				VizSettingsBackup = nullptr;
+			#endif
+		}
+	};
 
 	#if WITH_EDITOR
-		if (Archive.IsSaving())
+		if (Archive.IsSaving() && Archive.IsCooking())
 		{
-			if (Archive.IsCooking())
-			{
+			bProcessedDataOnCook = true;
+			#if WITH_EDITORONLY_DATA
+				VizSettingsBackup = VizSettings;
 				AnimSequence_DEPRECATED = nullptr;
 				VizSettings = nullptr;
-			}
+			#endif
+		}
+	#endif
+
+	Super::Serialize(Archive);
+	Archive.UsingCustomVersion(UE::MLDeformer::FMLDeformerObjectVersion::GUID);
 
+	#if WITH_EDITOR
+		if (Archive.IsSaving())
+		{
 			InitVertexMap();
 		}
 
@@ -94,8 +114,6 @@ void UMLDeformerModel::Serialize(FArchive& Archive)
 			UpdateCachedNumVertices();
 		}
 	#endif
-
-	Super::Serialize(Archive);
 }
 
 UMLDeformerAsset* UMLDeformerModel::GetDeformerAsset() const
@@ -266,6 +284,34 @@ bool UMLDeformerModel::IsCompatibleDebugActor(const AActor* Actor, UMLDeformerCo
 }
 
 
+void UMLDeformerModel::SetRecoverStrippedDataAfterCook(bool bRecover)
+{ 
+#if WITH_EDITORONLY_DATA
+	bRecoverStrippedDataAfterCook = bRecover;
+#endif
+}
+
+bool UMLDeformerModel::GetRecoverStrippedDataAfterCook() const
+{ 
+	#if WITH_EDITORONLY_DATA
+		return bRecoverStrippedDataAfterCook;
+	#else
+		return true;
+	#endif
+}
+
+void UMLDeformerModel::SetTrainingDevice(const FString& DeviceName)
+{
+	if (!TrainingDeviceList.Contains(DeviceName))
+	{
+		UE_LOG(LogMLDeformer, Warning, TEXT("Training device '%s' not found, falling back to CPU."), *DeviceName);
+		TrainingDevice.Empty();
+		return;
+	}
+
+	TrainingDevice = DeviceName;
+}
+
 #if WITH_EDITOR
 	TArray<FName> UMLDeformerModel::GetVertexAttributeNames() const
 	{
diff --git a/Engine/Plugins/Animation/MLDeformer/MLDeformerFramework/Source/MLDeformerFramework/Private/MLDeformerMorphModel.cpp b/Engine/Plugins/Animation/MLDeformer/MLDeformerFramework/Source/MLDeformerFramework/Private/MLDeformerMorphModel.cpp
--- a/Engine/Plugins/Animation/MLDeformer/MLDeformerFramework/Source/MLDeformerFramework/Private/MLDeformerMorphModel.cpp
+++ b/Engine/Plugins/Animation/MLDeformer/MLDeformerFramework/Source/MLDeformerFramework/Private/MLDeformerMorphModel.cpp
@@ -29,22 +29,27 @@ void UMLDeformerMorphModel::Serialize(FArchive& Archive)
 	bool bModifiedPropertiesForCook = false;
 	TArray<FVector3f> SavedMorphTargetDeltas;
 	auto ModifyPropertiesForCook = [this, &bModifiedPropertiesForCook, &SavedMorphTargetDeltas]()
-		{
-			bModifiedPropertiesForCook = true;
-			SavedMorphTargetDeltas = MoveTemp(MorphTargetDeltas);
-			MorphTargetDeltas.Empty();
-		};
+	{
+		bModifiedPropertiesForCook = true;
+		SavedMorphTargetDeltas = MoveTemp(MorphTargetDeltas);
+		MorphTargetDeltas.Empty();
+	};
+
 	auto RestorePropertiesForCook = [this, &bModifiedPropertiesForCook, &SavedMorphTargetDeltas]()
+	{
+		if (!bModifiedPropertiesForCook)
 		{
-			if (!bModifiedPropertiesForCook)
-			{
-				return;
-			}
-			MorphTargetDeltas = MoveTemp(SavedMorphTargetDeltas);
-		};
+			return;
+		}
+		MorphTargetDeltas = MoveTemp(SavedMorphTargetDeltas);
+	};
+
 	ON_SCOPE_EXIT
 	{
-		RestorePropertiesForCook();
+		if (GetRecoverStrippedDataAfterCook())
+		{
+			RestorePropertiesForCook();
+		}
 	};
 
 	int32 NumSaveLODs = 0;
diff --git a/Engine/Plugins/Animation/MLDeformer/MLDeformerFramework/Source/MLDeformerFramework/Private/MLDeformerMorphModelInputInfo.cpp b/Engine/Plugins/Animation/MLDeformer/MLDeformerFramework/Source/MLDeformerFramework/Private/MLDeformerMorphModelInputInfo.cpp
--- a/Engine/Plugins/Animation/MLDeformer/MLDeformerFramework/Source/MLDeformerFramework/Private/MLDeformerMorphModelInputInfo.cpp
+++ b/Engine/Plugins/Animation/MLDeformer/MLDeformerFramework/Source/MLDeformerFramework/Private/MLDeformerMorphModelInputInfo.cpp
@@ -6,15 +6,28 @@
 
 void UMLDeformerMorphModelInputInfo::Serialize(FArchive& Archive)
 {
-	Super::Serialize(Archive);
-
 	// Strip editor only data.
 #if WITH_EDITORONLY_DATA
+	TArray<float> InputItemMaskBufferBackup;
+	bool bProcessedDataOnCook = false;
+	ON_SCOPE_EXIT
+	{
+		UMLDeformerModel* Model = Cast<UMLDeformerModel>(GetOuter());
+		if (bProcessedDataOnCook && Model && Model->GetRecoverStrippedDataAfterCook())
+		{
+			InputItemMaskBuffer = MoveTemp(InputItemMaskBufferBackup);
+		}
+	};
+
 	if (Archive.IsSaving() && Archive.IsCooking())
 	{
+		InputItemMaskBufferBackup = MoveTemp(InputItemMaskBuffer);
 		InputItemMaskBuffer.Empty();
+		bProcessedDataOnCook = true;
 	}
 #endif
+
+	Super::Serialize(Archive);
 }
 
 #if WITH_EDITORONLY_DATA
diff --git a/Engine/Plugins/Animation/MLDeformer/MLDeformerFramework/Source/MLDeformerFramework/Public/MLDeformerModel.h b/Engine/Plugins/Animation/MLDeformer/MLDeformerFramework/Source/MLDeformerFramework/Public/MLDeformerModel.h
--- a/Engine/Plugins/Animation/MLDeformer/MLDeformerFramework/Source/MLDeformerFramework/Public/MLDeformerModel.h
+++ b/Engine/Plugins/Animation/MLDeformer/MLDeformerFramework/Source/MLDeformerFramework/Public/MLDeformerModel.h
@@ -409,12 +409,36 @@ public:
 	FNeuralNetworkModifyDelegate& GetNeuralNetworkModifyDelegate()				{ return NeuralNetworkModifyDelegate_DEPRECATED; }
 
 	/**
-	 * Get the delegate which will be called when to inform when the model instance needs to be reinitailized.
+	 * Get the delegate which will be called when to inform when the model instance needs to be reinitialized.
 	 * The ML Deformer component will hook into this for example. This is broadcast when something changes in the structure, such as when
 	 * training finished and the neural network changed.
 	 */
 	FMLDeformerReinitModelInstancesDelegate& GetReinitModelInstanceDelegate()	{ return ReinitModelInstanceDelegate; }
 
+	/**
+	 * Specify whether we want to recover stripped data that is removed from this model when cooking. Examples of this data could be the raw uncompressed vertex deltas of all morph targets.
+	 * On default we strip the data, save the cooked asset, and then recover the stripped data again. But using this method you can disable that, by setting it to false.
+	 * Generally we only want to disable this during some automated tests.
+	 */
+	void SetRecoverStrippedDataAfterCook(bool bRecover);
+
+	/**
+	 * Check whether we want to recover stripped data that is removed from this model when cooking. Examples of this data could be the raw uncompressed vertex deltas of all morph targets.
+	 * On default we strip the data, save the cooked asset, and then recover the stripped data again. But some automated cooking tests might want this disabled.
+	 */
+	bool GetRecoverStrippedDataAfterCook() const;
+
+	UFUNCTION()
+	const TArray<FString>& GetTrainingDeviceList() const			{ return TrainingDeviceList; }
+
+	UFUNCTION(BlueprintCallable, Category = "Training")
+	const FString& GetTrainingDevice() const						{ return TrainingDevice; }
+
+	UFUNCTION()
+	void SetTrainingDevice(const FString& DeviceName);
+	void SetTrainingDeviceToCpu()									{ TrainingDevice = "Cpu"; }
+	void SetTrainingDeviceList(const TArray<FString>& Devices)		{ TrainingDeviceList = Devices; }
+
 
 #if WITH_EDITORONLY_DATA
 	// UObject overrides.
@@ -431,7 +455,7 @@ public:
 
 	/**
 	 * Check whether we should include bone transforms as input to the model during training or not.
-	 * @return Returns true when bone transfomations should be a part of the network inputs, during the training process.
+	 * @return Returns true when bone transformations should be a part of the network inputs, during the training process.
 	 */
 	UE_DEPRECATED(5.3, "This method and property has been removed and shouldn't be used anymore.")
 	bool ShouldIncludeBonesInTraining() const					{ PRAGMA_DISABLE_DEPRECATION_WARNINGS; return bIncludeBones_DEPRECATED; PRAGMA_ENABLE_DEPRECATION_WARNINGS; }
@@ -510,7 +534,7 @@ public:
 	void SetTrainingFrameLimit(int32 MaxNumFrames)				{ MaxTrainingFrames = MaxNumFrames; }
 
 	/**
-	 * Get the target mesh alignment tranformation.
+	 * Get the target mesh alignment transformation.
 	 * This is a transformation that is applied to the vertex positions of the target mesh, before we calculate the deltas
 	 * between the linear skinned mesh and the target mesh.
 	 * This is useful when you imported target mesh that isn't scaled the same, or perhaps it is rotated 90 degrees over the x axis.
@@ -603,6 +627,7 @@ public:
 	static FName GetCurveIncludeListPropertyName()		{ return GET_MEMBER_NAME_CHECKED(UMLDeformerModel, CurveIncludeList); }
 	static FName GetMaxTrainingFramesPropertyName()		{ return GET_MEMBER_NAME_CHECKED(UMLDeformerModel, MaxTrainingFrames); }
 	static FName GetMaxNumLODsPropertyName()			{ return GET_MEMBER_NAME_CHECKED(UMLDeformerModel, MaxNumLODs); }
+	static FName GetTrainingDevicePropertyName()		{ return GET_MEMBER_NAME_CHECKED(UMLDeformerModel, TrainingDevice); }
 
 	UE_DEPRECATED(5.3, "This property has been removed and shouldn't be used anymore.")
 	static FName GetShouldIncludeBonesPropertyName()	{ PRAGMA_DISABLE_DEPRECATION_WARNINGS; return GET_MEMBER_NAME_CHECKED(UMLDeformerModel, bIncludeBones_DEPRECATED); PRAGMA_ENABLE_DEPRECATION_WARNINGS; }
@@ -704,6 +729,13 @@ private:
 	UPROPERTY()
 	int32 NumTargetMeshVerts = 0;
 
+	/** The device used for training. On default it will init to Cuda's preferred device, or Cpu if no such device present. */
+	UPROPERTY(EditAnywhere, AdvancedDisplay, Category = "Training Settings", meta = (GetOptions = "GetTrainingDeviceList", NoResetToDefault))
+	FString TrainingDevice;
+
+	/** The list of training devices that will show in the combo box in the UI. */
+	TArray<FString> TrainingDeviceList;
+
 	/** 
 	 * How many Skeletal Mesh LOD levels should we generate MLD lods for at most?
 	 * Some examples:
@@ -777,5 +809,11 @@ private:
 	UE_DEPRECATED(5.5, "This property has been removed.")
 	UPROPERTY(meta = (DeprecatedProperty, DeprecationMessage = "This property has been removed and isn't used anymore."))
 	float DeltaCutoffLength_DEPRECATED = 30.0f;
+	
+	/** 
+	 * Do we want to recover stripped data at the end of the Serialize call during cook?
+	 * This is enabled on default, but some automated tests might disable this.
+	 */
+	bool bRecoverStrippedDataAfterCook = true;
 #endif
 };
diff --git a/Engine/Plugins/Animation/MLDeformer/MLDeformerFramework/Source/MLDeformerFrameworkEditor/Private/MLDeformerAssetFactory.h b/Engine/Plugins/Animation/MLDeformer/MLDeformerFramework/Source/MLDeformerFrameworkEditor/Private/MLDeformerAssetFactory.h
deleted file mode 100644
diff --git a/Engine/Plugins/Animation/MLDeformer/MLDeformerFramework/Source/MLDeformerFrameworkEditor/Private/MLDeformerEditorModel.cpp b/Engine/Plugins/Animation/MLDeformer/MLDeformerFramework/Source/MLDeformerFrameworkEditor/Private/MLDeformerEditorModel.cpp
--- a/Engine/Plugins/Animation/MLDeformer/MLDeformerFramework/Source/MLDeformerFrameworkEditor/Private/MLDeformerEditorModel.cpp
+++ b/Engine/Plugins/Animation/MLDeformer/MLDeformerFramework/Source/MLDeformerFrameworkEditor/Private/MLDeformerEditorModel.cpp
@@ -125,6 +125,7 @@ namespace UE::MLDeformer
 		);
 
 		UpdateLODMappings();
+		UpdateTrainingDeviceList();
 	}
 
 	void FMLDeformerEditorModel::UpdateVertexAttributeNames()
diff --git a/Engine/Plugins/Animation/MLDeformer/MLDeformerFramework/Source/MLDeformerFrameworkEditor/Private/MLDeformerEditorToolkit.cpp b/Engine/Plugins/Animation/MLDeformer/MLDeformerFramework/Source/MLDeformerFrameworkEditor/Private/MLDeformerEditorToolkit.cpp
--- a/Engine/Plugins/Animation/MLDeformer/MLDeformerFramework/Source/MLDeformerFrameworkEditor/Private/MLDeformerEditorToolkit.cpp
+++ b/Engine/Plugins/Animation/MLDeformer/MLDeformerFramework/Source/MLDeformerFrameworkEditor/Private/MLDeformerEditorToolkit.cpp
@@ -73,7 +73,10 @@ namespace UE::MLDeformer
 			FMLDeformerEditorModel::InitSettings InitSettings;
 			InitSettings.Editor = this;
 			InitSettings.Model = DeformerModel;
-			EditorModel->Init(InitSettings);
+			if (EditorModel)
+			{
+				EditorModel->Init(InitSettings);
+			}
 			ActiveModel = TSharedPtr<FMLDeformerEditorModel>(EditorModel);
 		}
 
diff --git a/Engine/Plugins/Animation/MLDeformer/MLDeformerFramework/Source/MLDeformerFrameworkEditor/Private/MLDeformerModelDetails.cpp b/Engine/Plugins/Animation/MLDeformer/MLDeformerFramework/Source/MLDeformerFrameworkEditor/Private/MLDeformerModelDetails.cpp
--- a/Engine/Plugins/Animation/MLDeformer/MLDeformerFramework/Source/MLDeformerFrameworkEditor/Private/MLDeformerModelDetails.cpp
+++ b/Engine/Plugins/Animation/MLDeformer/MLDeformerFramework/Source/MLDeformerFrameworkEditor/Private/MLDeformerModelDetails.cpp
@@ -191,6 +191,12 @@ namespace UE::MLDeformer
 			}
 		}
 
+		TrainingSettingsCategoryBuilder->AddProperty(UMLDeformerModel::GetTrainingDevicePropertyName(), UMLDeformerModel::StaticClass())
+			.Visibility(TAttribute<EVisibility>::CreateLambda(
+				[this]()
+				{
+					return Model->GetTrainingDeviceList().IsEmpty() ? EVisibility::Collapsed : EVisibility::Visible;
+				}));
 		TrainingSettingsCategoryBuilder->AddProperty(UMLDeformerModel::GetMaxTrainingFramesPropertyName(), UMLDeformerModel::StaticClass());
 		AddTrainingSettingsErrors();
 
diff --git a/Engine/Plugins/Animation/MLDeformer/MLDeformerFramework/Source/MLDeformerFrameworkEditor/Private/MLDeformerMorphModelDetails.cpp b/Engine/Plugins/Animation/MLDeformer/MLDeformerFramework/Source/MLDeformerFrameworkEditor/Private/MLDeformerMorphModelDetails.cpp
--- a/Engine/Plugins/Animation/MLDeformer/MLDeformerFramework/Source/MLDeformerFrameworkEditor/Private/MLDeformerMorphModelDetails.cpp
+++ b/Engine/Plugins/Animation/MLDeformer/MLDeformerFramework/Source/MLDeformerFrameworkEditor/Private/MLDeformerMorphModelDetails.cpp
@@ -42,7 +42,7 @@ namespace UE::MLDeformer
 
 	bool FMLDeformerMorphModelDetails::ShouldShowShadingError() const
 	{
-		if (!Model)
+		if (!Model || !EditorModel)
 		{
 			return false;
 		}
@@ -132,6 +132,11 @@ namespace UE::MLDeformer
 		// Create all the detail categories and add the properties of the base class.
 		FMLDeformerGeomCacheModelDetails::CustomizeDetails(DetailBuilder);
 
+		if (!EditorModel)
+		{
+			return;
+		}
+
 		if (ShouldShowShadingError())
 		{
 			FDetailWidgetRow& CannotRunMLDeformerWarning = BaseMeshCategoryBuilder->AddCustomRow(FText::FromString("CannotRunMLDeformerError"))
diff --git a/Engine/Plugins/Animation/MLDeformer/MLDeformerFramework/Source/MLDeformerFrameworkEditor/Private/MLDeformerTrainingModel.cpp b/Engine/Plugins/Animation/MLDeformer/MLDeformerFramework/Source/MLDeformerFrameworkEditor/Private/MLDeformerTrainingModel.cpp
--- a/Engine/Plugins/Animation/MLDeformer/MLDeformerFramework/Source/MLDeformerFrameworkEditor/Private/MLDeformerTrainingModel.cpp
+++ b/Engine/Plugins/Animation/MLDeformer/MLDeformerFramework/Source/MLDeformerFrameworkEditor/Private/MLDeformerTrainingModel.cpp
@@ -188,3 +188,15 @@ bool UMLDeformerTrainingModel::SampleFrame(int32 Index)
 	UE_LOG(LogMLDeformer, Warning, TEXT("Please use UMLDeformerTrainingModel::NextSample() instead."));
 	return false;
 }
+
+void UMLDeformerTrainingModel::SetDeviceList(const TArray<FString>& DeviceNames, int32 PreferredDeviceIndex)
+{
+	UMLDeformerModel* Model = EditorModel->GetModel();
+	Model->SetTrainingDeviceList(DeviceNames);
+
+	if (Model->GetTrainingDevice().IsEmpty() || !DeviceNames.Contains(Model->GetTrainingDevice()))
+	{
+		check(DeviceNames.IsValidIndex(PreferredDeviceIndex));
+		Model->SetTrainingDevice(DeviceNames[PreferredDeviceIndex]);
+	}
+}
diff --git a/Engine/Plugins/Animation/MLDeformer/MLDeformerFramework/Source/MLDeformerFrameworkEditor/Public/MLDeformerAssetFactory.h b/Engine/Plugins/Animation/MLDeformer/MLDeformerFramework/Source/MLDeformerFrameworkEditor/Public/MLDeformerAssetFactory.h
new file mode 100644
--- /dev/null
+++ b/Engine/Plugins/Animation/MLDeformer/MLDeformerFramework/Source/MLDeformerFrameworkEditor/Public/MLDeformerAssetFactory.h
@@ -0,0 +1,32 @@
+// Copyright Epic Games, Inc. All Rights Reserved.
+
+#pragma once
+
+#include "CoreTypes.h"
+#include "Factories/Factory.h"
+#include "MLDeformerAssetFactory.generated.h"
+
+/**
+ * The factory for the ML Deformer asset type.
+ * This basically integrates the new asset type into the editor, so you can right click and create a new ML Deformer asset.
+ */
+UCLASS(hidecategories=Object)
+class MLDEFORMERFRAMEWORKEDITOR_API UMLDeformerFactory
+	: public UFactory
+{
+	GENERATED_BODY()
+
+public:
+	UMLDeformerFactory();
+
+	// UFactory overrides.
+	virtual FText GetDisplayName() const override;
+	virtual uint32 GetMenuCategories() const override;
+	virtual FText GetToolTip() const override;
+	virtual FString GetDefaultNewAssetName() const override;
+	virtual UObject* FactoryCreateNew(UClass* Class, UObject* InParent, FName Name, EObjectFlags Flags, UObject* Context, FFeedbackContext* Warn) override;
+	virtual bool ConfigureProperties() override;
+	virtual bool ShouldShowInNewMenu() const override;
+	virtual const TArray<FText>& GetMenuCategorySubMenus() const override;
+	// ~END UFactory overrides.
+};
diff --git a/Engine/Plugins/Animation/MLDeformer/MLDeformerFramework/Source/MLDeformerFrameworkEditor/Public/MLDeformerEditorModel.h b/Engine/Plugins/Animation/MLDeformer/MLDeformerFramework/Source/MLDeformerFrameworkEditor/Public/MLDeformerEditorModel.h
--- a/Engine/Plugins/Animation/MLDeformer/MLDeformerFramework/Source/MLDeformerFrameworkEditor/Public/MLDeformerEditorModel.h
+++ b/Engine/Plugins/Animation/MLDeformer/MLDeformerFramework/Source/MLDeformerFrameworkEditor/Public/MLDeformerEditorModel.h
@@ -735,6 +735,40 @@ namespace UE::MLDeformer
 		 */
 		virtual void UpdateActorLODs();
 
+		/**
+		 * This should update the list of available training devices.
+		 * If not implemented, the user won't be able to pick a training device.
+		 * With training device we mean the device used to store the tensors at. Typically you want to make it build a list of GPU's and have the CPU in there as well.
+		 *
+		 * <code>
+		 * UYourTrainingModel* TrainingModel = NewDerivedObject<UYourTrainingModel>();
+		 * if (TrainingModel)
+		 * {
+		 *     TrainingModel->Init(this);
+		 *     TrainingModel->UpdateAvailableDevices();
+		 *     TrainingModel->ConditionalBeginDestroy();
+		 * }
+		 * </code>
+		 * 
+		 * You can implement the following in Python as an example:
+		 * <code>
+		 * @unreal.ufunction(override=True)
+         * def update_available_devices(self):
+         *     reload(mldeformer.training_helpers)
+         *     mldeformer.training_helpers.update_training_device_list(self)
+		 * </code>
+		 * 
+		 * Then inside the python code you can do something like:
+		 * 
+		 * <code>
+		 * training_device = model.get_training_device()
+         * device_index = mldeformer.training_helpers.find_cuda_device_index(device_name=training_device)
+         * if torch.cuda.is_available() and device_index != -1:
+         *      torch.cuda.set_device(device_index)
+		 * </code>
+		 */
+		virtual void UpdateTrainingDeviceList() {}
+
 		/** Apply the transforms of the debug actor to the actors in the asset editor world. This will internally call ApplyDebugActorTransforms(DebugActorComponentSpaceTransforms). */
 		void ApplyDebugActorTransforms();
 
diff --git a/Engine/Plugins/Animation/MLDeformer/MLDeformerFramework/Source/MLDeformerFrameworkEditor/Public/MLDeformerTrainingModel.h b/Engine/Plugins/Animation/MLDeformer/MLDeformerFramework/Source/MLDeformerFrameworkEditor/Public/MLDeformerTrainingModel.h
--- a/Engine/Plugins/Animation/MLDeformer/MLDeformerFramework/Source/MLDeformerFrameworkEditor/Public/MLDeformerTrainingModel.h
+++ b/Engine/Plugins/Animation/MLDeformer/MLDeformerFramework/Source/MLDeformerFrameworkEditor/Public/MLDeformerTrainingModel.h
@@ -12,6 +12,20 @@ class UMLDeformerModel;
 namespace UE::MLDeformer
 {
 	class FMLDeformerEditorModel;
+
+	// A helper to create the last derived object of a given type.
+	// This is useful when we create a class in Python, derived from a given base class.
+	template<class T>
+	static T* NewDerivedObject()
+	{
+		TArray<UClass*> Classes;
+		GetDerivedClasses(T::StaticClass(), Classes);
+		if (Classes.IsEmpty())
+		{
+			return nullptr;
+		}
+		return NewObject<T>(GetTransientPackage(), Classes.Last());
+	}
 }
 
 /**
@@ -126,6 +140,13 @@ public:
 	UFUNCTION(BlueprintPure, Category = "Training Data")
 	TArray<int32> GetMaskIndexPerSampleArray() const	{ return MaskIndexPerSample; }
 
+	/** Set the list of possible training devices. */
+	UFUNCTION(BlueprintCallable, Category = "Training Model")
+	void SetDeviceList(const TArray<FString>& DeviceNames, int32 PreferredDeviceIndex);
+
+	UFUNCTION(BlueprintImplementableEvent, Category = "Training Model")
+	void UpdateAvailableDevices() const;
+
 protected:
 	/** This updates the sample deltas, curves, and bone rotations. */
 	UE_DEPRECATED(5.4, "Please use UMLDeformerTrainingModel::SampleNextFrame() instead.")
diff --git a/Engine/Plugins/Animation/MLDeformer/NearestNeighborModel/Source/NearestNeighborModelEditor/Private/NearestNeighborEditorModel.cpp b/Engine/Plugins/Animation/MLDeformer/NearestNeighborModel/Source/NearestNeighborModelEditor/Private/NearestNeighborEditorModel.cpp
--- a/Engine/Plugins/Animation/MLDeformer/NearestNeighborModel/Source/NearestNeighborModelEditor/Private/NearestNeighborEditorModel.cpp
+++ b/Engine/Plugins/Animation/MLDeformer/NearestNeighborModel/Source/NearestNeighborModelEditor/Private/NearestNeighborEditorModel.cpp
@@ -12,6 +12,7 @@
 #include "Misc/MessageDialog.h"
 #include "MLDeformerAsset.h"
 #include "MLDeformerComponent.h"
+#include "MLDeformerTrainingModel.h"
 #include "MLDeformerEditorStyle.h"
 #include "MLDeformerEditorToolkit.h"
 #include "NearestNeighborGeomCacheSampler.h"
@@ -115,6 +116,19 @@ namespace UE::NearestNeighborModel
 		UNearestNeighborModelInputInfo* NearestNeighborInputInfo = static_cast<UNearestNeighborModelInputInfo*>(InputInfo);
 		NearestNeighborInputInfo->InitRefBoneRotations(Model->GetSkeletalMesh());
 	}
+
+	void FNearestNeighborEditorModel::UpdateTrainingDeviceList()
+	{
+		// Update the list of devices we can use to train with.
+		// This is the cpu and list of cuda devices.
+		UNearestNeighborTrainingModel* TrainingModel = UE::MLDeformer::NewDerivedObject<UNearestNeighborTrainingModel>();
+		if (TrainingModel)
+		{
+			TrainingModel->Init(this);
+			TrainingModel->UpdateAvailableDevices();
+			TrainingModel->ConditionalBeginDestroy();
+		}
+	}
 	
 	ETrainingResult FNearestNeighborEditorModel::Train()
 	{
@@ -612,7 +626,7 @@ namespace UE::NearestNeighborModel
 			return EOpFlag::Error;
 		}
 		
-		UNearestNeighborTrainingModel *TrainingModel = FHelpers::NewDerivedObject<UNearestNeighborTrainingModel>();
+		UNearestNeighborTrainingModel* TrainingModel = UE::MLDeformer::NewDerivedObject<UNearestNeighborTrainingModel>();
 		if (!TrainingModel)
 		{
 			return EOpFlag::Error;
diff --git a/Engine/Plugins/Animation/MLDeformer/NearestNeighborModel/Source/NearestNeighborModelEditor/Public/NearestNeighborEditorModel.h b/Engine/Plugins/Animation/MLDeformer/NearestNeighborModel/Source/NearestNeighborModelEditor/Public/NearestNeighborEditorModel.h
--- a/Engine/Plugins/Animation/MLDeformer/NearestNeighborModel/Source/NearestNeighborModelEditor/Public/NearestNeighborEditorModel.h
+++ b/Engine/Plugins/Animation/MLDeformer/NearestNeighborModel/Source/NearestNeighborModelEditor/Public/NearestNeighborEditorModel.h
@@ -48,6 +48,7 @@ namespace UE::NearestNeighborModel
 		virtual void OnPropertyChanged(FPropertyChangedEvent& PropertyChangedEvent) override;
 		virtual void OnPostTraining(ETrainingResult TrainingResult, bool bUsePartiallyTrainedWhenAborted) override;
 		virtual void Render(const FSceneView* View, FViewport* Viewport, FPrimitiveDrawInterface* PDI) override;
+		virtual void UpdateTrainingDeviceList() override;
 		// ~END FMLDeformerEditorModel overrides.
 		
 		// UMLDeformerMorphModelEditorModel overrides.
diff --git a/Engine/Plugins/Animation/MLDeformer/NeuralMorphModel/Source/NeuralMorphModel/Private/NeuralMorphModel.cpp b/Engine/Plugins/Animation/MLDeformer/NeuralMorphModel/Source/NeuralMorphModel/Private/NeuralMorphModel.cpp
--- a/Engine/Plugins/Animation/MLDeformer/NeuralMorphModel/Source/NeuralMorphModel/Private/NeuralMorphModel.cpp
+++ b/Engine/Plugins/Animation/MLDeformer/NeuralMorphModel/Source/NeuralMorphModel/Private/NeuralMorphModel.cpp
@@ -156,13 +156,17 @@ namespace
 
 void UNeuralMorphModel::Serialize(FArchive& Archive)
 {
-	Super::Serialize(Archive);
-	Archive.UsingCustomVersion(UE::MLDeformer::FMLDeformerObjectVersion::GUID);
-
-	if (Archive.IsSaving())
+	TMap<FName, FMLDeformerMaskInfo> BoneMaskInfoMapBackup;
+	TMap<FName, FMLDeformerMaskInfo> BoneGroupMaskInfoMapBackup;
+	bool bProcessedDataOnCook = false;
+	ON_SCOPE_EXIT
 	{
-		UpdateMissingGroupNames();
-	}
+		if (bProcessedDataOnCook && GetRecoverStrippedDataAfterCook())
+		{
+			BoneMaskInfoMap = MoveTemp(BoneMaskInfoMapBackup);
+			BoneGroupMaskInfoMap = MoveTemp(BoneGroupMaskInfoMapBackup);
+		}
+	};
 
 	if (Archive.IsSaving() && Archive.IsCooking())
 	{
@@ -173,8 +177,19 @@ void UNeuralMorphModel::Serialize(FArchive& Archive)
 		}
 
 		// Strip the mask data in the cooked asset.
+		BoneMaskInfoMapBackup = MoveTemp(BoneMaskInfoMap);
+		BoneGroupMaskInfoMapBackup = MoveTemp(BoneGroupMaskInfoMap);
 		BoneMaskInfoMap.Empty();
 		BoneGroupMaskInfoMap.Empty();
+		bProcessedDataOnCook = true;
+	}
+
+	Super::Serialize(Archive);
+	Archive.UsingCustomVersion(UE::MLDeformer::FMLDeformerObjectVersion::GUID);
+
+	if (Archive.IsSaving())
+	{
+		UpdateMissingGroupNames();
 	}
 
 	// Convert the UMLDeformerInputInfo object into a UNeuralMorphInputInfo object for backward compatiblity.
diff --git a/Engine/Plugins/Animation/MLDeformer/NeuralMorphModel/Source/NeuralMorphModelEditor/Private/NeuralMorphEditorModel.cpp b/Engine/Plugins/Animation/MLDeformer/NeuralMorphModel/Source/NeuralMorphModelEditor/Private/NeuralMorphEditorModel.cpp
--- a/Engine/Plugins/Animation/MLDeformer/NeuralMorphModel/Source/NeuralMorphModelEditor/Private/NeuralMorphEditorModel.cpp
+++ b/Engine/Plugins/Animation/MLDeformer/NeuralMorphModel/Source/NeuralMorphModelEditor/Private/NeuralMorphEditorModel.cpp
@@ -92,6 +92,19 @@ namespace UE::NeuralMorphModel
 			.EditorModel(this);
 	}
 
+	void FNeuralMorphEditorModel::UpdateTrainingDeviceList()
+	{
+		// Update the list of devices we can use to train with.
+		// This is the cpu and list of cuda devices.
+		UNeuralMorphTrainingModel* TrainingModel = NewDerivedObject<UNeuralMorphTrainingModel>();
+		if (TrainingModel)
+		{
+			TrainingModel->Init(this);
+			TrainingModel->UpdateAvailableDevices();
+			TrainingModel->ConditionalBeginDestroy();
+		}
+	}
+
 	void FNeuralMorphEditorModel::Init(const InitSettings& Settings)
 	{
 		FMLDeformerMorphModelEditorModel::Init(Settings);
diff --git a/Engine/Plugins/Animation/MLDeformer/NeuralMorphModel/Source/NeuralMorphModelEditor/Private/NeuralMorphModelDetails.cpp b/Engine/Plugins/Animation/MLDeformer/NeuralMorphModel/Source/NeuralMorphModelEditor/Private/NeuralMorphModelDetails.cpp
--- a/Engine/Plugins/Animation/MLDeformer/NeuralMorphModel/Source/NeuralMorphModelEditor/Private/NeuralMorphModelDetails.cpp
+++ b/Engine/Plugins/Animation/MLDeformer/NeuralMorphModel/Source/NeuralMorphModelEditor/Private/NeuralMorphModelDetails.cpp
@@ -22,6 +22,11 @@ namespace UE::NeuralMorphModel
 		// Create all the detail categories and add the properties of the base class.
 		FMLDeformerMorphModelDetails::CustomizeDetails(DetailBuilder);
 
+		if (!EditorModel)
+		{
+			return;
+		}
+
 		UNeuralMorphModel* NeuralMorphModel = Cast<UNeuralMorphModel>(Model);
 		check(NeuralMorphModel);
 
diff --git a/Engine/Plugins/Animation/MLDeformer/NeuralMorphModel/Source/NeuralMorphModelEditor/Public/NeuralMorphEditorModel.h b/Engine/Plugins/Animation/MLDeformer/NeuralMorphModel/Source/NeuralMorphModelEditor/Public/NeuralMorphEditorModel.h
--- a/Engine/Plugins/Animation/MLDeformer/NeuralMorphModel/Source/NeuralMorphModelEditor/Public/NeuralMorphEditorModel.h
+++ b/Engine/Plugins/Animation/MLDeformer/NeuralMorphModel/Source/NeuralMorphModelEditor/Public/NeuralMorphEditorModel.h
@@ -43,6 +43,7 @@ namespace UE::NeuralMorphModel
 		virtual void OnPostInputAssetChanged() override;
 		virtual void Render(const FSceneView* View, FViewport* Viewport, FPrimitiveDrawInterface* PDI) override;
 		virtual bool GetSupportsPerTrainingInputAnimVertexMask() const override	{ return true; }
+		virtual void UpdateTrainingDeviceList() override;
 		// ~END FMLDeformerEditorModel overrides.
 
 		// FMLDeformerMorphModelEditorModel overrides.
