[MLDeformer] Bug fix when switching back and forth between ML models and then train again. Also some small tweaks to the python training code to move a tensor to the gpu when using gpu. And moved some display logging to verbose while sampling frames.

diff --git a/Engine/Plugins/Animation/MLDeformer/MLDeformerFramework/Source/MLDeformerFrameworkEditor/Private/MLDeformerEditorToolkit.cpp b/Engine/Plugins/Animation/MLDeformer/MLDeformerFramework/Source/MLDeformerFrameworkEditor/Private/MLDeformerEditorToolkit.cpp
--- a/Engine/Plugins/Animation/MLDeformer/MLDeformerFramework/Source/MLDeformerFrameworkEditor/Private/MLDeformerEditorToolkit.cpp
+++ b/Engine/Plugins/Animation/MLDeformer/MLDeformerFramework/Source/MLDeformerFrameworkEditor/Private/MLDeformerEditorToolkit.cpp
@@ -338,6 +338,8 @@ namespace UE::MLDeformer
 		// Tell the editor we use this model now.
 		TSharedPtr<FMLDeformerEditorModel> NewModel = TSharedPtr<FMLDeformerEditorModel>(EditorModel);
 		NewModel->CopyBaseSettingsFromModel(ActiveModel.Get());
+		NewModel->GetModel()->InitVertexMap();
+		NewModel->GetModel()->InitGPUData();
 
 		ActiveModel.Reset();
 		ActiveModel = NewModel;
diff --git a/Engine/Plugins/Animation/MLDeformer/MLDeformerFramework/Source/MLDeformerFrameworkEditor/Private/MLDeformerGeomCacheTrainingModel.cpp b/Engine/Plugins/Animation/MLDeformer/MLDeformerFramework/Source/MLDeformerFrameworkEditor/Private/MLDeformerGeomCacheTrainingModel.cpp
--- a/Engine/Plugins/Animation/MLDeformer/MLDeformerFramework/Source/MLDeformerFrameworkEditor/Private/MLDeformerGeomCacheTrainingModel.cpp
+++ b/Engine/Plugins/Animation/MLDeformer/MLDeformerFramework/Source/MLDeformerFrameworkEditor/Private/MLDeformerGeomCacheTrainingModel.cpp
@@ -75,7 +75,7 @@ bool UMLDeformerGeomCacheTrainingModel::SampleNextFrame()
 	Sampler->SetVertexDeltaSpace(EVertexDeltaSpace::PreSkinning);
 	Sampler->Sample(CurFrameToSample);
 	MaskIndexPerSample.Add(GetMaskIndexForAnimIndex(SampleAnimIndex));
-	UE_LOG(LogMLDeformer, Display, TEXT("Sampling frame %d of anim %d"), CurFrameToSample, SampleAnimIndex);
+	UE_LOG(LogMLDeformer, Verbose, TEXT("Sampling frame %d of anim %d"), CurFrameToSample, SampleAnimIndex);
 
 	// Copy sampled values.
 	SampleDeltas = Sampler->GetVertexDeltas();
