[MLDeformer] Fixed a crash when switching models and then clicking the LOD "Auto" button.

diff --git a/Engine/Plugins/Animation/MLDeformer/MLDeformerFramework/Source/MLDeformerFrameworkEditor/Private/MLDeformerEditorToolkit.cpp b/Engine/Plugins/Animation/MLDeformer/MLDeformerFramework/Source/MLDeformerFrameworkEditor/Private/MLDeformerEditorToolkit.cpp
--- a/Engine/Plugins/Animation/MLDeformer/MLDeformerFramework/Source/MLDeformerFrameworkEditor/Private/MLDeformerEditorToolkit.cpp
+++ b/Engine/Plugins/Animation/MLDeformer/MLDeformerFramework/Source/MLDeformerFrameworkEditor/Private/MLDeformerEditorToolkit.cpp
@@ -371,7 +371,6 @@ namespace UE::MLDeformer
 			TimeSlider->SetModel(WeakModel);
 		}
 
-		GEngine->ForceGarbageCollection(true);
 		DeformerAsset->Modify();
 	}
 
