[MLDeformer] Fixed a bug where the heatmap could get in a wrong state when toggling it on and off while playing, and while in ground truth mode while playing, and then scrubbing.

diff --git a/Engine/Plugins/Animation/MLDeformer/MLDeformerFramework/Source/MLDeformerFrameworkEditor/Private/MLDeformerEditorModel.cpp b/Engine/Plugins/Animation/MLDeformer/MLDeformerFramework/Source/MLDeformerFrameworkEditor/Private/MLDeformerEditorModel.cpp
--- a/Engine/Plugins/Animation/MLDeformer/MLDeformerFramework/Source/MLDeformerFrameworkEditor/Private/MLDeformerEditorModel.cpp
+++ b/Engine/Plugins/Animation/MLDeformer/MLDeformerFramework/Source/MLDeformerFrameworkEditor/Private/MLDeformerEditorModel.cpp
@@ -1035,6 +1035,14 @@ namespace UE::MLDeformer
 		
 		bIsScrubbingTimeline = bIsScrubbing;
 		UpdatePaintModePose(/*bFullUpdate=*/!bIsScrubbing);
+
+		// If we scrubbed and we were playing before, and we are in ground truth heat map mode, then we need to possibly enable the heatmap again.
+		// We disable it while playing in ground truth mode as we cannot really get correct ground truth heatmaps during playback because of interpolation and other issues.
+		// If we do not do this, then while ground truth heat map is enabled and we are playing, and we scrub, then the heat map isn't enabled again.
+		if (VizSettings->GetShowHeatMap() && VizSettings->GetHeatMapMode() == EMLDeformerHeatMapMode::GroundTruth)
+		{
+			UpdateHeatMapMaterialBasedOnMode();
+		}
 	}
 
 	void FMLDeformerEditorModel::UpdatePaintModePose(bool bFullUpdate)
@@ -1275,10 +1283,12 @@ namespace UE::MLDeformer
 			SetHeatMapMaterialEnabled(Model->GetVizSettings()->GetShowHeatMap());
 			UpdateDeformerGraph();
 			UpdateStepInterpolationMode();
+			UpdateHeatMapMaterialBasedOnMode();
 		}
 		else if (Property->GetFName() == UMLDeformerVizSettings::GetHeatMapModePropertyName())
 		{
 			UpdateStepInterpolationMode();
+			UpdateHeatMapMaterialBasedOnMode();
 		}
 		else
 		if (Property->GetFName() == UMLDeformerVizSettings::GetDrawLinearSkinnedActorPropertyName() ||
@@ -1350,6 +1360,23 @@ namespace UE::MLDeformer
 		}
 	}
 
+	void FMLDeformerEditorModel::UpdateHeatMapMaterialBasedOnMode()
+	{
+		const FMLDeformerEditorActor* BaseActor = GetVisualizationModeBaseActor();
+		const bool bIsPlaying = (BaseActor && BaseActor->IsPlaying());
+		const UMLDeformerVizSettings* VizSettings = Model->GetVizSettings();
+		const bool bInTestMode = (VizSettings->GetVisualizationMode() == EMLDeformerVizMode::TestData);
+
+		if (bIsPlaying && VizSettings->GetHeatMapMode() == EMLDeformerHeatMapMode::GroundTruth)
+		{
+			SetHeatMapMaterialEnabled(false);
+		}
+		else
+		{
+			SetHeatMapMaterialEnabled(VizSettings->GetShowHeatMap());
+		}
+	}
+
 	bool FMLDeformerEditorModel::IsPlayingAnim() const
 	{
 		const FMLDeformerEditorActor* BaseActor = GetVisualizationModeBaseActor();
diff --git a/Engine/Plugins/Animation/MLDeformer/MLDeformerFramework/Source/MLDeformerFrameworkEditor/Public/MLDeformerEditorModel.h b/Engine/Plugins/Animation/MLDeformer/MLDeformerFramework/Source/MLDeformerFrameworkEditor/Public/MLDeformerEditorModel.h
--- a/Engine/Plugins/Animation/MLDeformer/MLDeformerFramework/Source/MLDeformerFrameworkEditor/Public/MLDeformerEditorModel.h
+++ b/Engine/Plugins/Animation/MLDeformer/MLDeformerFramework/Source/MLDeformerFrameworkEditor/Public/MLDeformerEditorModel.h
@@ -1171,6 +1171,14 @@ namespace UE::MLDeformer
 		 */
 		void UpdateStepInterpolationMode();
 
+		/**
+		 * Because of some technical reasons we cannot display a correct ground truth heat map at the moment while playing an animation.
+		 * Scrubbing works fine though, because we can align on the right keyframes then.
+		 * This method will look whether we are having heatmaps enabled or not, and if we are playing or not and if we're in ground truth mode.
+		 * It then disables or enables the heatmap material based on the current state.
+		 */
+		void UpdateHeatMapMaterialBasedOnMode();
+
 	protected:
 		struct FLODInfo
 		{
