[MLDeformer] Fixed the timeline from breaking when switching models. This was because some delegates were still pointing to the old editor model which got destroyed after switching models.Directly fixed a few small errors in the comments.

diff --git a/Engine/Plugins/Animation/MLDeformer/MLDeformerFramework/Source/MLDeformerFrameworkEditor/Private/SMLDeformerTimeline.cpp b/Engine/Plugins/Animation/MLDeformer/MLDeformerFramework/Source/MLDeformerFrameworkEditor/Private/SMLDeformerTimeline.cpp
--- a/Engine/Plugins/Animation/MLDeformer/MLDeformerFramework/Source/MLDeformerFrameworkEditor/Private/SMLDeformerTimeline.cpp
+++ b/Engine/Plugins/Animation/MLDeformer/MLDeformerFramework/Source/MLDeformerFrameworkEditor/Private/SMLDeformerTimeline.cpp
@@ -113,6 +113,8 @@ namespace UE::MLDeformer
 		/** Determine frame time from a mouse position. */
 		FFrameTime GetFrameTimeFromMouse(const FGeometry& Geometry, FVector2D ScreenSpacePosition) const;
 
+		FTimeSliderArgs& GetTimeSliderArgs() { return TimeSliderArgs; }
+
 	private:
 		// Forward declared as class members to prevent name collision with similar types defined in other units.
 		struct FDrawTickArgs;
@@ -294,6 +296,10 @@ namespace UE::MLDeformer
 	void FMLTimeSliderController::SetModel(TWeakPtr<FMLDeformerEditorModel> InModel)
 	{
 		WeakModel = InModel;
+		if (WeakTimeline.IsValid())
+		{
+			WeakTimeline.Pin()->UpdateTimeSliderArgs(TimeSliderArgs);
+		}
 	}
 
 	FFrameTime FMLTimeSliderController::ComputeFrameTimeFromMouse(const FGeometry& Geometry, FVector2D ScreenSpacePosition, FScrubRangeToScreen RangeToScreen, bool CheckSnapping) const
@@ -1340,6 +1346,31 @@ namespace UE::MLDeformer
 		return SCompoundWidget::OnPaint(Args, AllottedGeometry, MyCullingRect, OutDrawElements, LayerId, InWidgetStyle, bParentEnabled);
 	}
 
+
+	void SMLDeformerTimeline::UpdateTimeSliderArgs(FTimeSliderArgs& TimeSliderArgs)
+	{
+		TimeSliderArgs.ScrubPosition = MakeAttributeLambda([this](){ return Model.IsValid() ? Model.Pin()->GetTickResScrubPosition() : FFrameTime(0); });
+		TimeSliderArgs.PlaybackRange = MakeAttributeLambda([this](){ return Model.IsValid() ? Model.Pin()->GetPlaybackRange() : TRange<FFrameNumber>(0, 0); });
+		TimeSliderArgs.ClampRange = MakeAttributeLambda([this]()
+		{
+			if (Model.IsValid())
+			{
+				const TRange<double> Range = Model.Pin()->GetWorkingRange();
+				return FAnimatedRange(Range.GetLowerBoundValue(), Range.GetUpperBoundValue());
+			}
+			else
+			{
+				return FAnimatedRange(0.0, 0.0);
+			}
+		});
+
+		if (Model.IsValid())
+		{
+			TimeSliderArgs.OnClampRangeChanged = FOnTimeRangeChanged::CreateSP(Model.Pin().Get(), &FMLDeformerEditorModel::HandleWorkingRangeChanged);
+		}
+	}
+
+
 	void SMLDeformerTimeline::Construct(const FArguments& InArgs, FMLDeformerEditorToolkit* InEditor)
 	{
 		const TWeakPtr<FMLDeformerEditorModel> WeakModel = InEditor->GetActiveModelPointer();
diff --git a/Engine/Plugins/Animation/MLDeformer/MLDeformerFramework/Source/MLDeformerFrameworkEditor/Private/SMLDeformerTimeline.h b/Engine/Plugins/Animation/MLDeformer/MLDeformerFramework/Source/MLDeformerFrameworkEditor/Private/SMLDeformerTimeline.h
--- a/Engine/Plugins/Animation/MLDeformer/MLDeformerFramework/Source/MLDeformerFrameworkEditor/Private/SMLDeformerTimeline.h
+++ b/Engine/Plugins/Animation/MLDeformer/MLDeformerFramework/Source/MLDeformerFrameworkEditor/Private/SMLDeformerTimeline.h
@@ -18,6 +18,7 @@
 struct FPaintPlaybackRangeArgs;
 class SSearchBox;
 class UAnimSingleNodeInstance;
+struct FTimeSliderArgs;
 enum class EViewRangeInterpolation;
 enum class EFrameNumberDisplayFormats : uint8;
 
@@ -66,6 +67,8 @@ namespace UE::MLDeformer
 
 		void OnDeletedTrainingInputAnim(int32 Index);
 
+		void UpdateTimeSliderArgs(FTimeSliderArgs& TimeSliderArgs);
+
 	private:
 		/**
 		 * @return The fill percentage of the animation outliner.
diff --git a/Engine/Plugins/Animation/MLDeformer/MLDeformerFramework/Source/MLDeformerFrameworkEditor/Public/MLDeformerTrainingModel.h b/Engine/Plugins/Animation/MLDeformer/MLDeformerFramework/Source/MLDeformerFrameworkEditor/Public/MLDeformerTrainingModel.h
--- a/Engine/Plugins/Animation/MLDeformer/MLDeformerFramework/Source/MLDeformerFrameworkEditor/Public/MLDeformerTrainingModel.h
+++ b/Engine/Plugins/Animation/MLDeformer/MLDeformerFramework/Source/MLDeformerFrameworkEditor/Public/MLDeformerTrainingModel.h
@@ -104,7 +104,7 @@ public:
 	bool SetCurrentSampleIndex(int32 Index);
 
 	/**
-	 * Take the next sample.
+	 * Take the next sample. This will update the deltas, curve values and bone rotation arrays with values sampled at the next frame.
 	 * This will return false when there is something wrong or we sampled more times than NumSamples() returns.
 	 */
 	UFUNCTION(BlueprintCallable, Category = "Training Data")
@@ -177,15 +177,15 @@ protected:
 	int32 GetMaskIndexForAnimIndex(int32 AnimIndex) const;
 
 public:
-	// The delta values per vertex for this sample. This is updated after SetCurrentSampleIndex is called. Contains an xyz (3 floats) for each vertex.
+	// The delta values per vertex for this sample. This is updated after NextSample is called. Contains an xyz (3 floats) for each vertex.
 	UPROPERTY(EditAnywhere, BlueprintReadWrite, Category = "Training Data")
 	TArray<float> SampleDeltas;
 
-	// The curve weights. This is updated after SetCurrentSampleIndex is called.
+	// The curve weights. This is updated after NextSample is called.
 	UPROPERTY(EditAnywhere, BlueprintReadWrite, Category = "Training Data")
 	TArray<float> SampleCurveValues;
 
-	// The bone rotations in bone (local) space for this sample. This is updated after SetCurrentSampleIndex is called and is 6 floats per bone (2 columns of 3x3 rotation matrix).
+	// The bone rotations in bone (local) space for this sample. This is updated after NextSample is called and is 6 floats per bone (2 columns of 3x3 rotation matrix).
 	UPROPERTY(EditAnywhere, BlueprintReadWrite, Category = "Training Data")
 	TArray<float> SampleBoneRotations;
 
