[MLDeformer] Added some cook scoped object in the GeometryCacheModel to make it tag the training geom caches as editor only and be excluded from the cook.  Also removed some unneeded code to set editor_only marked properties to nullptr before serialization during cook.

diff --git a/Engine/Plugins/Animation/MLDeformer/MLDeformerFramework/Source/MLDeformerFramework/Private/MLDeformerGeomCacheModel.cpp b/Engine/Plugins/Animation/MLDeformer/MLDeformerFramework/Source/MLDeformerFramework/Private/MLDeformerGeomCacheModel.cpp
--- a/Engine/Plugins/Animation/MLDeformer/MLDeformerFramework/Source/MLDeformerFramework/Private/MLDeformerGeomCacheModel.cpp
+++ b/Engine/Plugins/Animation/MLDeformer/MLDeformerFramework/Source/MLDeformerFramework/Private/MLDeformerGeomCacheModel.cpp
@@ -8,70 +8,13 @@
 #include "UObject/AssetRegistryTagsContext.h"
 #include "UObject/Object.h"
 #include "UObject/UObjectGlobals.h"
+#include "UObject/ICookInfo.h"
 #include "GeometryCache.h"
 
 #include UE_INLINE_GENERATED_CPP_BY_NAME(MLDeformerGeomCacheModel)
 
 #define LOCTEXT_NAMESPACE "MLDeformerGeomCacheModel"
 
-void UMLDeformerGeomCacheModel::Serialize(FArchive& Archive)
-{
-	#if WITH_EDITOR
-		bool bModifiedPropertiesForCook = false;
-		TArray<UGeometryCache*> SavedGeomCaches;
-		TArray<UAnimSequence*> SavedAnimSequences;
-		auto ModifyPropertiesForCook = [this, &bModifiedPropertiesForCook, &SavedGeomCaches, &SavedAnimSequences]()
-		{
-			bModifiedPropertiesForCook = true;
-			GeometryCache_DEPRECATED = nullptr;
-			SavedAnimSequences.Reset();
-			SavedAnimSequences.Reserve(TrainingInputAnims.Num());
-			SavedGeomCaches.Reset();
-			SavedGeomCaches.Reserve(TrainingInputAnims.Num());
-			for (FMLDeformerGeomCacheTrainingInputAnim& Anim : TrainingInputAnims)
-			{
-				SavedGeomCaches.Add(Anim.GetGeometryCache());
-				SavedAnimSequences.Add(Anim.GetAnimSequence());
-				Anim.SetGeometryCache(nullptr);
-				Anim.SetAnimSequence(nullptr);
-			}
-		};
-
-		auto RestorePropertiesForCook = [this, &bModifiedPropertiesForCook, &SavedGeomCaches, &SavedAnimSequences]()
-		{
-			if (!bModifiedPropertiesForCook)
-			{
-				return;
-			}
-			check(TrainingInputAnims.Num() == SavedGeomCaches.Num());
-			check(TrainingInputAnims.Num() == SavedAnimSequences.Num());
-			for (int32 Index = 0; Index < TrainingInputAnims.Num(); ++Index)			
-			{
-				FMLDeformerGeomCacheTrainingInputAnim& Anim = TrainingInputAnims[Index];
-				Anim.SetGeometryCache(SavedGeomCaches[Index]);
-				Anim.SetAnimSequence(SavedAnimSequences[Index]);
-			}		
-		};
-
-		ON_SCOPE_EXIT
-		{
-			if (GetRecoverStrippedDataAfterCook())
-			{
-				RestorePropertiesForCook();
-			}
-		};
-	#endif
-
-	#if WITH_EDITOR
-		if (Archive.IsSaving() && Archive.IsCooking())
-		{
-			ModifyPropertiesForCook();
-		}
-	#endif
-
-	Super::Serialize(Archive);
-	Archive.UsingCustomVersion(UE::MLDeformer::FMLDeformerObjectVersion::GUID);
-}
 
 void UMLDeformerGeomCacheModel::PostLoad()
 {
@@ -129,6 +72,7 @@ void UMLDeformerGeomCacheModel::GetAssetRegistryTags(FAssetRegistryTagsContext C
 #if WITH_EDITOR
 void UMLDeformerGeomCacheModel::UpdateNumTargetMeshVertices()
 {
+	FCookLoadScope CookLoadScope(ECookLoadType::EditorOnly);
 	for (FMLDeformerGeomCacheTrainingInputAnim& Anim : TrainingInputAnims)
 	{
 		UGeometryCache* GeomCache = Anim.GetGeometryCache();
diff --git a/Engine/Plugins/Animation/MLDeformer/MLDeformerFramework/Source/MLDeformerFramework/Public/MLDeformerGeomCacheModel.h b/Engine/Plugins/Animation/MLDeformer/MLDeformerFramework/Source/MLDeformerFramework/Public/MLDeformerGeomCacheModel.h
--- a/Engine/Plugins/Animation/MLDeformer/MLDeformerFramework/Source/MLDeformerFramework/Public/MLDeformerGeomCacheModel.h
+++ b/Engine/Plugins/Animation/MLDeformer/MLDeformerFramework/Source/MLDeformerFramework/Public/MLDeformerGeomCacheModel.h
@@ -28,7 +28,6 @@ class MLDEFORMERFRAMEWORK_API UMLDeformerGeomCacheModel
 
 public:
 	// UObject overrides.
-	virtual void Serialize(FArchive& Archive) override;
 	virtual void PostLoad() override;
 	virtual void GetAssetRegistryTags(FAssetRegistryTagsContext Context) const override;
 	UE_DEPRECATED(5.4, "Implement the version that takes FAssetRegistryTagsContext instead.")
