[MLDeformer] Changed some default settings to see higher quality result on first training. This makes it easier for new users. More advanced users can tweak the values to reduce memory usage and increase performance at runtime.

The values changed are:
- Compression level, which was 20 before, but now 1.0
- Delta Zero Threshold, which was 0.0025 before, and now 0.001.

diff --git a/Engine/Plugins/Animation/MLDeformer/MLDeformerFramework/Source/MLDeformerFramework/Public/MLDeformerMorphModel.h b/Engine/Plugins/Animation/MLDeformer/MLDeformerFramework/Source/MLDeformerFramework/Public/MLDeformerMorphModel.h
--- a/Engine/Plugins/Animation/MLDeformer/MLDeformerFramework/Source/MLDeformerFramework/Public/MLDeformerMorphModel.h
+++ b/Engine/Plugins/Animation/MLDeformer/MLDeformerFramework/Source/MLDeformerFramework/Public/MLDeformerMorphModel.h
@@ -371,6 +371,7 @@ private:
 	 * Include vertex normals in the morph targets?
 	 * The advantage of this can be that it is higher performance than recomputing the normals.
 	 * The disadvantage is it can result in lower quality and uses more memory for the stored morph targets.
+	 * In most cases you want this unchecked and calculate the normals using a deformer graph or using the skeletal mesh's tangent recompute settings in the section details.
 	 */
 	UPROPERTY(EditAnywhere, Category = "Morph Targets", meta = (EditCondition = "CanDynamicallyUpdateMorphTargets()"))
 	bool bIncludeNormals = false;
@@ -379,16 +380,19 @@ private:
 	 * Morph target delta values that are smaller than or equal to this threshold will be zeroed out.
 	 * This essentially removes small deltas from morph targets, which will lower the memory usage at runtime, however when set too high it can also introduce visual artifacts.
 	 * A value of 0 will result in the highest quality morph targets, at the cost of higher runtime memory usage.
+	 * On the left side in the MLD asset editor you can see the estimated GPU Memory Usage. That value should change when you modify this property.
 	 */
 	UPROPERTY(EditAnywhere, Category = "Morph Targets", DisplayName = "Delta Zero Threshold", meta = (ClampMin = "0.0", ClampMax = "1.0", ForceUnits="cm", EditCondition = "CanDynamicallyUpdateMorphTargets()"))
-	float MorphDeltaZeroThreshold = 0.0025f;
+	float MorphDeltaZeroThreshold = 0.001f;
 
 	/** 
 	 * The morph target compression level. Higher values result in larger compression, but could result in visual artifacts.
-	 * Most of the times this is a value between 20 and 200.
+	 * Most of the times this is a value between 1 and 50. It is best to first try increasing the "Delta Zero Threshold" property as high as you can though.
+	 * Once that is set to the highest acceptable value, try increasing this morph comrpession level as high as visually acceptable.
+	 * On the left side in the MLD asset editor you can see the estimated GPU Memory Usage. That value should change when you modify this property.
 	 */
 	UPROPERTY(EditAnywhere, Category = "Morph Targets", DisplayName = "Compression Level", meta = (ClampMin = "0.01", ClampMax = "1000", EditCondition = "CanDynamicallyUpdateMorphTargets()"))
-	float MorphCompressionLevel = 20.0f;
+	float MorphCompressionLevel = 1.0f;
 
 	/**
 	 * The channel data that represents the delta mask multipliers.
diff --git a/Engine/Plugins/Animation/MLDeformer/NeuralMorphModel/Source/NeuralMorphModel/Public/NeuralMorphModel.h b/Engine/Plugins/Animation/MLDeformer/NeuralMorphModel/Source/NeuralMorphModel/Public/NeuralMorphModel.h
--- a/Engine/Plugins/Animation/MLDeformer/NeuralMorphModel/Source/NeuralMorphModel/Public/NeuralMorphModel.h
+++ b/Engine/Plugins/Animation/MLDeformer/NeuralMorphModel/Source/NeuralMorphModel/Public/NeuralMorphModel.h
@@ -147,7 +147,7 @@ public:
 	/** 
 	 * The number of iterations to train the model for. 
 	 * If you are quickly iterating then around 1000 to 3000 iterations should be enough.
-	 * If you want to generate final assets you might want to use a higher number of iterations, like 10k to 100k.
+	 * If you want to generate final assets you might want to use a higher number of iterations, like 10k to 100k or even up to one million.
 	 * Once the loss doesn't go down anymore, you know that more iterations most likely won't help much.
 	 */
 	UPROPERTY(EditAnywhere, BlueprintReadWrite, Category = "Training Settings", meta = (ClampMin = "1", ClampMax = "1000000"))
