[MLDeformer] PR #12744: Add advanced preview scene settings tab to editor

#pr 12744

diff --git a/Engine/Plugins/Animation/MLDeformer/MLDeformerFramework/Source/MLDeformerFrameworkEditor/Private/MLDeformerApplicationMode.cpp b/Engine/Plugins/Animation/MLDeformer/MLDeformerFramework/Source/MLDeformerFrameworkEditor/Private/MLDeformerApplicationMode.cpp
--- a/Engine/Plugins/Animation/MLDeformer/MLDeformerFramework/Source/MLDeformerFrameworkEditor/Private/MLDeformerApplicationMode.cpp
+++ b/Engine/Plugins/Animation/MLDeformer/MLDeformerFramework/Source/MLDeformerFrameworkEditor/Private/MLDeformerApplicationMode.cpp
@@ -35,6 +35,7 @@ namespace UE::MLDeformer
 		// Register Persona tabs.
 		FPersonaModule& PersonaModule = FModuleManager::LoadModuleChecked<FPersonaModule>("Persona");
 		TabFactories.RegisterFactory(PersonaModule.CreatePersonaViewportTabFactory(InHostingApp, ViewportArgs));
+		TabFactories.RegisterFactory(PersonaModule.CreateAdvancedPreviewSceneTabFactory(InHostingApp, InPreviewScene));
 		TabFactories.RegisterFactory(PersonaModule.CreatePersonaToolboxTabFactory(MLDeformerEditor));
 		TabFactories.RegisterFactory(PersonaModule.CreateDetailsTabFactory(InHostingApp, FOnDetailsCreated::CreateSP(&MLDeformerEditor.Get(), &FMLDeformerEditorToolkit::HandleDetailsCreated)));
 
