[MLDeformer] PR #12767: Avoid crash due to nullptr in FMLDeformerEditorToolkit::OnModelChanged

#pr 12767

diff --git a/Engine/Plugins/Animation/MLDeformer/MLDeformerFramework/Source/MLDeformerFrameworkEditor/Private/MLDeformerEditorToolkit.cpp b/Engine/Plugins/Animation/MLDeformer/MLDeformerFramework/Source/MLDeformerFrameworkEditor/Private/MLDeformerEditorToolkit.cpp
--- a/Engine/Plugins/Animation/MLDeformer/MLDeformerFramework/Source/MLDeformerFrameworkEditor/Private/MLDeformerEditorToolkit.cpp
+++ b/Engine/Plugins/Animation/MLDeformer/MLDeformerFramework/Source/MLDeformerFrameworkEditor/Private/MLDeformerEditorToolkit.cpp
@@ -337,7 +337,12 @@ namespace UE::MLDeformer
 
 		// Tell the editor we use this model now.
 		TSharedPtr<FMLDeformerEditorModel> NewModel = TSharedPtr<FMLDeformerEditorModel>(EditorModel);
-		NewModel->CopyBaseSettingsFromModel(ActiveModel.Get());
+
+		if (ActiveModel)
+		{
+			NewModel->CopyBaseSettingsFromModel(ActiveModel.Get());
+		}
+
 		NewModel->GetModel()->InitVertexMap();
 		NewModel->GetModel()->InitGPUData();
 
