[MLDeformer] Make it take the component space transforms instead of bone space, as sometimes this doesn't work well when poses are copied without leader poses.

diff --git a/Engine/Plugins/Animation/MLDeformer/MLDeformerFramework/Source/MLDeformerFramework/Private/MLDeformerModelInstance.cpp b/Engine/Plugins/Animation/MLDeformer/MLDeformerFramework/Source/MLDeformerFramework/Private/MLDeformerModelInstance.cpp
--- a/Engine/Plugins/Animation/MLDeformer/MLDeformerFramework/Source/MLDeformerFramework/Private/MLDeformerModelInstance.cpp
+++ b/Engine/Plugins/Animation/MLDeformer/MLDeformerFramework/Source/MLDeformerFramework/Private/MLDeformerModelInstance.cpp
@@ -239,18 +239,28 @@ void UMLDeformerModelInstance::UpdateBoneTransforms()
 	{
 		if (FinalSkelMeshComponent == SkeletalMeshComponent)
 		{
-			// Grab the transforms from our own skeletal mesh component.
-			// These are local space transforms, relative to the parent bone.
-			const TArrayView<const FTransform> Transforms = FinalSkelMeshComponent->GetBoneSpaceTransformsView();
-			if (!Transforms.IsEmpty())
+			const TArray<FTransform>& ComponentTransforms = FinalSkelMeshComponent->GetComponentSpaceTransforms();
+			if (!ComponentTransforms.IsEmpty())
 			{
+				const USkinnedAsset* SkinnedAsset = FinalSkelMeshComponent->GetSkinnedAsset();
+				const FReferenceSkeleton& RefSkel = SkinnedAsset->GetRefSkeleton();
 				const int32 NumTrainingBones = AssetBonesToSkelMeshMappings.Num();
 				for (int32 Index = 0; Index < NumTrainingBones; ++Index)
 				{
 					const int32 ComponentBoneIndex = AssetBonesToSkelMeshMappings[Index];
-					if (Transforms.IsValidIndex(ComponentBoneIndex))
+					if (ComponentTransforms.IsValidIndex(ComponentBoneIndex))
 					{
-						TrainingBoneTransforms[Index] = Transforms[ComponentBoneIndex];
+						const FTransform& ComponentSpaceTransform = ComponentTransforms[ComponentBoneIndex];
+						const int32 ParentIndex = RefSkel.GetParentIndex(ComponentBoneIndex);
+						if (ComponentTransforms.IsValidIndex(ParentIndex))
+						{
+							TrainingBoneTransforms[Index] = ComponentSpaceTransform.GetRelativeTransform(ComponentTransforms[ParentIndex]);
+							TrainingBoneTransforms[Index].NormalizeRotation();
+						}
+						else
+						{
+							TrainingBoneTransforms[Index] = ComponentSpaceTransform;
+						}
 					}
 				}
 			}
