[MLDeformer] Crash fix while cooking. This was because the MLDeformerModelInstance wasn't keeping the SkeletalMeshComponent alive and that could a crash. Made the TObjectPtr a transient UPROPERTY now.

diff --git a/Engine/Plugins/Animation/MLDeformer/MLDeformerFramework/Source/MLDeformerFramework/Public/MLDeformerModelInstance.h b/Engine/Plugins/Animation/MLDeformer/MLDeformerFramework/Source/MLDeformerFramework/Public/MLDeformerModelInstance.h
--- a/Engine/Plugins/Animation/MLDeformer/MLDeformerFramework/Source/MLDeformerFramework/Public/MLDeformerModelInstance.h
+++ b/Engine/Plugins/Animation/MLDeformer/MLDeformerFramework/Source/MLDeformerFramework/Public/MLDeformerModelInstance.h
@@ -254,6 +254,7 @@ protected:
 	TObjectPtr<UMLDeformerModel> Model = nullptr;
 
 	/** The skeletal mesh component we work with. This is mainly used for compatibility checks. */
+	UPROPERTY(Transient)
 	TObjectPtr<USkeletalMeshComponent> SkeletalMeshComponent = nullptr;
 
 	/** The cached current local space bone transforms for the current frame. */
