[MLDeformer] Add a confirmation dialog to the Finalize Morph Targets button as you cannot undo this.

diff --git a/Engine/Plugins/Animation/MLDeformer/MLDeformerFramework/Source/MLDeformerFrameworkEditor/Private/MLDeformerMorphModelDetails.cpp b/Engine/Plugins/Animation/MLDeformer/MLDeformerFramework/Source/MLDeformerFrameworkEditor/Private/MLDeformerMorphModelDetails.cpp
--- a/Engine/Plugins/Animation/MLDeformer/MLDeformerFramework/Source/MLDeformerFrameworkEditor/Private/MLDeformerMorphModelDetails.cpp
+++ b/Engine/Plugins/Animation/MLDeformer/MLDeformerFramework/Source/MLDeformerFrameworkEditor/Private/MLDeformerMorphModelDetails.cpp
@@ -15,6 +15,7 @@
 #include "Engine/SkinnedAssetCommon.h"
 #include "Materials/MaterialInterface.h"
 #include "Materials/Material.h"
+#include "Misc/MessageDialog.h"
 
 #define LOCTEXT_NAMESPACE "MLDeformerMorphModelDetails"
 
@@ -231,9 +232,25 @@ namespace UE::MLDeformer
 					.IsEnabled_Lambda([this]() { return MorphModel && !MorphModel->GetMorphTargetDeltas().IsEmpty(); })
 					.OnClicked_Lambda([this]() -> FReply
 					{
-						if (MorphModel)
+						const EAppReturnType::Type Result = FMessageDialog::Open(
+							EAppMsgType::YesNo,
+							LOCTEXT("FinalizeMorphsWarningMessage",
+									"This will remove the raw uncompressed trained morph target deltas, just like when cooking the asset. "
+									"Doing this will make the uncooked asset that you submit to source control a lot smaller.\n\n"
+									"However, after doing this, changing settings in the morph target category will not have an effect until "
+									"you retrain the model again. It does not impact how the deformation performs at runtime.\n\n"
+									"If you proceed, you cannot undo this operation. To get the uncompressed deltas back you have to train the "
+									"model again.\n"
+									"\n"
+									"Would you like to continue?"),
+							LOCTEXT("FinalizeMorphsWarningDialogTitle", "Finalize morph targets?"));
+
+						if (Result == EAppReturnType::Yes)
 						{
-							MorphModel->FinalizeMorphTargets();
+							if (MorphModel)
+							{
+								MorphModel->FinalizeMorphTargets();
+							}
 						}
 						return FReply::Handled();
 					})
