[MLDeformer] Show a UI warning when the frame rate of a geom cache and animation don't match. As this can result in issues when scrubbing the timeline or training. Directly added some safety check to the morph model instance, which could potentially cause a crash.

diff --git a/Engine/Plugins/Animation/MLDeformer/MLDeformerFramework/Source/MLDeformerFramework/Private/MLDeformerGeomCacheHelpers.cpp b/Engine/Plugins/Animation/MLDeformer/MLDeformerFramework/Source/MLDeformerFramework/Private/MLDeformerGeomCacheHelpers.cpp
--- a/Engine/Plugins/Animation/MLDeformer/MLDeformerFramework/Source/MLDeformerFramework/Private/MLDeformerGeomCacheHelpers.cpp
+++ b/Engine/Plugins/Animation/MLDeformer/MLDeformerFramework/Source/MLDeformerFramework/Private/MLDeformerGeomCacheHelpers.cpp
@@ -455,7 +455,7 @@ namespace UE::MLDeformer
 			const int32 NumAnimSeqFrames = InAnimSequence->GetNumberOfSampledKeys();
 			const float AnimSeqDuration = InAnimSequence->GetPlayLength();
 			const float GeomCacheDuration = InGeomCache->CalculateDuration();
-
+			
 			if (NumGeomCacheFrames != NumAnimSeqFrames)
 			{
 				Result = FText::Format(
@@ -463,7 +463,7 @@ namespace UE::MLDeformer
 					FText::AsNumber(NumAnimSeqFrames),
 					FText::AsNumber(NumGeomCacheFrames));
 			}
-
+			
 			if (FMath::Abs(AnimSeqDuration - GeomCacheDuration) > 0.001f)
 			{
 				FNumberFormattingOptions Options;
@@ -477,6 +477,22 @@ namespace UE::MLDeformer
 					FText::AsNumber(AnimSeqDuration, &Options),
 					FText::AsNumber(GeomCacheDuration, &Options));
 			}
+
+			const double AnimSequenceFrameRate = InAnimSequence->GetSamplingFrameRate().AsDecimal();
+			const double GeomCacheFrameRate = static_cast<double>(NumGeomCacheFrames - 1) / GeomCacheDuration;
+			if (FMath::Abs(GeomCacheFrameRate - AnimSequenceFrameRate) > KINDA_SMALL_NUMBER)
+			{
+				FNumberFormattingOptions Options;
+				Options.SetUseGrouping(false);
+				Options.SetMaximumFractionalDigits(2);
+
+				Result = FText::Format(
+					LOCTEXT("FPSMismatch", "{0}{1}Anim sequence frame rate of {2} doesn't match the geometry cache's frame rate of {3}."),
+					Result,
+					!Result.IsEmpty() ? FText::FromString("\n\n") : FText::GetEmpty(),
+					FText::AsNumber(AnimSequenceFrameRate, &Options),
+					FText::AsNumber(GeomCacheFrameRate, &Options));
+			}
 		}
 
 		return Result;
diff --git a/Engine/Plugins/Animation/MLDeformer/MLDeformerFramework/Source/MLDeformerFramework/Private/MLDeformerMorphModelInstance.cpp b/Engine/Plugins/Animation/MLDeformer/MLDeformerFramework/Source/MLDeformerFramework/Private/MLDeformerMorphModelInstance.cpp
--- a/Engine/Plugins/Animation/MLDeformer/MLDeformerFramework/Source/MLDeformerFramework/Private/MLDeformerMorphModelInstance.cpp
+++ b/Engine/Plugins/Animation/MLDeformer/MLDeformerFramework/Source/MLDeformerFramework/Private/MLDeformerMorphModelInstance.cpp
@@ -119,18 +119,24 @@ FExternalMorphSetWeights* UMLDeformerMorphModelInstance::FindWeightData(int32 LO
 {
 	// Check if our LOD index is valid first, as we might not have registered yet.
 	USkeletalMeshComponent* SkelMeshComponent = SkeletalMeshComponent.Get();
-	if (SkelMeshComponent == nullptr || !SkelMeshComponent->IsValidExternalMorphSetLODIndex(LOD))
+	if (!SkelMeshComponent || !SkelMeshComponent->IsValidExternalMorphSetLODIndex(LOD))
 	{
 		return nullptr;
 	}
 
 	// Grab the weight data for this morph set.
 	// This could potentially fail if we are applying this deformer to the wrong skeletal mesh component.
-	return SkelMeshComponent->GetExternalMorphWeights(LOD).MorphSets.Find(ExternalMorphSetID);
+	FExternalMorphWeightData& WeightData = SkelMeshComponent->GetExternalMorphWeights(LOD);
+	return WeightData.MorphSets.Find(ExternalMorphSetID);
 }
 
 void UMLDeformerMorphModelInstance::HandleZeroModelWeight()
 {
+	if (!SkeletalMeshComponent)
+	{
+		return;
+	}
+
 	const int32 LOD = SkeletalMeshComponent->GetPredictedLODLevel();
 	FExternalMorphSetWeights* WeightData = FindWeightData(LOD);
 	if (WeightData)
