[MLDeformer] Added in-engine tools to help with the creation of training data for ML Deformers. Also added support for Undo/Redo to the ML Deformer asset editor. Fixed some crash as well, where a skeletal mesh pointer could become invalid.

diff --git a/Engine/Plugins/Animation/MLDeformer/MLDeformerFramework/Source/MLDeformerFramework/Private/MLDeformerModel.cpp b/Engine/Plugins/Animation/MLDeformer/MLDeformerFramework/Source/MLDeformerFramework/Private/MLDeformerModel.cpp
--- a/Engine/Plugins/Animation/MLDeformer/MLDeformerFramework/Source/MLDeformerFramework/Private/MLDeformerModel.cpp
+++ b/Engine/Plugins/Animation/MLDeformer/MLDeformerFramework/Source/MLDeformerFramework/Private/MLDeformerModel.cpp
@@ -17,6 +17,7 @@
 #include "AssetRegistry/AssetData.h"
 #include "Components/SkeletalMeshComponent.h"
 #include "UObject/AssetRegistryTagsContext.h"
+#include "Misc/TransactionObjectEvent.h"
 
 #include UE_INLINE_GENERATED_CPP_BY_NAME(MLDeformerModel)
 
@@ -62,7 +63,7 @@ FString UMLDeformerModel::GetDisplayName() const
 
 UMLDeformerInputInfo* UMLDeformerModel::CreateInputInfo()
 { 
-	return NewObject<UMLDeformerInputInfo>(this);
+	return NewObject<UMLDeformerInputInfo>(this, NAME_None, RF_Transactional);
 }
 
 UMLDeformerModelInstance* UMLDeformerModel::CreateModelInstance(UMLDeformerComponent* Component)
@@ -159,6 +160,8 @@ void UMLDeformerModel::PostLoad()
 	#if WITH_EDITOR
 		UpdateMemoryUsage();
 	#endif
+
+	SetFlags(RF_Transactional);
 }
 
 // Used for the FBoenReference, so it knows what skeleton to pick bones from.
@@ -365,8 +368,28 @@ void UMLDeformerModel::SetTrainingDevice(const FString& DeviceName)
 		}
 	}
 
+	void UMLDeformerModel::PostTransacted(const FTransactionObjectEvent& Event)
+	{
+		Super::PostTransacted(Event);
+		OnPostTransacted().Broadcast(Event);
+	}
+
+	void UMLDeformerModel::PreEditUndo()
+	{
+		Super::PreEditUndo();
+		OnPreEditUndo().Broadcast();
+	}
+
+	void UMLDeformerModel::PostEditUndo()
+	{
+		Super::PostEditUndo();
+		OnPostEditUndo().Broadcast();
+	}
+
 	void UMLDeformerModel::PostEditChangeProperty(FPropertyChangedEvent& PropertyChangedEvent)
 	{
+		Super::PostEditChangeProperty(PropertyChangedEvent);
+
 		const FProperty* Property = PropertyChangedEvent.Property;
 		if (Property == nullptr)
 		{
diff --git a/Engine/Plugins/Animation/MLDeformer/MLDeformerFramework/Source/MLDeformerFramework/Private/MLDeformerModelInstance.cpp b/Engine/Plugins/Animation/MLDeformer/MLDeformerFramework/Source/MLDeformerFramework/Private/MLDeformerModelInstance.cpp
--- a/Engine/Plugins/Animation/MLDeformer/MLDeformerFramework/Source/MLDeformerFramework/Private/MLDeformerModelInstance.cpp
+++ b/Engine/Plugins/Animation/MLDeformer/MLDeformerFramework/Source/MLDeformerFramework/Private/MLDeformerModelInstance.cpp
@@ -35,7 +35,7 @@ bool UMLDeformerModelInstance::IsReadyForFinishDestroy()
 
 USkeletalMeshComponent* UMLDeformerModelInstance::GetSkeletalMeshComponent() const
 { 
-	return SkeletalMeshComponent;
+	return SkeletalMeshComponent.Get();
 }
 
 UMLDeformerModel* UMLDeformerModelInstance::GetModel() const
@@ -128,7 +128,7 @@ void UMLDeformerModelInstance::Init(USkeletalMeshComponent* SkelMeshComponent)
 
 void UMLDeformerModelInstance::UpdateCompatibilityStatus()
 {
-	bIsCompatible = SkeletalMeshComponent && SkeletalMeshComponent->GetSkeletalMeshAsset() && CheckCompatibility(SkeletalMeshComponent, true).IsEmpty();
+	bIsCompatible = SkeletalMeshComponent && SkeletalMeshComponent->GetSkeletalMeshAsset() && CheckCompatibility(SkeletalMeshComponent.Get(), true).IsEmpty();
 }
 
 FString UMLDeformerModelInstance::CheckCompatibility(USkeletalMeshComponent* InSkelMeshComponent, bool bLogIssues)
@@ -177,7 +177,7 @@ USkeletalMeshComponent* UMLDeformerModelInstance::GetFinalSkeletalMeshComponent(
 		}
 	#endif
 
-	return SkeletalMeshComponent;
+	return SkeletalMeshComponent.Get();
 }
 
 
diff --git a/Engine/Plugins/Animation/MLDeformer/MLDeformerFramework/Source/MLDeformerFramework/Private/MLDeformerMorphModel.cpp b/Engine/Plugins/Animation/MLDeformer/MLDeformerFramework/Source/MLDeformerFramework/Private/MLDeformerMorphModel.cpp
--- a/Engine/Plugins/Animation/MLDeformer/MLDeformerFramework/Source/MLDeformerFramework/Private/MLDeformerMorphModel.cpp
+++ b/Engine/Plugins/Animation/MLDeformer/MLDeformerFramework/Source/MLDeformerFramework/Private/MLDeformerMorphModel.cpp
@@ -118,7 +118,6 @@ void UMLDeformerMorphModel::Serialize(FArchive& Archive)
 			int32 NumLODs = 1;
 			Archive << NumLODs;
 
-			check(GetNumLODs() == 0);
 			AddMorphSets(NumLODs);
 
 			for (int32 LOD = 0; LOD < NumLODs; ++LOD)
@@ -135,6 +134,11 @@ void UMLDeformerMorphModel::Serialize(FArchive& Archive)
 	}
 }
 
+UMLDeformerInputInfo* UMLDeformerMorphModel::CreateInputInfo()
+{
+	return NewObject<UMLDeformerMorphModelInputInfo>(this, NAME_None, RF_Transactional);
+}
+
 void UMLDeformerMorphModel::PostLoad()
 {
 	Super::PostLoad();
@@ -206,7 +210,7 @@ UMLDeformerModelInstance* UMLDeformerMorphModel::CreateModelInstance(UMLDeformer
 
 void UMLDeformerMorphModel::SetMorphTargetDeltaFloats(const TArray<float>& Deltas)
 {
-	FloatArrayToVector3Array(Deltas, MorphTargetDeltas);
+	FloatArrayToVector3Array(Deltas, GetMorphTargetDeltas());
 }
 
 void UMLDeformerMorphModel::SetMorphTargetDeltas(const TArray<FVector3f>& Deltas)
@@ -244,7 +248,7 @@ void UMLDeformerMorphModel::AddMorphSets(int32 NumToAdd)
 
 int32 UMLDeformerMorphModel::GetMorphTargetDeltaStartIndex(int32 MorphTargetIndex) const
 {
-	if (MorphTargetDeltas.Num() == 0)
+	if (GetMorphTargetDeltas().Num() == 0)
 	{
 		return INDEX_NONE;
 	}
@@ -386,14 +390,14 @@ void UMLDeformerMorphModel::UpdateMemoryUsage()
 
 void UMLDeformerMorphModel::FinalizeMorphTargets()
 {
-	MorphTargetDeltas.Empty();
+	GetMorphTargetDeltas().Empty();
 	UpdateStatistics();
 	UpdateMemoryUsage();
 }
 
 bool UMLDeformerMorphModel::HasRawMorph() const
 {
-	return !MorphTargetDeltas.IsEmpty();
+	return !GetMorphTargetDeltas().IsEmpty();
 }
 #endif
 
diff --git a/Engine/Plugins/Animation/MLDeformer/MLDeformerFramework/Source/MLDeformerFramework/Private/MLDeformerMorphModelInstance.cpp b/Engine/Plugins/Animation/MLDeformer/MLDeformerFramework/Source/MLDeformerFramework/Private/MLDeformerMorphModelInstance.cpp
--- a/Engine/Plugins/Animation/MLDeformer/MLDeformerFramework/Source/MLDeformerFramework/Private/MLDeformerMorphModelInstance.cpp
+++ b/Engine/Plugins/Animation/MLDeformer/MLDeformerFramework/Source/MLDeformerFramework/Private/MLDeformerMorphModelInstance.cpp
@@ -38,7 +38,7 @@ int32 UMLDeformerMorphModelInstance::GetExternalMorphSetID() const
 void UMLDeformerMorphModelInstance::BeginDestroy()
 {
 	// Try to unregister the morph target and morph target set.
-	if (IsValid(SkeletalMeshComponent))
+	if (SkeletalMeshComponent)
 	{
 		const UMLDeformerMorphModel* MorphModel = Cast<UMLDeformerMorphModel>(Model);
 		if (MorphModel)
@@ -118,16 +118,14 @@ void UMLDeformerMorphModelInstance::PostMLDeformerComponentInit()
 FExternalMorphSetWeights* UMLDeformerMorphModelInstance::FindWeightData(int32 LOD) const
 {
 	// Check if our LOD index is valid first, as we might not have registered yet.
-	USkeletalMeshComponent* SkelMeshComponent = SkeletalMeshComponent.Get();
-	if (!SkelMeshComponent || !SkelMeshComponent->IsValidExternalMorphSetLODIndex(LOD))
+	if (!SkeletalMeshComponent || !SkeletalMeshComponent->IsValidExternalMorphSetLODIndex(LOD))
 	{
 		return nullptr;
 	}
 
 	// Grab the weight data for this morph set.
 	// This could potentially fail if we are applying this deformer to the wrong skeletal mesh component.
-	FExternalMorphWeightData& WeightData = SkelMeshComponent->GetExternalMorphWeights(LOD);
-	return WeightData.MorphSets.Find(ExternalMorphSetID);
+	return SkeletalMeshComponent->GetExternalMorphWeights(LOD).MorphSets.Find(ExternalMorphSetID);
 }
 
 void UMLDeformerMorphModelInstance::HandleZeroModelWeight()
diff --git a/Engine/Plugins/Animation/MLDeformer/MLDeformerFramework/Source/MLDeformerFramework/Private/MLDeformerTrainingDataProcessorSettings.cpp b/Engine/Plugins/Animation/MLDeformer/MLDeformerFramework/Source/MLDeformerFramework/Private/MLDeformerTrainingDataProcessorSettings.cpp
new file mode 100644
--- /dev/null
+++ b/Engine/Plugins/Animation/MLDeformer/MLDeformerFramework/Source/MLDeformerFramework/Private/MLDeformerTrainingDataProcessorSettings.cpp
@@ -0,0 +1,87 @@
+// Copyright Epic Games, Inc. All Rights Reserved.
+
+#include "MLDeformerTrainingDataProcessorSettings.h"
+#include "MLDeformerModel.h"
+#include "Animation/Skeleton.h"
+#include "Animation/AnimSequence.h"
+#include "Engine/SkeletalMesh.h"
+
+#define LOCTEXT_NAMESPACE "MLDeformerTrainingDataProcessorSettings"
+
+#if WITH_EDITOR
+bool UMLDeformerTrainingDataProcessorSettings::IsValid(const USkeleton* Skeleton) const
+{
+	if (!Skeleton)
+	{
+		return false;
+	}
+
+	// If there is nothing to do.
+	if (!bReduceFrames && !bRemixPoses)
+	{
+		return false;
+	}
+
+	// Make sure we have at least one valid input animation.
+	// As in theory the entries in the list could be empty ones or could all be disabled.
+	int32 NumValidAnims = 0;
+	int32 TotalNumFrames = 0;
+	for (const FMLDeformerTrainingDataProcessorAnim& Anim : AnimList)
+	{
+		if (Anim.bEnabled && Anim.AnimSequence.LoadSynchronous())
+		{
+			const USkeleton* AnimSkeleton = Anim.AnimSequence->GetSkeleton();
+			if (AnimSkeleton == Skeleton)
+			{
+				TotalNumFrames += Anim.AnimSequence->GetDataModel() ? Anim.AnimSequence->GetDataModel()->GetNumberOfFrames() : 0;
+				NumValidAnims++;
+			}
+		}
+	}
+
+	// We have no valid animations, so there is nothing to do.
+	if (NumValidAnims == 0)
+	{
+		return false;
+	}
+
+	// We need at least some frames in the anim sequences.
+	if (TotalNumFrames == 0)
+	{
+		return false;
+	}
+
+	// Make sure we have a valid output animation sequence setup.
+	if (!OutputAnimSequence.LoadSynchronous())
+	{
+		return false;
+	}
+
+	return true;
+}
+
+USkeleton* UMLDeformerTrainingDataProcessorSettings::FindSkeleton() const
+{
+	UMLDeformerModel* Model = Cast<UMLDeformerModel>(GetOuter());
+	if (Model)
+	{
+		return Model->GetSkeletalMesh() ? Model->GetSkeletalMesh()->GetSkeleton() : nullptr;
+	}
+	return nullptr;
+}
+
+int32 UMLDeformerTrainingDataProcessorSettings::GetNumInputAnimationFrames() const
+{
+	int32 NumFrames = 0;
+	for (const FMLDeformerTrainingDataProcessorAnim& Anim : AnimList)
+	{
+		if (Anim.bEnabled && Anim.AnimSequence.LoadSynchronous())
+		{
+			NumFrames += Anim.AnimSequence->GetDataModel() ? Anim.AnimSequence->GetDataModel()->GetNumberOfFrames() : 0;
+		}
+	}
+	return NumFrames;
+}
+#endif // #if WITH_EDITOR
+
+#undef LOCTEXT_NAMESPACE
diff --git a/Engine/Plugins/Animation/MLDeformer/MLDeformerFramework/Source/MLDeformerFramework/Public/MLDeformerModel.h b/Engine/Plugins/Animation/MLDeformer/MLDeformerFramework/Source/MLDeformerFramework/Public/MLDeformerModel.h
--- a/Engine/Plugins/Animation/MLDeformer/MLDeformerFramework/Source/MLDeformerFramework/Public/MLDeformerModel.h
+++ b/Engine/Plugins/Animation/MLDeformer/MLDeformerFramework/Source/MLDeformerFramework/Public/MLDeformerModel.h
@@ -20,6 +20,7 @@ class UMLDeformerVizSettings;
 class UMLDeformerModelInstance;
 class UMLDeformerComponent;
 class UMLDeformerInputInfo;
+class UMLDeformerTrainingDataProcessorSettings;
 
 /** The channel to get the mask data from. */
 UENUM()
@@ -99,6 +100,9 @@ namespace UE::MLDeformer
 
 // DEPRECATED: Use FMLDeformerReinitModelInstancesDelegate instead.
 DECLARE_EVENT_OneParam(UMLDeformerModel, FMLDeformerModelOnPostEditProperty, FPropertyChangedEvent&)
+DECLARE_EVENT_OneParam(UMLDeformerModel, FMLDeformerModelOnPostTransacted, const FTransactionObjectEvent&)
+DECLARE_EVENT(UMLDeformerModel, FMLDeformerModelOnPreEditUndo)
+DECLARE_EVENT(UMLDeformerModel, FMLDeformerModelOnPostEditUndo)
 
 /**
  * The ML Deformer runtime model base class.
@@ -454,6 +458,9 @@ public:
 #if WITH_EDITORONLY_DATA
 	// UObject overrides.
 	virtual void PostEditChangeProperty(FPropertyChangedEvent& PropertyChangedEvent) override;
+	virtual void PostTransacted(const FTransactionObjectEvent& Event) override;
+	virtual void PreEditUndo() override;
+	virtual void PostEditUndo() override;
 	// ~END UObject overrides.
 
 	/**
@@ -499,7 +506,25 @@ public:
 	 * The delegate that gets fired when a property value changes.
 	 * @return A reference to the delegate.
 	 */
-	FMLDeformerModelOnPostEditProperty& OnPostEditChangeProperty() { return PostEditPropertyDelegate; }
+	FMLDeformerModelOnPostEditProperty& OnPostEditChangeProperty()	{ return PostEditPropertyDelegate; }
+
+	/**
+	 * Get the delegate that gets fired when a transaction is done on this model.
+	 * @return A reference to the delegate.
+	 */
+	FMLDeformerModelOnPostTransacted& OnPostTransacted()		{ return PostTransactedDelegate; }
+
+	/**
+	 * Get the delegate that gets fired before we do undo.
+	 * @return A reference to the delegate.
+	 */
+	FMLDeformerModelOnPreEditUndo& OnPreEditUndo()				{ return PreEditUndoDelegate; }
+
+	/**
+	 * Get the delegate that gets fired after we do undo.
+	 * @return A reference to the delegate.
+	 */
+	FMLDeformerModelOnPostEditUndo& OnPostEditUndo()			{ return PostEditUndoDelegate; }
 
 	/**
 	 * Get the visualization settings for this model. These settings are only used in the editor.
@@ -628,6 +653,19 @@ public:
 	 */
 	void SetVizSettings(UMLDeformerVizSettings* VizSettingsObject)			{ VizSettings = VizSettingsObject; }
 
+	/**
+	 * Get the settings that have been configured inside the Training Data Processor tool.
+	 * These settings are stored as part of the ML Deformer model. When this hasn't been set up it will return a nullptr.
+	 * @return A pointer to the training data processor settings, or a nullptr in case no settings have been setup before.
+	 */
+	UMLDeformerTrainingDataProcessorSettings* GetTrainingDataProcessorSettings() const			{ return TrainingDataProcessorSettings; }
+
+	/**
+	 * Set the object that is used to store the settings for the training data processor tool.
+	 * @param Settings The training settings, which will be saved along with the ML Deformer model.
+	 */
+	void SetTrainingDataProcessorSettings(UMLDeformerTrainingDataProcessorSettings* Settings)	{ TrainingDataProcessorSettings = Settings; }
+
 	// Get property names.
 	UE_DEPRECATED(5.4, "This method will be removed.")
 	static FName GetAnimSequencePropertyName()			{ return GET_MEMBER_NAME_CHECKED(UMLDeformerModel, AnimSequence_DEPRECATED); }
@@ -720,6 +758,15 @@ private:
 	/** The delegate that gets fired when a property is being modified. */
 	FMLDeformerModelOnPostEditProperty PostEditPropertyDelegate;
 
+	/** The delegate to capture PostTransacted on this model uobject. */
+	FMLDeformerModelOnPostTransacted PostTransactedDelegate;
+
+	/** The delegate to capture PreEditUndo on this model uobject. */
+	FMLDeformerModelOnPreEditUndo PreEditUndoDelegate;
+
+	/** The delegate to capture PostEditUndo on this model uobject. */
+	FMLDeformerModelOnPostEditUndo PostEditUndoDelegate;
+
 	/** GPU buffers for Vertex Map. */
 	UE::MLDeformer::FVertexMapBuffer VertexMapBuffer;
 
@@ -771,6 +818,10 @@ private:
 	UPROPERTY()
 	TArray<int32> VertexMap;
 
+	/** The settings for the training data processor. */
+	UPROPERTY()
+	TObjectPtr<UMLDeformerTrainingDataProcessorSettings> TrainingDataProcessorSettings;
+
 	/** The skeletal mesh that represents the linear skinned mesh. */
 	UPROPERTY(EditAnywhere, Category = "Base Mesh")
 	TObjectPtr<USkeletalMesh> SkeletalMesh;
@@ -820,7 +871,7 @@ private:
 	UE_DEPRECATED(5.5, "This property has been removed.")
 	UPROPERTY(meta = (DeprecatedProperty, DeprecationMessage = "This property has been removed and isn't used anymore."))
 	float DeltaCutoffLength_DEPRECATED = 30.0f;
-	
+
 	/** 
 	 * Do we want to recover stripped data at the end of the Serialize call during cook?
 	 * This is enabled on default, but some automated tests might disable this.
diff --git a/Engine/Plugins/Animation/MLDeformer/MLDeformerFramework/Source/MLDeformerFramework/Public/MLDeformerModelInstance.h b/Engine/Plugins/Animation/MLDeformer/MLDeformerFramework/Source/MLDeformerFramework/Public/MLDeformerModelInstance.h
--- a/Engine/Plugins/Animation/MLDeformer/MLDeformerFramework/Source/MLDeformerFramework/Public/MLDeformerModelInstance.h
+++ b/Engine/Plugins/Animation/MLDeformer/MLDeformerFramework/Source/MLDeformerFramework/Public/MLDeformerModelInstance.h
@@ -255,7 +255,7 @@ protected:
 
 	/** The skeletal mesh component we work with. This is mainly used for compatibility checks. */
 	UPROPERTY(Transient)
-	TObjectPtr<USkeletalMeshComponent> SkeletalMeshComponent = nullptr;
+	TObjectPtr<USkeletalMeshComponent> SkeletalMeshComponent;
 
 	/** The cached current local space bone transforms for the current frame. */
 	TArray<FTransform> TrainingBoneTransforms;
diff --git a/Engine/Plugins/Animation/MLDeformer/MLDeformerFramework/Source/MLDeformerFramework/Public/MLDeformerMorphModel.h b/Engine/Plugins/Animation/MLDeformer/MLDeformerFramework/Source/MLDeformerFramework/Public/MLDeformerMorphModel.h
--- a/Engine/Plugins/Animation/MLDeformer/MLDeformerFramework/Source/MLDeformerFramework/Public/MLDeformerMorphModel.h
+++ b/Engine/Plugins/Animation/MLDeformer/MLDeformerFramework/Source/MLDeformerFramework/Public/MLDeformerMorphModel.h
@@ -30,6 +30,7 @@ public:
 	virtual bool DoesSupportQualityLevels() const override			{ return true; }	// We can disable morph targets based on the Deformer LOD float value.
 	virtual bool DoesSupportLOD() const override					{ return true; }	// We support deformations in different LOD levels.
 	virtual void Serialize(FArchive& Archive) override;
+	virtual UMLDeformerInputInfo* CreateInputInfo() override;
 	virtual UMLDeformerModelInstance* CreateModelInstance(UMLDeformerComponent* Component) override;
 #if WITH_EDITOR
 	virtual void UpdateMemoryUsage() override;
@@ -319,8 +320,11 @@ private:
 	/**
 	 * The entire set of uncompressed morph target deltas, 3 per vertex, for each morph target, as one flattened buffer.
 	 * So the size of this buffer is: (NumVertsPerMorphTarget * 3 * NumMorphTargets).
+	 * You can get the number of vertices per morph target using Model->GetInputInfo()->GetNumBaseMeshVerts().
+	 * This property is Transient, because we don't want it included in undo/redo, as this can be a few hundred megabytes large.
+	 * Instead, this data is serialized manually inside the Serialize method.
 	 */
-	UPROPERTY()
+	UPROPERTY(NonTransactional)
 	TArray<FVector3f> MorphTargetDeltas;
 
 	/** The number of morph targets. */
diff --git a/Engine/Plugins/Animation/MLDeformer/MLDeformerFramework/Source/MLDeformerFramework/Public/MLDeformerObjectVersion.h b/Engine/Plugins/Animation/MLDeformer/MLDeformerFramework/Source/MLDeformerFramework/Public/MLDeformerObjectVersion.h
--- a/Engine/Plugins/Animation/MLDeformer/MLDeformerFramework/Source/MLDeformerFramework/Public/MLDeformerObjectVersion.h
+++ b/Engine/Plugins/Animation/MLDeformer/MLDeformerFramework/Source/MLDeformerFramework/Public/MLDeformerObjectVersion.h
@@ -22,6 +22,9 @@ namespace UE::MLDeformer
 			// We deprecated the UNeuralMorphMaskInfo and should use FMLDeformerMaskInfo instead.
 			MaskInfoMovedToFramework,
 
+			// Undo support added.
+			UndoSupportAdded,
+
 			// ----- New versions can be added above this line -----
 			VersionPlusOne,
 			LatestVersion = VersionPlusOne - 1
diff --git a/Engine/Plugins/Animation/MLDeformer/MLDeformerFramework/Source/MLDeformerFramework/Public/MLDeformerTrainingDataProcessorSettings.h b/Engine/Plugins/Animation/MLDeformer/MLDeformerFramework/Source/MLDeformerFramework/Public/MLDeformerTrainingDataProcessorSettings.h
new file mode 100644
--- /dev/null
+++ b/Engine/Plugins/Animation/MLDeformer/MLDeformerFramework/Source/MLDeformerFramework/Public/MLDeformerTrainingDataProcessorSettings.h
@@ -0,0 +1,176 @@
+// Copyright Epic Games, Inc. All Rights Reserved.
+
+#pragma once
+
+#include "CoreTypes.h"
+#include "UObject/SoftObjectPtr.h"
+#include "MLDeformerTrainingDataProcessorSettings.generated.h"
+
+class UAnimSequence;
+class USkeleton;
+
+USTRUCT()
+struct MLDEFORMERFRAMEWORK_API FMLDeformerTrainingDataProcessorAnim
+{
+	GENERATED_BODY()
+
+public:
+	/** The animation sequence that we should sample frames from. */
+	UPROPERTY(EditAnywhere, Category = "", meta = (NoResetToDefault))
+	TSoftObjectPtr<UAnimSequence> AnimSequence;
+
+	/** Should we sample frames from this animation sequence? */
+	UPROPERTY(EditAnywhere, Category = "")
+	bool bEnabled = true;
+};
+
+
+USTRUCT()
+struct MLDEFORMERFRAMEWORK_API FMLDeformerTrainingDataProcessorBoneGroup
+{
+	GENERATED_BODY()
+
+public:
+	UPROPERTY()
+	FString GroupName;
+
+	UPROPERTY()
+	TArray<FName> BoneNames;
+};
+
+
+USTRUCT()
+struct MLDEFORMERFRAMEWORK_API FMLDeformerTrainingDataProcessorBoneGroupsList
+{
+	GENERATED_BODY()
+
+public:
+	UPROPERTY()
+	TArray<FMLDeformerTrainingDataProcessorBoneGroup> Groups;
+};
+
+
+USTRUCT()
+struct MLDEFORMERFRAMEWORK_API FMLDeformerTrainingDataProcessorBoneList
+{
+	GENERATED_BODY()
+
+public:
+	UPROPERTY()
+	TArray<FName> BoneNames;
+};
+
+
+/**
+ * The settings for the ML Deformer Training Data Processor tool.
+ * These settings are also stored along with the ML Deformer asset.
+ */
+UCLASS()
+class MLDEFORMERFRAMEWORK_API UMLDeformerTrainingDataProcessorSettings
+	: public UObject
+{
+	GENERATED_BODY()
+
+public:
+	UMLDeformerTrainingDataProcessorSettings() = default;
+
+#if WITH_EDITOR
+	/**
+	 * This returns a pointer to the skeleton used by the SkeletalMesh that is set in the UMLDeformerModel that this settings class is part of.
+	 * If there is no skeletal mesh set on the model, a nullptr will be returned.
+	 * @return A pointer to the skeleton we should be using, or nullptr if not found.
+	 */
+	USkeleton* FindSkeleton() const;
+
+	/**
+	 * Get the number of frames of all enabled input animations combined.
+	 * This will ignore disabled input animations. It will try to load the animation sequences and get their number of frames and sum them up.
+	 * @return The number of frames that would be sampled.
+	 */
+	int32 GetNumInputAnimationFrames() const;
+
+	/**
+	 * Check whether these settings are valid when used with a given skeleton.
+	 * The skeleton we use is likely the skeleton used by the ML Deformer asset.
+	 * @param Skeleton The skeleton to check against. This makes sure the input animations must be the same as this provided skeleton.
+	 * @return Returns true if this configuration is valid, or false otherwise.
+	 */
+	bool IsValid(const USkeleton* Skeleton) const;
+#endif
+	
+public:
+	/**
+	 * The list of animation sequences from which we will grab a pose for each frame.
+	 * This final list of poses then will go through the pose remixing and reduction steps.
+	 */
+	UPROPERTY(EditAnywhere, Category = "Input Animations", meta = (NoResetToDefault))
+	TArray<FMLDeformerTrainingDataProcessorAnim> AnimList;
+
+	/**
+	 * The list of bones that is used to calculate which poses have most variation.
+	 * It does not influence the number of output bones in the output animation.
+	 */
+	UPROPERTY(EditAnywhere, Category = "Frame Reduction", meta = (NoResetToDefault))
+	FMLDeformerTrainingDataProcessorBoneList BoneList;
+
+	/**
+	 * The list of bone groups used during pose remixing.
+	 * Pose remixing will basically shuffle the keyframes for the group of bones as a whole.
+	 * So if you put the left arm bones in a group, the frame numbers of the arm will be shuffled, but
+	 * while keeping all the arm bones at the same frame. The rest of the body can be in a different frame.
+	 * This might sound strange, but it can help the machine learning deformer algorithm produce better reconstructions.
+	 */
+	UPROPERTY(EditAnywhere, Category = "Pose Remixing", meta = (NoResetToDefault))
+	FMLDeformerTrainingDataProcessorBoneGroupsList BoneGroups;
+
+	/**
+	 * The number of output frames you want the output animation sequence to have.
+	 * If this number is higher than the number of frames in the input animations, it will just use the maximum
+	 * number of available frames.
+	 * This can be used to reduce the number of frames for training, which is useful if you want to reduce the size of your training data set.
+	 * This will also change the reordering of frames, so it is NOT just going to cut off some frames at the end.
+	 * It will use a greedy algorithm to find poses that are as different as possible from each other.
+	 * By doing that we hope to get as much pose space coverage as possible.
+	 * This setting is only used when the "Reduce Frames" option is enabled.
+	 */
+	UPROPERTY(EditAnywhere, Category = "Output", meta = (ClampMin="1", EditCondition = "bReduceFrames"))
+	int32 NumOutputFrames = 5000;
+
+	/**
+	 * Specify whether we should reduce the number of frames that we sampled from the input animations, or not.
+	 * This can be used to reduce the number of frames for training, which is useful if you want to reduce the size of your training data set.
+	 * This will also change the reordering of frames, so it is NOT just going to cut off some frames at the end.
+	 * It will use a greedy algorithm to find poses that are as different as possible from each other.
+	 * By doing that we hope to get as much pose space coverage as possible.
+	 * This setting is only used when the "Reduce Frames" option is enabled.
+	 */
+	UPROPERTY(EditAnywhere, Category = "Output")
+	bool bReduceFrames = true;
+
+	/**
+	 * Enable this when you want to perform pose remixing.
+	 * Pose remixing will basically shuffle the keyframes for the group of bones as a whole.
+	 * So if you put the left arm bones in a group, the frame numbers of the arm will be shuffled, but
+	 * while keeping all the arm bones at the same frame. The rest of the body can be in a different frame.
+	 * This might sound strange, but it can help the machine learning deformer algorithm produce better reconstructions.
+	 */
+	UPROPERTY(EditAnywhere, Category = "Output")
+	bool bRemixPoses = true;
+
+	/**
+	 * The random seed that is used during pose remixing.
+	 * The remixing algorithm shuffles frames, which uses randomization.
+	 * In order to always get the same results if we perform pose remixing on the same data set, we have to use a fixed seed.
+	 * Changing the seed could result in a different order of frames.
+	 */
+	UPROPERTY(EditAnywhere, Category = "Output", meta = (ClampMin = "0", EditCondition = "bRemixPoses"))
+	int32 RandomSeed = 777;
+
+	/**
+	 * The output animation sequence.
+	 * This animation sequence will be modified and will be filled with the generated frames.
+	 * Please watch out with picking existing animation sequences, as their contents will be overwritten after you press Generate.
+	 */
+	UPROPERTY(EditAnywhere, Category = "Output", meta = (NoResetToDefault))
+	TSoftObjectPtr<UAnimSequence> OutputAnimSequence;
+};
diff --git a/Engine/Plugins/Animation/MLDeformer/MLDeformerFramework/Source/MLDeformerFrameworkEditor/Private/MLDeformerAssetFactory.cpp b/Engine/Plugins/Animation/MLDeformer/MLDeformerFramework/Source/MLDeformerFrameworkEditor/Private/MLDeformerAssetFactory.cpp
--- a/Engine/Plugins/Animation/MLDeformer/MLDeformerFramework/Source/MLDeformerFrameworkEditor/Private/MLDeformerAssetFactory.cpp
+++ b/Engine/Plugins/Animation/MLDeformer/MLDeformerFramework/Source/MLDeformerFrameworkEditor/Private/MLDeformerAssetFactory.cpp
@@ -37,7 +37,7 @@ UObject* UMLDeformerFactory::FactoryCreateNew(UClass* Class, UObject* InParent,
 
 		// Create the highest priority ML model on default and use that in our asset.
 		const int32 HighestPriorityIndex = ModelRegistry.GetHighestPriorityModelIndex();
-		TObjectPtr<UMLDeformerModel> Model = NewObject<UMLDeformerModel>(DeformerAsset, ModelTypes[HighestPriorityIndex]);
+		TObjectPtr<UMLDeformerModel> Model = NewObject<UMLDeformerModel>(DeformerAsset, ModelTypes[HighestPriorityIndex], NAME_None, RF_Transactional);
 		check(Model);
 
 		Model->Init(DeformerAsset);
diff --git a/Engine/Plugins/Animation/MLDeformer/MLDeformerFramework/Source/MLDeformerFrameworkEditor/Private/MLDeformerEditorModel.cpp b/Engine/Plugins/Animation/MLDeformer/MLDeformerFramework/Source/MLDeformerFrameworkEditor/Private/MLDeformerEditorModel.cpp
--- a/Engine/Plugins/Animation/MLDeformer/MLDeformerFramework/Source/MLDeformerFrameworkEditor/Private/MLDeformerEditorModel.cpp
+++ b/Engine/Plugins/Animation/MLDeformer/MLDeformerFramework/Source/MLDeformerFrameworkEditor/Private/MLDeformerEditorModel.cpp
@@ -52,6 +52,7 @@
 #include "Slate/SceneViewport.h"
 #include "Rendering/SkeletalMeshLODModel.h"
 #include "Async/ParallelFor.h"
+#include "Misc/TransactionObjectEvent.h"
 
 #define LOCTEXT_NAMESPACE "MLDeformerEditorModel"
 
@@ -74,6 +75,21 @@ namespace UE::MLDeformer
 			Model->OnPostEditChangeProperty().Remove(PostEditPropertyDelegateHandle);
 		}
 
+		if (PreEditUndoDelegateHandle.IsValid())
+		{
+			Model->OnPreEditUndo().Remove(PreEditUndoDelegateHandle);
+		}
+
+		if (PostEditUndoDelegateHandle.IsValid())
+		{
+			Model->OnPostEditUndo().Remove(PostEditUndoDelegateHandle);
+		}
+
+		if (PostTransactedDelegateHandle.IsValid())
+		{
+			Model->OnPostTransacted().Remove(PostTransactedDelegateHandle);
+		}
+
 		ClearWorld();
 
 		// Delete the samplers.
@@ -106,8 +122,11 @@ namespace UE::MLDeformer
 		WorkingRange = TRange<double>(0.0, 100.0);
 		PlaybackRange = TRange<double>(0.0, 100.0);
 
-		// Start listening to property changes in the ML Deformer model.
+		// Start listening to property changes and transactions to the ML Deformer model.
 		PostEditPropertyDelegateHandle = Model->OnPostEditChangeProperty().AddRaw(this, &FMLDeformerEditorModel::OnPostEditChangeProperty);
+		PostTransactedDelegateHandle = Model->OnPostTransacted().AddRaw(this, &FMLDeformerEditorModel::OnPostTransacted);
+		PreEditUndoDelegateHandle = Model->OnPreEditUndo().AddRaw(this, &FMLDeformerEditorModel::OnPreEditUndo);
+		PostEditUndoDelegateHandle = Model->OnPostEditUndo().AddRaw(this, &FMLDeformerEditorModel::OnPostEditUndo);
 
 		// Start listening for object modified and object property changed events.
 		// We use this to trigger reinits of the UI and components, when one of our input assets got modified (by property changes, or reimports, etc).
@@ -1312,6 +1331,18 @@ namespace UE::MLDeformer
 		}
 	}
 
+	void FMLDeformerEditorModel::OnPostTransacted(const FTransactionObjectEvent& Event)
+	{
+		const TArray<FName>& PropertiesChanged = Event.GetChangedProperties();
+		for (const FName Name : PropertiesChanged)
+		{
+			if (Event.GetEventType() == ETransactionObjectEventType::UndoRedo)
+			{
+				BroadcastPropertyChanged(Name);
+			}
+		}
+	}
+
 	FMLDeformerEditorActor* FMLDeformerEditorModel::GetVisualizationModeBaseActor() const
 	{
 		const UMLDeformerVizSettings* VizSettings = Model->GetVizSettings();
@@ -3369,6 +3400,18 @@ namespace UE::MLDeformer
 		return TVertexAttributesConstRef<float>();
 	}
 
+	void FMLDeformerEditorModel::BroadcastPropertyChanged(const FName PropertyName)
+	{
+		FProperty* Property = Model->GetClass()->FindPropertyByName(PropertyName);
+		if (Property == nullptr)
+		{
+			UE_LOG(LogMLDeformer, Error, TEXT("Failed to find property '%s' in class '%s'"), *PropertyName.ToString(), *Model->GetName());
+			return;
+		}
+
+		FPropertyChangedEvent Event(Property, EPropertyChangeType::ValueSet);
+		Model->PostEditChangeProperty(Event);
+	}
 }	// namespace UE::MLDeformer
 
 #undef LOCTEXT_NAMESPACE
diff --git a/Engine/Plugins/Animation/MLDeformer/MLDeformerFramework/Source/MLDeformerFrameworkEditor/Private/MLDeformerEditorModule.cpp b/Engine/Plugins/Animation/MLDeformer/MLDeformerFramework/Source/MLDeformerFrameworkEditor/Private/MLDeformerEditorModule.cpp
--- a/Engine/Plugins/Animation/MLDeformer/MLDeformerFramework/Source/MLDeformerFrameworkEditor/Private/MLDeformerEditorModule.cpp
+++ b/Engine/Plugins/Animation/MLDeformer/MLDeformerFramework/Source/MLDeformerFrameworkEditor/Private/MLDeformerEditorModule.cpp
@@ -11,6 +11,13 @@
 #include "Modules/ModuleManager.h"
 #include "EditorModeRegistry.h"
 #include "PropertyEditorDelegates.h"
+#include "Tools/TrainingDataProcessor/TrainingDataProcessorTool.h"
+#include "Tools/TrainingDataProcessor/AnimCustomization.h"
+#include "Tools/TrainingDataProcessor/BoneListCustomization.h"
+#include "Tools/TrainingDataProcessor/BoneGroupsListCustomization.h"
+#include "Tools/TrainingDataProcessor/SBoneListWidget.h"
+#include "Tools/TrainingDataProcessor/SBoneGroupsListWidget.h"
+#include "Tools/TrainingDataProcessor/TrainingDataProcessorSettingsDetailCustomization.h"
 
 #define LOCTEXT_NAMESPACE "MLDeformerEditorModule"
 
@@ -22,27 +29,43 @@ namespace UE::MLDeformer
 	{
 		FEditorModeRegistry::Get().RegisterMode<FMLDeformerEditorMode>(FMLDeformerEditorMode::ModeName, LOCTEXT("MLDeformerEditorMode", "MLDeformer"), FSlateIcon(), false);
 
-		// Register object detail customizations.
+		// Register detail customizations.
 		FPropertyEditorModule& PropertyModule = FModuleManager::GetModuleChecked<FPropertyEditorModule>("PropertyEditor");
 		PropertyModule.RegisterCustomPropertyTypeLayout("MLDeformerCurveReference", FOnGetPropertyTypeCustomizationInstance::CreateStatic(&FMLDeformerCurveReferenceCustomization::MakeInstance) );
 		PropertyModule.RegisterCustomPropertyTypeLayout("MLDeformerGeomCacheTrainingInputAnim", FOnGetPropertyTypeCustomizationInstance::CreateStatic(&FMLDeformerGeomCacheTrainingInputAnimCustomization::MakeInstance) );
+		PropertyModule.RegisterCustomPropertyTypeLayout("MLDeformerTrainingDataProcessorAnim", FOnGetPropertyTypeCustomizationInstance::CreateStatic(&TrainingDataProcessor::FAnimCustomization::MakeInstance) );
+		PropertyModule.RegisterCustomPropertyTypeLayout("MLDeformerTrainingDataProcessorBoneList", FOnGetPropertyTypeCustomizationInstance::CreateStatic(&TrainingDataProcessor::FBoneListCustomization::MakeInstance) );
+		PropertyModule.RegisterCustomPropertyTypeLayout("MLDeformerTrainingDataProcessorBoneGroupsList", FOnGetPropertyTypeCustomizationInstance::CreateStatic(&TrainingDataProcessor::FBoneGroupsListCustomization::MakeInstance) );
+		PropertyModule.RegisterCustomClassLayout("MLDeformerTrainingDataProcessorSettings", FOnGetDetailCustomizationInstance::CreateStatic(&TrainingDataProcessor::FTrainingDataProcessorSettingsDetailCustomization::MakeInstance));
 		PropertyModule.NotifyCustomizationModuleChanged();
 
+		TrainingDataProcessor::RegisterTool();
+
 		SMLDeformerInputWidget::RegisterCommands();
+		TrainingDataProcessor::FBoneListWidgetCommands::Register();
+		TrainingDataProcessor::FBoneGroupsListWidgetCommands::Register();
 	}
 
 	void FMLDeformerEditorModule::ShutdownModule()
 	{
 		FEditorModeRegistry::Get().UnregisterMode(FMLDeformerEditorMode::ModeName);
 		
-		// Unregister object detail customizations.
+		// Unregister detail customizations.
 		if (FModuleManager::Get().IsModuleLoaded(TEXT("PropertyEditor")))
 		{
 			FPropertyEditorModule& PropertyModule = FModuleManager::GetModuleChecked<FPropertyEditorModule>(TEXT("PropertyEditor"));
 			PropertyModule.UnregisterCustomPropertyTypeLayout(TEXT("MLDeformerCurveReference"));
 			PropertyModule.UnregisterCustomPropertyTypeLayout(TEXT("MLDeformerGeomCacheTrainingInputAnim"));
+			PropertyModule.UnregisterCustomPropertyTypeLayout(TEXT("MLDeformerTrainingDataProcessorAnim"));
+			PropertyModule.UnregisterCustomPropertyTypeLayout(TEXT("MLDeformerTrainingDataProcessorBoneList"));
+			PropertyModule.UnregisterCustomPropertyTypeLayout(TEXT("MLDeformerTrainingDataProcessorBoneGroupsList"));
+			PropertyModule.UnregisterCustomClassLayout(TEXT("MLDeformerTrainingDataProcessorSettings"));
 			PropertyModule.NotifyCustomizationModuleChanged();
 		}
+
+		SMLDeformerInputWidget::UnregisterCommands();
+		TrainingDataProcessor::FBoneListWidgetCommands::Unregister();
+		TrainingDataProcessor::FBoneGroupsListWidgetCommands::Unregister();
 	}
 }	// namespace UE::MLDeformer
 
diff --git a/Engine/Plugins/Animation/MLDeformer/MLDeformerFramework/Source/MLDeformerFrameworkEditor/Private/MLDeformerEditorToolkit.cpp b/Engine/Plugins/Animation/MLDeformer/MLDeformerFramework/Source/MLDeformerFrameworkEditor/Private/MLDeformerEditorToolkit.cpp
--- a/Engine/Plugins/Animation/MLDeformer/MLDeformerFramework/Source/MLDeformerFrameworkEditor/Private/MLDeformerEditorToolkit.cpp
+++ b/Engine/Plugins/Animation/MLDeformer/MLDeformerFramework/Source/MLDeformerFrameworkEditor/Private/MLDeformerEditorToolkit.cpp
@@ -47,6 +47,7 @@ namespace UE::MLDeformer
 
 	FMLDeformerEditorToolkit::~FMLDeformerEditorToolkit()
 	{
+		GEditor->UnregisterForUndo(this);
 		ActiveModel.Reset();
 	}
 
@@ -122,6 +123,7 @@ namespace UE::MLDeformer
 			ActiveModel->SetTrainingFrame(ActiveModel->GetModel()->GetVizSettings()->GetTrainingFrameNumber(), false);
 			ActiveModel->SetTestFrame(ActiveModel->GetModel()->GetVizSettings()->GetTestingFrameNumber(), false);
 			ActiveModel->InvalidateDeltas();
+			ActiveModel->GetModel()->GetVizSettings()->SetFlags(RF_Transactional);
 		}
 
 		if (DeformerModel)
@@ -140,6 +142,8 @@ namespace UE::MLDeformer
 
 		ShowNoModelsWarningIfNeeded();
 
+		GEditor->RegisterForUndo(this);
+
 		bIsInitialized = true;
 	}
 
@@ -253,6 +257,12 @@ namespace UE::MLDeformer
 		const bool bShouldCloseWindowAfterMenuSelection = true;
 		FMenuBuilder MenuBuilder(bShouldCloseWindowAfterMenuSelection, InCommandList);
 
+		// Sort the tools menu items before we add them.
+		ToolsMenuExtenders.Sort([this](const TUniquePtr<FToolsMenuExtender>& A, const TUniquePtr<FToolsMenuExtender>& B)
+			{
+				return A->GetMenuEntry(*this).LabelOverride.Get().ToString() < B->GetMenuEntry(*this).LabelOverride.Get().ToString();
+			} );
+
 		for (const TUniquePtr<FToolsMenuExtender>& Extender : ToolsMenuExtenders)
 		{
 			MenuBuilder.AddMenuEntry(Extender->GetMenuEntry(*this));
@@ -323,7 +333,7 @@ namespace UE::MLDeformer
 		// Get the runtime model type based on the index, and create an instance of it.
 		TArray<UClass*> ModelTypes;
 		ModelRegistry.GetRegisteredModels().GenerateKeyArray(ModelTypes);
-		TObjectPtr<UMLDeformerModel> Model = NewObject<UMLDeformerModel>(DeformerAsset, ModelTypes[Index]);
+		TObjectPtr<UMLDeformerModel> Model = NewObject<UMLDeformerModel>(DeformerAsset, ModelTypes[Index], NAME_None, RF_Transactional);
 		check(Model);
 		Model->Init(DeformerAsset);
 		DeformerAsset->SetModel(Model);
@@ -843,6 +853,22 @@ namespace UE::MLDeformer
 		RETURN_QUICK_DECLARE_CYCLE_STAT(FMLDeformerEditorToolkit, STATGROUP_Tickables);
 	}
 
+	void FMLDeformerEditorToolkit::PostUndo(bool bSuccess)
+	{
+		if (ActiveModel && !IsPaintModeActive())
+		{
+			ActiveModel->TriggerInputAssetChanged(true);
+		}
+	}
+
+	void FMLDeformerEditorToolkit::PostRedo(bool bSuccess)
+	{
+		if (ActiveModel && !IsPaintModeActive())
+		{
+			ActiveModel->TriggerInputAssetChanged(true);
+		}
+	}
+
 	void FMLDeformerEditorToolkit::HandlePreviewSceneCreated(const TSharedRef<IPersonaPreviewScene>& InPersonaPreviewScene)
 	{
 		if (ActiveModel)
@@ -1207,8 +1233,8 @@ namespace UE::MLDeformer
 		// Update the actor visibility as the mode changes that.
 		if (ActiveModel.IsValid())
 		{
-			ActiveModel->UpdateActorVisibility();
 			ActiveModel->UpdateLabels();
+			ActiveModel->TriggerInputAssetChanged(false);
 		}
 	}
 
diff --git a/Engine/Plugins/Animation/MLDeformer/MLDeformerFramework/Source/MLDeformerFrameworkEditor/Private/MLDeformerModelDetails.cpp b/Engine/Plugins/Animation/MLDeformer/MLDeformerFramework/Source/MLDeformerFrameworkEditor/Private/MLDeformerModelDetails.cpp
--- a/Engine/Plugins/Animation/MLDeformer/MLDeformerFramework/Source/MLDeformerFrameworkEditor/Private/MLDeformerModelDetails.cpp
+++ b/Engine/Plugins/Animation/MLDeformer/MLDeformerFramework/Source/MLDeformerFrameworkEditor/Private/MLDeformerModelDetails.cpp
@@ -40,7 +40,7 @@ namespace UE::MLDeformer
 		EditorModel = nullptr;
 		if (Objects.Num() == 1 && Objects[0] != nullptr)
 		{
-			Model = static_cast<UMLDeformerModel*>(Objects[0].Get());
+			Model = Cast<UMLDeformerModel>(Objects[0].Get());
 
 			// Get the editor model for this runtime model.
 			FMLDeformerEditorModule& EditorModule = FModuleManager::LoadModuleChecked<FMLDeformerEditorModule>("MLDeformerFrameworkEditor");
diff --git a/Engine/Plugins/Animation/MLDeformer/MLDeformerFramework/Source/MLDeformerFrameworkEditor/Private/SMLDeformerBonePickerDialog.cpp b/Engine/Plugins/Animation/MLDeformer/MLDeformerFramework/Source/MLDeformerFrameworkEditor/Private/SMLDeformerBonePickerDialog.cpp
--- a/Engine/Plugins/Animation/MLDeformer/MLDeformerFramework/Source/MLDeformerFrameworkEditor/Private/SMLDeformerBonePickerDialog.cpp
+++ b/Engine/Plugins/Animation/MLDeformer/MLDeformerFramework/Source/MLDeformerFrameworkEditor/Private/SMLDeformerBonePickerDialog.cpp
@@ -25,7 +25,7 @@ namespace UE::MLDeformer
 		SuperArgs.OnGetChildren(this, &SMLDeformerBonePickerTreeWidget::HandleGetChildrenForTree);
 		SuperArgs.OnMouseButtonDoubleClick(this, &SMLDeformerBonePickerTreeWidget::OnMouseDoubleClicked);
 		SuperArgs.HighlightParentNodesForSelection(false);
-		SuperArgs.AllowInvisibleItemSelection(true);  // Without this we deselect everything when we filter or we collapse.
+		SuperArgs.AllowInvisibleItemSelection(true);
 
 		STreeView<TSharedPtr<FMLDeformerBonePickerTreeElement>>::Construct(SuperArgs);
 	}
@@ -99,7 +99,7 @@ namespace UE::MLDeformer
 		return FText();
 	}
 
-	void SMLDeformerBonePickerDialog::RefreshBoneTreeElements()
+	void SMLDeformerBonePickerDialog::RefreshBoneTreeElements() const
 	{
 		BoneTreeWidget->Clear();
 
diff --git a/Engine/Plugins/Animation/MLDeformer/MLDeformerFramework/Source/MLDeformerFrameworkEditor/Private/SMLDeformerCurvePickerDialog.cpp b/Engine/Plugins/Animation/MLDeformer/MLDeformerFramework/Source/MLDeformerFrameworkEditor/Private/SMLDeformerCurvePickerDialog.cpp
--- a/Engine/Plugins/Animation/MLDeformer/MLDeformerFramework/Source/MLDeformerFrameworkEditor/Private/SMLDeformerCurvePickerDialog.cpp
+++ b/Engine/Plugins/Animation/MLDeformer/MLDeformerFramework/Source/MLDeformerFrameworkEditor/Private/SMLDeformerCurvePickerDialog.cpp
@@ -141,7 +141,7 @@ namespace UE::MLDeformer
 		return FReply::Unhandled();
 	}
 
-	void SMLDeformerCurvePickerDialog::RefreshListElements()
+	void SMLDeformerCurvePickerDialog::RefreshListElements() const 
 	{
 		CurveListView->Clear();
 
diff --git a/Engine/Plugins/Animation/MLDeformer/MLDeformerFramework/Source/MLDeformerFrameworkEditor/Private/SMLDeformerInputBonesWidget.cpp b/Engine/Plugins/Animation/MLDeformer/MLDeformerFramework/Source/MLDeformerFrameworkEditor/Private/SMLDeformerInputBonesWidget.cpp
--- a/Engine/Plugins/Animation/MLDeformer/MLDeformerFramework/Source/MLDeformerFrameworkEditor/Private/SMLDeformerInputBonesWidget.cpp
+++ b/Engine/Plugins/Animation/MLDeformer/MLDeformerFramework/Source/MLDeformerFrameworkEditor/Private/SMLDeformerInputBonesWidget.cpp
@@ -11,6 +11,7 @@
 #include "SMLDeformerInputWidget.h"
 #include "SMLDeformerBonePickerDialog.h"
 #include "Engine/SkeletalMesh.h"
+#include "ScopedTransaction.h"
 
 #define LOCTEXT_NAMESPACE "MLDeformerInputBonesWidget"
 
@@ -118,92 +119,83 @@ namespace UE::MLDeformer
 	{
 		RootElements.Reset();
 
-		const FLinearColor ErrorColor = FMLDeformerEditorStyle::Get().GetColor("MLDeformer.InputsWidget.ErrorColor");
+		const FSlateColor ErrorColor = FMLDeformerEditorStyle::Get().GetColor("MLDeformer.InputsWidget.ErrorColor");
 
 		// If we have no reference skeleton, just add everything as flat list as we don't have hierarchy data.
-		if (RefSkeleton == nullptr)
+		const FString& FilterText = InputBonesWidget->FilterText;
+		if (!RefSkeleton)
 		{
 			for (const FName BoneName : BoneNames)
 			{
-				if (InputBonesWidget->FilterText.IsEmpty() || BoneName.ToString().Contains(InputBonesWidget->FilterText))
+				if (FilterText.IsEmpty() || BoneName.ToString().Contains(FilterText))
 				{
 					TSharedPtr<FMLDeformerInputBoneTreeElement> Element = MakeShared<FMLDeformerInputBoneTreeElement>();
 					Element->Name = BoneName;
-					Element->TextColor = FSlateColor(ErrorColor);
-					AddElement(Element, nullptr);
+					Element->TextColor = ErrorColor;
+					RootElements.Add(Element);
 				}
 			}
 		}
 		else
 		{
-			const int32 NumRefBones = RefSkeleton->GetNum();
-
-			// Show the list of bones that doesn't exist in the ref skeleton first, marked as errors.
-			for (int32 BoneIndex = 0; BoneIndex < BoneNames.Num(); ++BoneIndex)
+			// Add all the bones to some element map so we know the tree element for a given bone name.
+			TMap<FName, TSharedPtr<FMLDeformerInputBoneTreeElement>> NameToElementMap;
+			for (const FName BoneName : BoneNames)
 			{
-				const FName BoneName = BoneNames[BoneIndex];
-				bool bExistsInRefSkel = false;
-
-				for (int32 Index = 0; Index < NumRefBones; ++Index)
+				if (!BoneName.ToString().Contains(FilterText))
 				{
-					if (RefSkeleton->GetBoneName(Index) == BoneName)
-					{
-						bExistsInRefSkel = true;
-						break;
-					}
+					continue;
 				}
 
-				if (!bExistsInRefSkel)
-				{
-					TSharedPtr<FMLDeformerInputBoneTreeElement> Element = MakeShared<FMLDeformerInputBoneTreeElement>();
-					Element->Name = BoneName;
-					Element->TextColor = FSlateColor(ErrorColor);
-					AddElement(Element, nullptr);
-				}
+				const int32 RefSkelBoneIndex = RefSkeleton->FindBoneIndex(BoneName);			
+				TSharedPtr<FMLDeformerInputBoneTreeElement> Element = MakeShared<FMLDeformerInputBoneTreeElement>();
+				Element->Name = BoneName;
+				Element->TextColor = (RefSkelBoneIndex != INDEX_NONE) ? FSlateColor::UseForeground() : ErrorColor;
+				NameToElementMap.Add(BoneName, Element);
 			}
 
-			TMap<FName, TSharedPtr<FMLDeformerInputBoneTreeElement>> NameToElementMap;
-			for (int32 Index = 0; Index < NumRefBones; ++Index)
+			// Handle parents and specify root items.
+			for (auto& Element : NameToElementMap)
 			{
-				const FName RefBoneName = RefSkeleton->GetBoneName(Index);
-				if (!InputBonesWidget->FilterText.IsEmpty() && !RefBoneName.ToString().Contains(InputBonesWidget->FilterText))
+				TSharedPtr<FMLDeformerInputBoneTreeElement> ParentElement = FindParentElementForBone(Element.Key, *RefSkeleton, NameToElementMap);
+				if (ParentElement.IsValid())
 				{
-					continue;
+					ParentElement->Children.Add(Element.Value);
 				}
-
-				// If this is a bone in our list.
-				if (BoneNames.Contains(RefBoneName))
+				else
 				{
-					// First try to find the parent.
-					TSharedPtr<FMLDeformerInputBoneTreeElement> ParentElement;
-					const int32 ParentIndex = RefSkeleton->GetParentIndex(Index);
-					if (ParentIndex != INDEX_NONE)
-					{
-						const FName ParentBoneName = RefSkeleton->GetBoneName(ParentIndex);
-						if (BoneNames.Contains(ParentBoneName))
-						{
-							const TSharedPtr<FMLDeformerInputBoneTreeElement>* ParentElementPtr = NameToElementMap.Find(ParentBoneName);
-							if (ParentElementPtr)
-							{
-								ParentElement = *ParentElementPtr;
-							}
-						}
-					}
-
-					// If we haven't already added this bone.
-					TSharedPtr<FMLDeformerInputBoneTreeElement> Element = MakeShared<FMLDeformerInputBoneTreeElement>();
-					Element->Name = RefBoneName;
-					Element->TextColor = FSlateColor::UseForeground();
-					AddElement(Element, ParentElement);
-					SetItemExpansion(Element, true);
-					NameToElementMap.Add(RefBoneName, Element);
+					RootElements.Add(Element.Value);
 				}
+
+				SetItemExpansion(Element.Value, true);
 			}
 		}
 
 		RecursiveSortElements(TSharedPtr<FMLDeformerInputBoneTreeElement>());
 	}
 
+	TSharedPtr<FMLDeformerInputBoneTreeElement> SMLDeformerInputBoneTreeWidget::FindParentElementForBone(FName BoneName, const FReferenceSkeleton& RefSkeleton, const TMap<FName, TSharedPtr<FMLDeformerInputBoneTreeElement>>& NameToElementMap) const
+	{
+		const int32 BoneIndex = RefSkeleton.FindBoneIndex(BoneName);
+		if (BoneIndex != INDEX_NONE)
+		{
+			int32 ParentIndex = RefSkeleton.GetParentIndex(BoneIndex);
+			while (ParentIndex != INDEX_NONE)
+			{
+				const FName ParentName = RefSkeleton.GetBoneName(ParentIndex);
+				const TSharedPtr<FMLDeformerInputBoneTreeElement>* ParentElementPtr = NameToElementMap.Find(ParentName);
+				if (ParentElementPtr)
+				{
+					return *ParentElementPtr;
+				}
+
+				ParentIndex = RefSkeleton.GetParentIndex(ParentIndex);
+			}
+		}
+
+		return TSharedPtr<FMLDeformerInputBoneTreeElement>();
+	}
+
 	TSharedRef<ITableRow> SMLDeformerInputBoneTreeWidget::MakeTableRowWidget(TSharedPtr<FMLDeformerInputBoneTreeElement> InItem, const TSharedRef<STableViewBase>& OwnerTable)
 	{
 		return InItem->MakeTreeRowWidget(OwnerTable, InItem.ToSharedRef(), SharedThis(this));
@@ -322,13 +314,13 @@ namespace UE::MLDeformer
 		USkeletalMesh* SkelMesh = EditorModel->GetModel()->GetSkeletalMesh();
 		if (SkelMesh)
 		{
-			const FLinearColor HighlightColor = FMLDeformerEditorStyle::Get().GetColor("MLDeformer.InputsWidget.HighlightColor");
+			const FSlateColor HighlightColor = FMLDeformerEditorStyle::Get().GetColor("MLDeformer.InputsWidget.HighlightColor");
 
 			TSharedPtr<SMLDeformerBonePickerDialog> Dialog = 
 				SNew(SMLDeformerBonePickerDialog)
 				.RefSkeleton(&SkelMesh->GetRefSkeleton())
 				.AllowMultiSelect(true)
-				.HighlightBoneNamesColor(FSlateColor(HighlightColor))
+				.HighlightBoneNamesColor(HighlightColor)
 				.HighlightBoneNames(TreeWidget->ExtractAllElementNames())
 				.ExtraWidget(InputWidget->GetExtraBonePickerWidget());
 
@@ -337,10 +329,14 @@ namespace UE::MLDeformer
 			const TArray<FName>& BoneNames = Dialog->GetPickedBoneNames();
 			if (!BoneNames.IsEmpty())
 			{
+				UMLDeformerModel* Model = EditorModel->GetModel();
+				FScopedTransaction Transaction(LOCTEXT("AddInputBonesText", "Add Input Bones"));
+				Model->Modify();
+
 				TArray<FName> BonesAdded;
 				for (const FName BoneName : BoneNames)
 				{
-					TArray<FBoneReference>& ModelBoneIncludeList = EditorModel->GetModel()->GetBoneIncludeList();
+					TArray<FBoneReference>& ModelBoneIncludeList = Model->GetBoneIncludeList();
 					
 					if (!ModelBoneIncludeList.Contains(BoneName))
 					{
@@ -365,7 +361,10 @@ namespace UE::MLDeformer
 
 	void SMLDeformerInputBonesWidget::OnClearInputBones()
 	{
-		EditorModel->GetModel()->GetBoneIncludeList().Empty();
+		UMLDeformerModel* Model = EditorModel->GetModel();
+		FScopedTransaction Transaction(LOCTEXT("ClearInputBonesText", "Clear Input Bones"));
+		Model->Modify();
+		Model->GetBoneIncludeList().Empty();
 		RefreshTree();
 		TreeWidget->ClearSelection();
 		InputWidget->OnClearInputBones();
@@ -379,7 +378,11 @@ namespace UE::MLDeformer
 			return;
 		}
 
-		TArray<FBoneReference>& ModelBoneIncludeList = EditorModel->GetModel()->GetBoneIncludeList();
+		UMLDeformerModel* Model = EditorModel->GetModel();
+		FScopedTransaction Transaction(LOCTEXT("RemoveInputBonesText", "Remove Input Bones"));
+		Model->Modify();
+
+		TArray<FBoneReference>& ModelBoneIncludeList = Model->GetBoneIncludeList();
 		TArray<FName> BoneNamesToRemove;
 		BoneNamesToRemove.Reserve(SelectedItems.Num());
 		for (const TSharedPtr<FMLDeformerInputBoneTreeElement>& Item : SelectedItems)
@@ -398,6 +401,10 @@ namespace UE::MLDeformer
 
 	void SMLDeformerInputBonesWidget::OnAddAnimatedBones()
 	{
+		UMLDeformerModel* Model = EditorModel->GetModel();
+		FScopedTransaction Transaction(LOCTEXT("AddAnimatedInputBonesText", "Add Animated Input Bones"));
+		Model->Modify();
+
 		EditorModel->AddAnimatedBonesToBonesIncludeList();
 		EditorModel->SetResamplingInputOutputsNeeded(true);
 		RefreshTree();
diff --git a/Engine/Plugins/Animation/MLDeformer/MLDeformerFramework/Source/MLDeformerFrameworkEditor/Private/SMLDeformerInputCurvesWidget.cpp b/Engine/Plugins/Animation/MLDeformer/MLDeformerFramework/Source/MLDeformerFrameworkEditor/Private/SMLDeformerInputCurvesWidget.cpp
--- a/Engine/Plugins/Animation/MLDeformer/MLDeformerFramework/Source/MLDeformerFrameworkEditor/Private/SMLDeformerInputCurvesWidget.cpp
+++ b/Engine/Plugins/Animation/MLDeformer/MLDeformerFramework/Source/MLDeformerFrameworkEditor/Private/SMLDeformerInputCurvesWidget.cpp
@@ -12,6 +12,7 @@
 #include "SMLDeformerCurvePickerDialog.h"
 #include "Engine/SkeletalMesh.h"
 #include "Animation/Skeleton.h"
+#include "ScopedTransaction.h"
 
 #define LOCTEXT_NAMESPACE "MLDeformerInputCurvesWidget"
 
@@ -240,12 +241,16 @@ namespace UE::MLDeformer
 			const TArray<FName>& PickerCurveNames = Dialog->GetPickedCurveNames();
 			if (!PickerCurveNames.IsEmpty())
 			{
+				UMLDeformerModel* Model = EditorModel->GetModel();
+				FScopedTransaction Transaction(TEXT("SMLDeformerInputCurvesWidget"), LOCTEXT("AddInputCurvesText", "Add Input Curves"), Model);
+				Model->Modify();
+
 				TArray<FName> CurvesAdded;
 				CurvesAdded.Reserve(PickerCurveNames.Num());
 				for (const FName CurveName : PickerCurveNames)
 				{
 					// Only add when the curve isn't already int he list.
-					TArray<FMLDeformerCurveReference>& ModelCurveIncludeList = EditorModel->GetModel()->GetCurveIncludeList();
+					TArray<FMLDeformerCurveReference>& ModelCurveIncludeList = Model->GetCurveIncludeList();
 					if (ModelCurveIncludeList.Find(CurveName) == INDEX_NONE)
 					{
 						ModelCurveIncludeList.Add(FMLDeformerCurveReference(CurveName));
@@ -269,7 +274,11 @@ namespace UE::MLDeformer
 
 	void SMLDeformerInputCurvesWidget::OnClearInputCurves()
 	{
-		EditorModel->GetModel()->GetCurveIncludeList().Empty();
+		UMLDeformerModel* Model = EditorModel->GetModel();
+		FScopedTransaction Transaction(LOCTEXT("ClearInputCurvesText", "Clear Input Curves"));
+		Model->Modify();
+
+		Model->GetCurveIncludeList().Empty();
 		RefreshList();
 		ListWidget->ClearSelection();
 		InputWidget->OnClearInputCurves();
@@ -283,7 +292,11 @@ namespace UE::MLDeformer
 			return;
 		}
 
-		TArray<FMLDeformerCurveReference>& ModelCurveIncludeList = EditorModel->GetModel()->GetCurveIncludeList();
+		UMLDeformerModel* Model = EditorModel->GetModel();
+		FScopedTransaction Transaction(LOCTEXT("RemoveInputCurvesText", "Remove Input Curves"));
+		Model->Modify();
+
+		TArray<FMLDeformerCurveReference>& ModelCurveIncludeList = Model->GetCurveIncludeList();
 		TArray<FName> CurvesToRemove;
 		CurvesToRemove.Reserve(SelectedItems.Num());
 		for (const TSharedPtr<FMLDeformerInputCurveListElement>& Item : SelectedItems)
@@ -302,6 +315,10 @@ namespace UE::MLDeformer
 
 	void SMLDeformerInputCurvesWidget::OnAddAnimatedCurves()
 	{
+		UMLDeformerModel* Model = EditorModel->GetModel();
+		FScopedTransaction Transaction(LOCTEXT("AddAnimatedInputCurvesText", "Add Animated Input Curves"));
+		Model->Modify();
+
 		EditorModel->AddAnimatedCurvesToCurvesIncludeList();
 		EditorModel->SetResamplingInputOutputsNeeded(true);
 		RefreshList();
diff --git a/Engine/Plugins/Animation/MLDeformer/MLDeformerFramework/Source/MLDeformerFrameworkEditor/Private/SMLDeformerInputWidget.cpp b/Engine/Plugins/Animation/MLDeformer/MLDeformerFramework/Source/MLDeformerFrameworkEditor/Private/SMLDeformerInputWidget.cpp
--- a/Engine/Plugins/Animation/MLDeformer/MLDeformerFramework/Source/MLDeformerFrameworkEditor/Private/SMLDeformerInputWidget.cpp
+++ b/Engine/Plugins/Animation/MLDeformer/MLDeformerFramework/Source/MLDeformerFrameworkEditor/Private/SMLDeformerInputWidget.cpp
@@ -228,6 +228,12 @@ namespace UE::MLDeformer
 		FMLDeformerInputCurvesWidgetCommands::Register();
 	}
 
+	void SMLDeformerInputWidget::UnregisterCommands()
+	{
+		FMLDeformerInputBonesWidgetCommands::Unregister();
+		FMLDeformerInputCurvesWidgetCommands::Unregister();
+	}
+	
 	TSharedPtr<FUICommandList> SMLDeformerInputWidget::GetBonesCommandList() const
 	{ 
 		return BonesCommandList;
diff --git a/Engine/Plugins/Animation/MLDeformer/MLDeformerFramework/Source/MLDeformerFrameworkEditor/Private/Tools/TrainingDataProcessor/AnimCustomization.cpp b/Engine/Plugins/Animation/MLDeformer/MLDeformerFramework/Source/MLDeformerFrameworkEditor/Private/Tools/TrainingDataProcessor/AnimCustomization.cpp
new file mode 100644
--- /dev/null
+++ b/Engine/Plugins/Animation/MLDeformer/MLDeformerFramework/Source/MLDeformerFrameworkEditor/Private/Tools/TrainingDataProcessor/AnimCustomization.cpp
@@ -0,0 +1,175 @@
+// Copyright Epic Games, Inc. All Rights Reserved.
+
+#include "Tools/TrainingDataProcessor/AnimCustomization.h"
+#include "MLDeformerTrainingDataProcessorSettings.h"
+#include "Animation/AnimSequence.h"
+#include "PropertyCustomizationHelpers.h"
+#include "IDetailChildrenBuilder.h"
+#include "Widgets/SBoxPanel.h"
+#include "Widgets/Text/STextBlock.h"
+#include "Widgets/Input/SComboBox.h"
+#include "SWarningOrErrorBox.h"
+#include "DetailWidgetRow.h"
+#include "IPropertyUtilities.h"
+#include "AssetRegistry/AssetRegistryModule.h"
+#include "AssetRegistry/AssetDataTagMap.h"
+#include "Animation/Skeleton.h"
+
+#define LOCTEXT_NAMESPACE "MLDeformerTrainingDataProcessorAnimListCustomize"
+
+namespace UE::MLDeformer::TrainingDataProcessor
+{
+	TSharedRef<IPropertyTypeCustomization> FAnimCustomization::MakeInstance()
+	{
+		return MakeShareable(new FAnimCustomization());
+	}
+
+	void FAnimCustomization::CustomizeHeader(TSharedRef<IPropertyHandle> StructPropertyHandle, class FDetailWidgetRow& HeaderRow,
+	                                         IPropertyTypeCustomizationUtils& StructCustomizationUtils)
+	{
+		const int32 AnimIndex = StructPropertyHandle->GetIndexInArray();
+		check(AnimIndex != INDEX_NONE);
+
+		PropertyUtilities = StructCustomizationUtils.GetPropertyUtilities();
+
+		const UMLDeformerTrainingDataProcessorSettings& Settings = *FindSettings(StructPropertyHandle);
+		const FMLDeformerTrainingDataProcessorAnim& Anim = Settings.AnimList[AnimIndex];
+
+		int32 NumFramesInAnim = 0;
+		const FAssetRegistryModule& AssetRegistryModule = FModuleManager::LoadModuleChecked<FAssetRegistryModule>("AssetRegistry");
+		const FAssetData AssetData = AssetRegistryModule.Get().GetAssetByObjectPath(Anim.AnimSequence.ToSoftObjectPath());
+		if (AssetData.IsValid())
+		{
+			const FAssetTagValueRef Tag = AssetData.TagsAndValues.FindTag(TEXT("Number Of Frames"));
+			if (Tag.IsSet())
+			{
+				NumFramesInAnim = FCString::Atoi(*Tag.GetValue());
+			}
+		}
+
+		HeaderRow
+			.NameContent()
+			.VAlign(EVerticalAlignment::VAlign_Center)
+			.HAlign(EHorizontalAlignment::HAlign_Left)
+			[
+				SNew(STextBlock)
+				.Text(FText::Format(LOCTEXT("HeaderTextFmt", "Animation #{0}"), AnimIndex))
+				.Font(StructCustomizationUtils.GetRegularFont())
+				.ColorAndOpacity(Anim.bEnabled ? FSlateColor::UseForeground() : FSlateColor::UseSubduedForeground())
+			]
+			.ValueContent()
+			[
+				SNew(SHorizontalBox) // Required to work around a text alignment issue. Otherwise the text block will not center align.
+				+ SHorizontalBox::Slot()
+				.VAlign(EVerticalAlignment::VAlign_Center)
+				.HAlign(EHorizontalAlignment::HAlign_Left)
+				[
+					SNew(STextBlock)
+					.Text(FText::Format(LOCTEXT("HeaderTextValue", "{0} Frames"), NumFramesInAnim))
+					.Font(StructCustomizationUtils.GetRegularFont())
+					.ColorAndOpacity(Anim.bEnabled ? FSlateColor::UseForeground() : FSlateColor::UseSubduedForeground())
+				]
+			];
+	}
+
+	void FAnimCustomization::CustomizeChildren(TSharedRef<IPropertyHandle> StructPropertyHandle, IDetailChildrenBuilder& ChildBuilder,
+	                                           IPropertyTypeCustomizationUtils& StructCustomizationUtils)
+	{
+		// Get all the child properties.
+		uint32 NumChildren;
+		StructPropertyHandle->GetNumChildren(NumChildren);
+		TMap<FName, TSharedPtr<IPropertyHandle>> PropertyHandles;
+		PropertyHandles.Reserve(NumChildren);
+		for (uint32 ChildIndex = 0; ChildIndex < NumChildren; ++ChildIndex)
+		{
+			TSharedRef<IPropertyHandle> ChildHandle = StructPropertyHandle->GetChildHandle(ChildIndex).ToSharedRef();
+			const FName PropertyName = ChildHandle->GetProperty()->GetFName();
+			PropertyHandles.Add(PropertyName, ChildHandle);
+		}
+
+		const int32 AnimIndex = StructPropertyHandle->GetIndexInArray();
+		check(AnimIndex != INDEX_NONE);
+
+		TSharedPtr<IPropertyHandle> AnimSequenceHandle = PropertyHandles.FindChecked(TEXT("AnimSequence"));
+		TSharedPtr<IPropertyHandle> EnabledHandle = PropertyHandles.FindChecked(TEXT("bEnabled"));
+		check(AnimSequenceHandle.IsValid());
+		check(EnabledHandle.IsValid());
+
+		AnimSequenceHandle->SetOnPropertyValueChanged(FSimpleDelegate::CreateSP(this, &FAnimCustomization::RefreshProperties));
+		EnabledHandle->SetOnPropertyValueChanged(FSimpleDelegate::CreateSP(this, &FAnimCustomization::RefreshProperties));
+
+		const UMLDeformerTrainingDataProcessorSettings& Settings = *FindSettings(StructPropertyHandle);
+		const FMLDeformerTrainingDataProcessorAnim& Anim = Settings.AnimList[AnimIndex];
+
+		// Add the properties
+		IDetailPropertyRow& AnimRow = ChildBuilder.AddProperty(AnimSequenceHandle.ToSharedRef());
+		AnimRow.EditCondition(Anim.bEnabled, nullptr);
+		AnimRow.CustomWidget()
+	       .NameContent()
+			[
+				AnimRow.GetPropertyHandle()->CreatePropertyNameWidget()
+			]
+			.ValueContent()
+			[
+				SNew(SObjectPropertyEntryBox)
+				.PropertyHandle(AnimRow.GetPropertyHandle())
+				.AllowedClass(UAnimSequence::StaticClass())
+				.ObjectPath(Anim.AnimSequence.ToSoftObjectPath().ToString())
+				.ThumbnailPool(StructCustomizationUtils.GetThumbnailPool())
+				.OnShouldFilterAsset_Lambda([this, &Settings](const FAssetData& AssetData)
+				{
+					const USkeleton* Skeleton = Settings.FindSkeleton();
+					return (!Skeleton || !Skeleton->IsCompatibleForEditor(AssetData));
+				})
+			];
+
+		ChildBuilder.AddProperty(EnabledHandle.ToSharedRef());
+
+		// Verify the geometry cache against the anim sequence.
+		ChildBuilder.AddCustomRow(FText::FromString("SkeletonMismatchWarning"))
+		            .Visibility(TAttribute<EVisibility>::CreateSP(this, &FAnimCustomization::GetAnimErrorVisibility, StructPropertyHandle, AnimIndex))
+		            .WholeRowContent()
+		[
+			SNew(SBox)
+			.Padding(FMargin(0.0f, 4.0f))
+			[
+				SNew(SWarningOrErrorBox)
+				.MessageStyle(EMessageStyle::Error)
+				.Message(LOCTEXT("SkeletonMismatchWarningText",
+				                 "The skeleton does not match the one used by the Skeletal Mesh of the ML Deformer model."))
+			]
+		];
+	}
+
+	EVisibility FAnimCustomization::GetAnimErrorVisibility(const TSharedRef<IPropertyHandle> StructPropertyHandle, int32 AnimIndex) const
+	{
+		const UMLDeformerTrainingDataProcessorSettings& Settings = *FindSettings(StructPropertyHandle);
+		const FMLDeformerTrainingDataProcessorAnim& Anim = Settings.AnimList[AnimIndex];
+
+		const USkeleton* Skeleton = Settings.FindSkeleton();
+		if (Skeleton && Anim.AnimSequence.LoadSynchronous() && Anim.AnimSequence->GetSkeleton() != Skeleton)
+		{
+			return EVisibility::Visible;
+		}
+		return EVisibility::Collapsed;
+	}
+
+	void FAnimCustomization::RefreshProperties() const
+	{
+		if (PropertyUtilities.IsValid())
+		{
+			PropertyUtilities->ForceRefresh();
+		}
+	}
+
+	UMLDeformerTrainingDataProcessorSettings* FAnimCustomization::FindSettings(const TSharedRef<IPropertyHandle>& StructPropertyHandle)
+	{
+		TArray<UObject*> Objects;
+		StructPropertyHandle->GetOuterObjects(Objects);
+		UObject* OwnerObject = !Objects.IsEmpty() ? Objects[0] : nullptr;
+		UMLDeformerTrainingDataProcessorSettings* Settings = Cast<UMLDeformerTrainingDataProcessorSettings>(OwnerObject);
+		return Settings;
+	}
+} // namespace UE::MLDeformer::TrainingDataProcessor
+
+#undef LOCTEXT_NAMESPACE
diff --git a/Engine/Plugins/Animation/MLDeformer/MLDeformerFramework/Source/MLDeformerFrameworkEditor/Private/Tools/TrainingDataProcessor/BoneGroupsListCustomization.cpp b/Engine/Plugins/Animation/MLDeformer/MLDeformerFramework/Source/MLDeformerFrameworkEditor/Private/Tools/TrainingDataProcessor/BoneGroupsListCustomization.cpp
new file mode 100644
--- /dev/null
+++ b/Engine/Plugins/Animation/MLDeformer/MLDeformerFramework/Source/MLDeformerFrameworkEditor/Private/Tools/TrainingDataProcessor/BoneGroupsListCustomization.cpp
@@ -0,0 +1,65 @@
+// Copyright Epic Games, Inc. All Rights Reserved.
+
+#include "Tools/TrainingDataProcessor/BoneGroupsListCustomization.h"
+#include "Tools/TrainingDataProcessor/SBoneGroupsListWidget.h"
+#include "Tools/TrainingDataProcessor/TrainingDataProcessorTool.h"
+#include "MLDeformerTrainingDataProcessorSettings.h"
+#include "MLDeformerModel.h"
+#include "PropertyCustomizationHelpers.h"
+#include "DetailWidgetRow.h"
+#include "IPropertyUtilities.h"
+#include "Animation/Skeleton.h"
+
+#define LOCTEXT_NAMESPACE "MLDeformerTrainingDataProcessorBoneListCustomize"
+
+namespace UE::MLDeformer::TrainingDataProcessor
+{
+	TSharedRef<IPropertyTypeCustomization> FBoneGroupsListCustomization::MakeInstance()
+	{
+		return MakeShareable(new FBoneGroupsListCustomization());
+	}
+
+	void FBoneGroupsListCustomization::CustomizeHeader(TSharedRef<IPropertyHandle> StructPropertyHandle, class FDetailWidgetRow& HeaderRow,
+	                                                   IPropertyTypeCustomizationUtils& StructCustomizationUtils)
+	{
+		StructProperty = StructPropertyHandle;
+		PropertyUtilities = StructCustomizationUtils.GetPropertyUtilities();
+
+		// Get the object that this property is inside.
+		// We will use that to perform transactions for undo/redo.
+		TArray<UObject*> Objects;
+		StructProperty->GetOuterObjects(Objects);
+		UObject* UndoObject = Objects.IsEmpty() ? nullptr : Objects[0];
+
+		HeaderRow
+			.NameContent()
+			.VAlign(EVerticalAlignment::VAlign_Center)
+			.HAlign(EHorizontalAlignment::HAlign_Left)
+			[
+				StructPropertyHandle->CreatePropertyNameWidget()
+			]
+			.ValueContent()
+			.HAlign(EHorizontalAlignment::HAlign_Fill)
+			[
+				SNew(SBoneGroupsListWidget, PropertyUtilities->GetNotifyHook())
+				.Skeleton(FindSkeletonForProperty(StructProperty))
+				.UndoObject(UndoObject)
+				.GetBoneGroups(this, &FBoneGroupsListCustomization::GetBoneGroups)
+			];
+	}
+
+	TArray<FMLDeformerTrainingDataProcessorBoneGroup>* FBoneGroupsListCustomization::GetBoneGroups() const
+	{
+		TArray<void*> RawData;
+		StructProperty->AccessRawData(RawData);
+		if (!RawData.IsEmpty())
+		{
+			FMLDeformerTrainingDataProcessorBoneGroupsList* BoneGroupList = static_cast<FMLDeformerTrainingDataProcessorBoneGroupsList*>(RawData[0]);
+			check(BoneGroupList);
+			return &BoneGroupList->Groups;
+		}
+		return nullptr;
+	}
+} // namespace UE::MLDeformer::TrainingDataProcessor
+
+#undef LOCTEXT_NAMESPACE
diff --git a/Engine/Plugins/Animation/MLDeformer/MLDeformerFramework/Source/MLDeformerFrameworkEditor/Private/Tools/TrainingDataProcessor/BoneListCustomization.cpp b/Engine/Plugins/Animation/MLDeformer/MLDeformerFramework/Source/MLDeformerFrameworkEditor/Private/Tools/TrainingDataProcessor/BoneListCustomization.cpp
new file mode 100644
--- /dev/null
+++ b/Engine/Plugins/Animation/MLDeformer/MLDeformerFramework/Source/MLDeformerFrameworkEditor/Private/Tools/TrainingDataProcessor/BoneListCustomization.cpp
@@ -0,0 +1,66 @@
+// Copyright Epic Games, Inc. All Rights Reserved.
+
+#include "Tools/TrainingDataProcessor/BoneListCustomization.h"
+#include "Tools/TrainingDataProcessor/SBoneListWidget.h"
+#include "Tools/TrainingDataProcessor/TrainingDataProcessorTool.h"
+#include "MLDeformerTrainingDataProcessorSettings.h"
+#include "MLDeformerModel.h"
+#include "PropertyCustomizationHelpers.h"
+#include "Widgets/Input/SComboBox.h"
+#include "DetailWidgetRow.h"
+#include "IPropertyUtilities.h"
+#include "Animation/Skeleton.h"
+
+#define LOCTEXT_NAMESPACE "MLDeformerTrainingDataProcessorBoneListCustomize"
+
+namespace UE::MLDeformer::TrainingDataProcessor
+{
+	TSharedRef<IPropertyTypeCustomization> FBoneListCustomization::MakeInstance()
+	{
+		return MakeShareable(new FBoneListCustomization());
+	}
+
+	void FBoneListCustomization::CustomizeHeader(TSharedRef<IPropertyHandle> StructPropertyHandle, FDetailWidgetRow& HeaderRow,
+	                                             IPropertyTypeCustomizationUtils& StructCustomizationUtils)
+	{
+		StructProperty = StructPropertyHandle;
+		PropertyUtilities = StructCustomizationUtils.GetPropertyUtilities();
+	
+		// Get the object that this property is inside.
+		// We will use that to perform transactions for undo/redo.
+		TArray<UObject*> Objects;
+		StructProperty->GetOuterObjects(Objects);
+		UObject* UndoObject = Objects.IsEmpty() ? nullptr : Objects[0];
+
+		HeaderRow
+			.NameContent()
+			.VAlign(EVerticalAlignment::VAlign_Center)
+			.HAlign(EHorizontalAlignment::HAlign_Left)
+			[
+				StructPropertyHandle->CreatePropertyNameWidget()
+			]
+			.ValueContent()
+			.HAlign(EHorizontalAlignment::HAlign_Fill)
+			[
+				SNew(SBoneListWidget, PropertyUtilities->GetNotifyHook())
+				.Skeleton(FindSkeletonForProperty(StructProperty))
+				.UndoObject(UndoObject)
+				.GetBoneNames(this, &FBoneListCustomization::GetBoneNames)
+			];
+	}
+
+	TArray<FName>* FBoneListCustomization::GetBoneNames() const
+	{
+		TArray<void*> RawData;
+		StructProperty->AccessRawData(RawData);
+		if (!RawData.IsEmpty())
+		{
+			FMLDeformerTrainingDataProcessorBoneList* BoneNames = static_cast<FMLDeformerTrainingDataProcessorBoneList*>(RawData[0]);
+			check(BoneNames);
+			return &BoneNames->BoneNames;
+		}
+		return nullptr;
+	}
+} // namespace UE::MLDeformer::TrainingDataProcessor
+
+#undef LOCTEXT_NAMESPACE
diff --git a/Engine/Plugins/Animation/MLDeformer/MLDeformerFramework/Source/MLDeformerFrameworkEditor/Private/Tools/TrainingDataProcessor/SBoneGroupsListWidget.cpp b/Engine/Plugins/Animation/MLDeformer/MLDeformerFramework/Source/MLDeformerFrameworkEditor/Private/Tools/TrainingDataProcessor/SBoneGroupsListWidget.cpp
new file mode 100644
--- /dev/null
+++ b/Engine/Plugins/Animation/MLDeformer/MLDeformerFramework/Source/MLDeformerFrameworkEditor/Private/Tools/TrainingDataProcessor/SBoneGroupsListWidget.cpp
@@ -0,0 +1,722 @@
+// Copyright Epic Games, Inc. All Rights Reserved.
+
+#include "Tools/TrainingDataProcessor/SBoneGroupsListWidget.h"
+#include "MLDeformerTrainingDataProcessorSettings.h"
+#include "MLDeformerModule.h"
+#include "SMLDeformerBonePickerDialog.h"
+#include "MLDeformerEditorStyle.h"
+#include "Framework/MultiBox/MultiBoxBuilder.h"
+#include "Engine/SkeletalMesh.h"
+#include "ScopedTransaction.h"
+#include "Widgets/Input/SButton.h"
+#include "Widgets/Input/SSearchBox.h"
+#include "Animation/Skeleton.h"
+#include "Framework/Commands/UICommandList.h"
+#include "Editor.h"
+#include "Misc/NotifyHook.h"
+
+#define LOCTEXT_NAMESPACE "BoneGroupsListWidget"
+
+namespace UE::MLDeformer::TrainingDataProcessor
+{
+	FBoneGroupsListWidgetCommands::FBoneGroupsListWidgetCommands()
+		: TCommands<FBoneGroupsListWidgetCommands>
+		("Bone Groups",
+		 LOCTEXT("BoneGroupsCommandDesc", "Bone Groups"),
+		 NAME_None,
+		 FMLDeformerEditorStyle::Get().GetStyleSetName())
+	{
+	}
+
+	void FBoneGroupsListWidgetCommands::RegisterCommands()
+	{
+		UI_COMMAND(CreateGroup, "Create New Group", "Create a new bone group.", EUserInterfaceActionType::Button, FInputChord(EKeys::Insert));
+		UI_COMMAND(DeleteSelectedItems, "Delete Selected Items", "Deletes the selected bones and/or groups.", EUserInterfaceActionType::Button,
+		           FInputChord(EKeys::Delete));
+		UI_COMMAND(ClearGroups, "Clear All Groups", "Clears the entire list of bone groups.", EUserInterfaceActionType::Button, FInputChord());
+		UI_COMMAND(AddBoneToGroup, "Add Bones To Group", "Add new bones to the group.", EUserInterfaceActionType::Button, FInputChord());
+	}
+
+	void SBoneGroupsTreeWidget::Construct(const FArguments& InArgs)
+	{
+		BoneGroupsWidget = InArgs._BoneGroupsWidget;
+
+		STreeView<TSharedPtr<FBoneGroupTreeElement>>::FArguments SuperArgs;
+		SuperArgs.TreeItemsSource(&RootElements);
+		SuperArgs.SelectionMode(ESelectionMode::Multi);
+		SuperArgs.OnGenerateRow(this, &SBoneGroupsTreeWidget::MakeTableRowWidget);
+		SuperArgs.OnGetChildren_Static(&SBoneGroupsTreeWidget::HandleGetChildrenForTree);
+		SuperArgs.OnContextMenuOpening(this, &SBoneGroupsTreeWidget::CreateContextMenuWidget);
+		SuperArgs.HighlightParentNodesForSelection(false);
+		SuperArgs.AllowInvisibleItemSelection(true);
+
+		STreeView<TSharedPtr<FBoneGroupTreeElement>>::Construct(SuperArgs);
+
+		RefreshTree();
+	}
+
+	TSharedPtr<SWidget> SBoneGroupsTreeWidget::CreateContextMenuWidget() const
+	{
+		const FBoneGroupsListWidgetCommands& Actions = FBoneGroupsListWidgetCommands::Get();
+
+		const TSharedPtr<FUICommandList> CommandList = BoneGroupsWidget.IsValid()
+			                                               ? BoneGroupsWidget.Pin()->GetCommandList()
+			                                               : TSharedPtr<FUICommandList>();
+		FMenuBuilder Menu(true, CommandList);
+
+		const TArray<TSharedPtr<FBoneGroupTreeElement>> CurSelectedItems = GetSelectedItems();
+		Menu.BeginSection("BoneGroupActions", LOCTEXT("BoneGroupActionsHeading", "Bone Group Actions"));
+		{
+			if (CurSelectedItems.Num() == 1 && CurSelectedItems[0]->IsGroup())
+			{
+				Menu.AddMenuEntry(Actions.AddBoneToGroup);
+			}
+
+			if (!CurSelectedItems.IsEmpty())
+			{
+				Menu.AddMenuEntry(Actions.DeleteSelectedItems);
+			}
+		}
+		Menu.EndSection();
+
+		return Menu.MakeWidget();
+	}
+
+	int32 SBoneGroupsTreeWidget::GetNumSelectedGroups() const
+	{
+		int32 NumSelectedGroups = 0;
+		for (const TSharedPtr<FBoneGroupTreeElement>& Item : SelectedItems)
+		{
+			if (Item->IsGroup())
+			{
+				NumSelectedGroups++;
+			}
+		}
+		return NumSelectedGroups;
+	}
+
+	void SBoneGroupsTreeWidget::AddElement(const TSharedPtr<FBoneGroupTreeElement>& Element, const TSharedPtr<FBoneGroupTreeElement>& ParentElement)
+	{
+		if (!ParentElement)
+		{
+			RootElements.Add(Element);
+		}
+		else
+		{
+			ParentElement->Children.Add(Element);
+			Element->ParentGroup = ParentElement->AsWeak();
+		}
+	}
+
+	TSharedRef<ITableRow> FBoneGroupTreeElement::MakeTreeRowWidget(const TSharedRef<STableViewBase>& InOwnerTable,
+	                                                               const TSharedRef<FBoneGroupTreeElement>& InTreeElement,
+	                                                               const TSharedPtr<SBoneGroupsTreeWidget>& InTreeWidget)
+	{
+		return SNew(SBoneGroupTreeRowWidget, InOwnerTable, InTreeElement, InTreeWidget);
+	}
+
+	TSharedRef<ITableRow> SBoneGroupsTreeWidget::MakeTableRowWidget(TSharedPtr<FBoneGroupTreeElement> InItem,
+	                                                                const TSharedRef<STableViewBase>& OwnerTable)
+	{
+		return InItem->MakeTreeRowWidget(OwnerTable, InItem.ToSharedRef(), SharedThis(this));
+	}
+
+	void SBoneGroupsTreeWidget::HandleGetChildrenForTree(TSharedPtr<FBoneGroupTreeElement> InItem,
+	                                                     TArray<TSharedPtr<FBoneGroupTreeElement>>& OutChildren)
+	{
+		OutChildren = InItem.Get()->Children;
+	}
+
+	void SBoneGroupsTreeWidget::UpdateTreeElements()
+	{
+		RootElements.Reset();
+		if (!BoneGroupsWidget.IsValid())
+		{
+			return;
+		}
+
+		const FLinearColor ErrorColor = FMLDeformerEditorStyle::Get().GetColor("MLDeformer.InputsWidget.ErrorColor");
+
+		const TSharedPtr<UE::MLDeformer::TrainingDataProcessor::SBoneGroupsListWidget> GroupWidget = BoneGroupsWidget.Pin();
+		TArray<FMLDeformerTrainingDataProcessorBoneGroup>* BoneGroups = GroupWidget->GetBoneGroupsValues();
+		if (!BoneGroups)
+		{
+			return;
+		}
+
+		TWeakObjectPtr<USkeleton> Skeleton = GroupWidget->GetSkeleton();
+		if (!Skeleton.IsValid())
+		{
+			return;
+		}
+
+		const TStrongObjectPtr<USkeleton> StrongSkeleton = Skeleton.Pin();
+		const FReferenceSkeleton& RefSkeleton = StrongSkeleton->GetReferenceSkeleton();
+
+		const FString& FilterText = GroupWidget->GetFilterText();
+
+		const int32 NumGroups = BoneGroups->Num();
+		for (int32 BoneGroupIndex = 0; BoneGroupIndex < NumGroups; ++BoneGroupIndex)
+		{
+			const FMLDeformerTrainingDataProcessorBoneGroup& BoneGroup = (*BoneGroups)[BoneGroupIndex];
+
+			bool bGroupHasVisibleBones = false;
+			if (!FilterText.IsEmpty())
+			{
+				for (const FName BoneName : BoneGroup.BoneNames)
+				{
+					if (BoneName.ToString().Contains(FilterText))
+					{
+						bGroupHasVisibleBones = true;
+						break;
+					}
+				}
+			}
+			else
+			{
+				bGroupHasVisibleBones = true;
+			}
+
+			if (!bGroupHasVisibleBones)
+			{
+				continue;
+			}
+
+			bool bGroupHasError = false;
+			for (const FName BoneName : BoneGroup.BoneNames)
+			{
+				if (BoneName.IsNone() || RefSkeleton.FindBoneIndex(BoneName) == INDEX_NONE)
+				{
+					bGroupHasError = true;
+					break;
+				}
+			}
+
+			// Add the group header.
+			TSharedPtr<FBoneGroupTreeElement> GroupElement = MakeShared<FBoneGroupTreeElement>();
+			GroupElement->Name = BoneGroup.GroupName;
+			GroupElement->TextColor = bGroupHasError ? FSlateColor(ErrorColor) : FSlateColor::UseForeground();
+			GroupElement->GroupIndex = BoneGroupIndex;
+			AddElement(GroupElement, nullptr);
+			SetItemExpansion(GroupElement, true);
+
+			// Add the items in the group.
+			for (int32 BoneIndex = 0; BoneIndex < BoneGroup.BoneNames.Num(); ++BoneIndex)
+			{
+				const FName BoneName = BoneGroup.BoneNames[BoneIndex];
+				if (!FilterText.IsEmpty() && !BoneName.ToString().Contains(FilterText))
+				{
+					continue;
+				}
+
+				const bool bBoneHasError = (BoneName.IsNone() || RefSkeleton.FindBoneIndex(BoneName) == INDEX_NONE);
+
+				TSharedPtr<FBoneGroupTreeElement> ItemElement = MakeShared<FBoneGroupTreeElement>();
+				ItemElement->Name = BoneName.ToString();
+				ItemElement->TextColor = bBoneHasError ? FSlateColor(ErrorColor) : FSlateColor::UseForeground();
+				ItemElement->GroupBoneIndex = BoneIndex;
+				AddElement(ItemElement, GroupElement);
+
+				bGroupHasError |= bBoneHasError;
+			}
+
+			if (bGroupHasError)
+			{
+				GroupElement->TextColor = FSlateColor(ErrorColor);
+			}
+		}
+	}
+
+	FReply SBoneGroupsTreeWidget::OnKeyDown(const FGeometry& InGeometry, const FKeyEvent& InKeyEvent)
+	{
+		TSharedPtr<FUICommandList> CommandList = BoneGroupsWidget.IsValid() ? BoneGroupsWidget.Pin()->GetCommandList() : TSharedPtr<FUICommandList>();
+		if (CommandList.IsValid() && CommandList->ProcessCommandBindings(InKeyEvent))
+		{
+			return FReply::Handled();
+		}
+
+		return STreeView<TSharedPtr<FBoneGroupTreeElement>>::OnKeyDown(InGeometry, InKeyEvent);
+	}
+
+	void SBoneGroupsTreeWidget::RefreshTree()
+	{
+		UpdateTreeElements();
+		RequestTreeRefresh();
+	}
+
+	TSharedPtr<SWidget> SBoneGroupsTreeWidget::CreateContextWidget()
+	{
+		return TSharedPtr<SWidget>();
+	}
+
+	void SBoneGroupTreeRowWidget::Construct(const FArguments& InArgs,
+	                                        const TSharedRef<STableViewBase>& OwnerTable,
+	                                        const TSharedRef<FBoneGroupTreeElement>& InTreeElement,
+	                                        const TSharedPtr<SBoneGroupsTreeWidget>& InTreeView)
+	{
+		WeakTreeElement = InTreeElement;
+		STableRow<TSharedPtr<FBoneGroupTreeElement>>::Construct
+		(
+			STableRow<TSharedPtr<FBoneGroupTreeElement>>::FArguments()
+			.ShowWires(true)
+			.Content()
+			[
+				SNew(STextBlock)
+				.Text(this, &SBoneGroupTreeRowWidget::GetName)
+				.Font_Lambda([this, &InTreeElement]()
+				{
+					if (WeakTreeElement.IsValid() && WeakTreeElement.Pin()->IsGroup())
+					{
+						return FAppStyle::GetFontStyle("BoldFont");
+					}
+					return FAppStyle::GetFontStyle("NormalFont");
+				})
+				.ColorAndOpacity_Lambda([this]()
+				{
+					return WeakTreeElement.IsValid() ? WeakTreeElement.Pin()->TextColor : FSlateColor::UseForeground();
+				})
+			],
+			OwnerTable
+		);
+	}
+
+	FText SBoneGroupTreeRowWidget::GetName() const
+	{
+		if (WeakTreeElement.IsValid())
+		{
+			return FText::FromString(WeakTreeElement.Pin()->Name);
+		}
+		return FText();
+	}
+
+	void SBoneGroupsTreeWidget::Refresh()
+	{
+		RefreshTree();
+	}
+
+	const TArray<TSharedPtr<FBoneGroupTreeElement>>& SBoneGroupsTreeWidget::GetRootElements() const
+	{
+		return RootElements;
+	}
+
+	//------------------------------------------------------------
+	// SBoneGroupsListWidget
+	//------------------------------------------------------------
+	SBoneGroupsListWidget::~SBoneGroupsListWidget()
+	{
+		GEditor->UnregisterForUndo(this);
+	}
+
+	void SBoneGroupsListWidget::Construct(const FArguments& InArgs, FNotifyHook* InNotifyHook)
+	{
+		Skeleton = InArgs._Skeleton;
+		UndoObject = InArgs._UndoObject;
+		GetBoneGroups = InArgs._GetBoneGroups;
+		NotifyHook = InNotifyHook;
+
+		FString ErrorMessage;
+		if (!Skeleton.IsValid())
+		{
+			ErrorMessage += LOCTEXT("SkeletonErrorMessage", "Please pass a Skeleton to your SBoneGroupsListWidget.\n").ToString();
+		}
+
+		if (!GetBoneGroups.IsBound())
+		{
+			ErrorMessage += LOCTEXT("GetBoneGroupsMessage", "GetBoneGroups has not been bound in your SBoneGroupsListWidget.\n").ToString();
+		}
+
+		if (!ErrorMessage.IsEmpty())
+		{
+			ChildSlot
+			[
+				SNew(STextBlock)
+				.Text(FText::FromString(ErrorMessage))
+				.AutoWrapText(true)
+				.ColorAndOpacity(FLinearColor(1.0f, 1.0f, 0.0f))
+			];
+		}
+		else
+		{
+			ChildSlot
+			[
+				SNew(SVerticalBox)
+				+ SVerticalBox::Slot()
+				.Padding(FMargin(0.0f, 4.0f, 0.0f, 2.0f))
+				.AutoHeight()
+				[
+					SNew(SHorizontalBox)
+					+ SHorizontalBox::Slot()
+					.AutoWidth()
+					.FillWidth(1.0f)
+					.VAlign(VAlign_Center)
+					[
+						SNew(SSearchBox)
+						.HintText(LOCTEXT("BoneGroupsSearchBoxHint", "Search"))
+						.OnTextChanged(this, &SBoneGroupsListWidget::OnFilterTextChanged)
+					]
+					+ SHorizontalBox::Slot()
+					.AutoWidth()
+					.Padding(1.0f, 1.0f)
+					.VAlign(VAlign_Center)
+					[
+						SNew(SButton)
+						.ButtonStyle(FAppStyle::Get(), "SimpleButton")
+						.ToolTipText(LOCTEXT("AddButtonToolTip", "Create and add a new bone group."))
+						.OnClicked(this, &SBoneGroupsListWidget::OnAddButtonClicked)
+						.ContentPadding(FMargin(0))
+						[
+							SNew(SImage)
+							.Image(FAppStyle::Get().GetBrush("Icons.PlusCircle"))
+							.ColorAndOpacity(FSlateColor::UseForeground())
+						]
+					]
+					+ SHorizontalBox::Slot()
+					.AutoWidth()
+					.VAlign(VAlign_Center)
+					.Padding(1.0f, 1.0f)
+					[
+						SNew(SButton)
+						.ButtonStyle(FAppStyle::Get(), "SimpleButton")
+						.ToolTipText(LOCTEXT("ClearButtonToolTip", "Clear all bone groups."))
+						.OnClicked(this, &SBoneGroupsListWidget::OnClearButtonClicked)
+						.ContentPadding(FMargin(0))
+						[
+							SNew(SImage)
+							.Image(FAppStyle::Get().GetBrush("Icons.Delete"))
+							.ColorAndOpacity(FSlateColor::UseForeground())
+						]
+					]
+
+				]
+				+ SVerticalBox::Slot()
+				.MinHeight(100.0f)
+				.MaxHeight(300.0f)
+				.Padding(FMargin(0.0f, 0.0f, 4.0f, 4.0f))
+				[
+					SAssignNew(TreeWidget, SBoneGroupsTreeWidget)
+					.BoneGroupsWidget(SharedThis(this))
+				]
+			];
+
+			CommandList = MakeShared<FUICommandList>();
+			BindCommands(CommandList);
+			RefreshTree();
+		}
+
+		GEditor->RegisterForUndo(this);
+	}
+
+	void SBoneGroupsListWidget::BindCommands(const TSharedPtr<FUICommandList>& InCommandList)
+	{
+		const FBoneGroupsListWidgetCommands& GroupCommands = FBoneGroupsListWidgetCommands::Get();
+		InCommandList->MapAction(GroupCommands.CreateGroup, FExecuteAction::CreateSP(this, &SBoneGroupsListWidget::OnCreateBoneGroup));
+		InCommandList->MapAction(GroupCommands.DeleteSelectedItems, FExecuteAction::CreateSP(this, &SBoneGroupsListWidget::OnDeleteSelectedItems));
+		InCommandList->MapAction(GroupCommands.ClearGroups, FExecuteAction::CreateSP(this, &SBoneGroupsListWidget::OnClearBoneGroups));
+		InCommandList->MapAction(GroupCommands.AddBoneToGroup, FExecuteAction::CreateSP(this, &SBoneGroupsListWidget::OnAddBoneToGroup));
+	}
+
+	void SBoneGroupsListWidget::OnFilterTextChanged(const FText& InFilterText)
+	{
+		FilterText = InFilterText.ToString();
+		RefreshTree();
+	}
+
+	void SBoneGroupsListWidget::RefreshTree() const
+	{
+		if (TreeWidget.IsValid())
+		{
+			TreeWidget->RefreshTree();
+		}
+	}
+
+	void SBoneGroupsListWidget::OnCreateBoneGroup() const
+	{
+		if (!Skeleton.IsValid())
+		{
+			return;
+		}
+
+		TArray<FMLDeformerTrainingDataProcessorBoneGroup>* BoneGroups = GetBoneGroupsValues();
+		if (!BoneGroups)
+		{
+			return;
+		}
+
+		const TStrongObjectPtr<USkeleton> StrongSkeleton = Skeleton.Pin();
+		const FReferenceSkeleton& RefSkeleton = StrongSkeleton->GetReferenceSkeleton();
+
+		// Remove bones that are already in a bone group.
+		// This prevents the user from adding bones that exist in multiple groups.
+		TArray<FName> AllowedBones = RefSkeleton.GetRawRefBoneNames();
+		for (const FMLDeformerTrainingDataProcessorBoneGroup& BoneGroup : *BoneGroups)
+		{
+			for (const FName BoneName : BoneGroup.BoneNames)
+			{
+				AllowedBones.Remove(BoneName);
+			}
+		}
+
+		const TSharedPtr<SMLDeformerBonePickerDialog> Dialog =
+			SNew(SMLDeformerBonePickerDialog)
+			.RefSkeleton(&RefSkeleton)
+			.AllowMultiSelect(true)
+			.IncludeList(AllowedBones);
+
+		Dialog->ShowModal();
+
+		const TArray<FName>& BoneNames = Dialog->GetPickedBoneNames();
+		if (BoneNames.IsEmpty())
+		{
+			return;
+		}
+
+		TStrongObjectPtr<UObject> StrongUndoObject = UndoObject.Pin();
+		check(StrongUndoObject);
+
+		FScopedTransaction Transaction(TEXT("SBoneGroupsListWidget"), LOCTEXT("CreateBoneGroupText", "Create Bone Group"), StrongUndoObject.Get());
+		StrongUndoObject->Modify();
+
+		BoneGroups->AddDefaulted();
+		FMLDeformerTrainingDataProcessorBoneGroup& BoneGroup = BoneGroups->Last();
+		BoneGroup.GroupName = TEXT("Bone Group");
+		for (const FName BoneName : BoneNames)
+		{
+			BoneGroup.BoneNames.AddUnique(BoneName);
+		}
+
+		RefreshTree();
+
+		// Trigger an event about adding the group.
+		if (NotifyHook)
+		{
+			FProperty* GroupsProperty = FindFieldChecked<FProperty>(FMLDeformerTrainingDataProcessorBoneGroupsList::StaticStruct(),
+			                                                        GET_MEMBER_NAME_CHECKED(FMLDeformerTrainingDataProcessorBoneGroupsList, Groups));
+			FPropertyChangedEvent GroupsEvent(GroupsProperty, EPropertyChangeType::ArrayAdd);
+			NotifyHook->NotifyPostChange(GroupsEvent, GroupsProperty);
+		}
+	}
+
+	void SBoneGroupsListWidget::OnDeleteSelectedItems() const
+	{
+		TArray<FMLDeformerTrainingDataProcessorBoneGroup>* BoneGroups = GetBoneGroupsValues();
+		if (!BoneGroups)
+		{
+			return;
+		}
+
+		TStrongObjectPtr<UObject> StrongUndoObject = UndoObject.Pin();
+		check(StrongUndoObject);
+
+		const TArray<TSharedPtr<FBoneGroupTreeElement>> CurSelectedItems = TreeWidget->GetSelectedItems();
+
+		TArray<int32> GroupsToRemove;
+		TArray<TSharedPtr<FBoneGroupTreeElement>> SelectedGroups;
+		TArray<TSharedPtr<FBoneGroupTreeElement>> SelectedBones;
+
+		// Check if the selection contains bones and/or groups.
+		for (const TSharedPtr<FBoneGroupTreeElement>& SelectedItem : CurSelectedItems)
+		{
+			if (SelectedItem->IsGroup())
+			{
+				SelectedGroups.Add(SelectedItem);
+				GroupsToRemove.Add(SelectedItem->GroupIndex);
+			}
+			else
+			{
+				SelectedBones.Add(SelectedItem);
+			}
+		}
+
+		FScopedTransaction Transaction(TEXT("SBoneGroupsListWidget"), LOCTEXT("RemoveBoneGroupItemsText", "Remove Bones from Group"),
+		                               StrongUndoObject.Get());
+		StrongUndoObject->Modify();
+
+		// Remove all selected bones.
+		if (!SelectedBones.IsEmpty())
+		{
+			for (const TSharedPtr<FBoneGroupTreeElement>& BoneItem : SelectedBones)
+			{
+				if (BoneItem.IsValid() && BoneItem->ParentGroup.IsValid())
+				{
+					const TSharedPtr<FBoneGroupTreeElement>& ParentGroup = BoneItem->ParentGroup.Pin();
+					FMLDeformerTrainingDataProcessorBoneGroup& ParentBoneGroup = (*BoneGroups)[ParentGroup->GroupIndex];
+					ParentBoneGroup.BoneNames.Remove(FName(BoneItem->Name));
+				}
+			}
+
+			// Trigger an event about changing the bone names for this group.
+			if (NotifyHook)
+			{
+				FProperty* BoneNamesProperty = FindFieldChecked<FProperty>(FMLDeformerTrainingDataProcessorBoneGroup::StaticStruct(),
+				                                                           GET_MEMBER_NAME_CHECKED(
+					                                                           FMLDeformerTrainingDataProcessorBoneGroup, BoneNames));
+				const FPropertyChangedEvent BoneNamesEvent(BoneNamesProperty, EPropertyChangeType::ArrayRemove);
+				NotifyHook->NotifyPostChange(BoneNamesEvent, BoneNamesProperty);
+			}
+		}
+
+		// Sort group indices big to small (back to front).
+		if (!GroupsToRemove.IsEmpty())
+		{
+			GroupsToRemove.Sort([](const int32& A, const int32& B) { return A > B; });
+
+			// Remove the items, back to front.
+			for (const int32 Index : GroupsToRemove)
+			{
+				BoneGroups->RemoveAt(Index);
+			}
+
+			// Trigger an event about groups being deleted.
+			if (NotifyHook)
+			{
+				FProperty* GroupsProperty = FindFieldChecked<FProperty>(FMLDeformerTrainingDataProcessorBoneGroupsList::StaticStruct(),
+				                                                        GET_MEMBER_NAME_CHECKED(
+					                                                        FMLDeformerTrainingDataProcessorBoneGroupsList, Groups));
+				const FPropertyChangedEvent GroupsEvent(GroupsProperty, EPropertyChangeType::ArrayRemove);
+				NotifyHook->NotifyPostChange(GroupsEvent, GroupsProperty);
+			}
+		}
+
+		RefreshTree();
+	}
+
+	void SBoneGroupsListWidget::OnClearBoneGroups() const
+	{
+		const TStrongObjectPtr<UObject> StrongUndoObject = UndoObject.Pin();
+		check(StrongUndoObject);
+
+		FScopedTransaction Transaction(LOCTEXT("ClearBoneGroupItemsText", "Clear Bone Groups"));
+		StrongUndoObject->Modify();
+
+		TArray<FMLDeformerTrainingDataProcessorBoneGroup>* BoneGroups = GetBoneGroupsValues();
+		if (BoneGroups)
+		{
+			BoneGroups->Empty();
+			RefreshTree();
+
+			// Trigger an event about clearing the group.
+			if (NotifyHook)
+			{
+				FProperty* GroupsProperty = FindFieldChecked<FProperty>(
+					FMLDeformerTrainingDataProcessorBoneGroupsList::StaticStruct(),
+					GET_MEMBER_NAME_CHECKED(FMLDeformerTrainingDataProcessorBoneGroupsList, Groups));
+				const FPropertyChangedEvent GroupsEvent(GroupsProperty, EPropertyChangeType::ArrayClear);
+				NotifyHook->NotifyPostChange(GroupsEvent, GroupsProperty);
+			}
+		}
+	}
+
+	void SBoneGroupsListWidget::OnAddBoneToGroup() const
+	{
+		if (!Skeleton.IsValid() || !TreeWidget.IsValid())
+		{
+			return;
+		}
+
+		const TStrongObjectPtr<USkeleton> StrongSkeleton = Skeleton.Pin();
+		const FReferenceSkeleton& RefSkeleton = StrongSkeleton->GetReferenceSkeleton();
+
+		const TStrongObjectPtr<UObject> StrongUndoObject = UndoObject.Pin();
+		check(StrongUndoObject);
+
+		// Find the group we want to add something to.
+		check(TreeWidget->GetSelectedItems().Num() == 1);
+		const int32 GroupIndex = TreeWidget->GetSelectedItems()[0]->GroupIndex;
+		check(GroupIndex != INDEX_NONE);
+
+		TArray<FMLDeformerTrainingDataProcessorBoneGroup>* BoneGroups = GetBoneGroupsValues();
+		FMLDeformerTrainingDataProcessorBoneGroup& BoneGroup = (*BoneGroups)[GroupIndex];
+
+		// Build the highlighted bone names list.
+		TArray<FName> HighlightedBones;
+		HighlightedBones.Reserve(BoneGroup.BoneNames.Num());
+		for (const FName BoneName : BoneGroup.BoneNames)
+		{
+			HighlightedBones.Add(BoneName);
+		}
+
+		const FLinearColor HighlightColor = FMLDeformerEditorStyle::Get().GetColor("MLDeformer.InputsWidget.HighlightColor");
+		const TSharedPtr<SMLDeformerBonePickerDialog> Dialog =
+			SNew(SMLDeformerBonePickerDialog)
+			.RefSkeleton(&RefSkeleton)
+			.AllowMultiSelect(true)
+			.HighlightBoneNames(HighlightedBones)
+			.HighlightBoneNamesColor(HighlightColor);
+
+		Dialog->ShowModal();
+
+		const TArray<FName>& BoneNames = Dialog->GetPickedBoneNames();
+		if (!BoneNames.IsEmpty())
+		{
+			FScopedTransaction Transaction(LOCTEXT("AddBoneGroupItemsText", "Add Bones To Group"));
+			StrongUndoObject->Modify();
+
+			for (const FName BoneName : BoneNames)
+			{
+				BoneGroup.BoneNames.AddUnique(BoneName);
+			}
+
+			RefreshTree();
+
+			// Trigger an event about changing the bone names for this group.
+			if (NotifyHook)
+			{
+				FProperty* BoneNamesProperty = FindFieldChecked<FProperty>(FMLDeformerTrainingDataProcessorBoneGroup::StaticStruct(),
+				                                                           GET_MEMBER_NAME_CHECKED(
+					                                                           FMLDeformerTrainingDataProcessorBoneGroup, BoneNames));
+				const FPropertyChangedEvent BoneNamesEvent(BoneNamesProperty, EPropertyChangeType::ValueSet);
+				NotifyHook->NotifyPostChange(BoneNamesEvent, BoneNamesProperty);
+			}
+		}
+	}
+
+	FReply SBoneGroupsListWidget::OnAddButtonClicked() const
+	{
+		OnCreateBoneGroup();
+		return FReply::Handled();
+	}
+
+	FReply SBoneGroupsListWidget::OnClearButtonClicked() const
+	{
+		OnClearBoneGroups();
+		return FReply::Handled();
+	}
+
+	TSharedPtr<SBoneGroupsTreeWidget> SBoneGroupsListWidget::GetTreeWidget() const
+	{
+		return TreeWidget;
+	}
+
+	TSharedPtr<FUICommandList> SBoneGroupsListWidget::GetCommandList() const
+	{
+		return CommandList;
+	}
+
+	TArray<FMLDeformerTrainingDataProcessorBoneGroup>* SBoneGroupsListWidget::GetBoneGroupsValues() const
+	{
+		check(GetBoneGroups.IsBound());
+		return GetBoneGroups.Execute();
+	}
+
+	TWeakObjectPtr<USkeleton> SBoneGroupsListWidget::GetSkeleton() const
+	{
+		return Skeleton;
+	}
+
+	void SBoneGroupsListWidget::PostUndo(bool bSuccess)
+	{
+		RefreshTree();
+	}
+
+	void SBoneGroupsListWidget::PostRedo(bool bSuccess)
+	{
+		RefreshTree();
+	}
+
+	const FString& SBoneGroupsListWidget::GetFilterText() const
+	{
+		return FilterText;
+	}
+} // namespace UE::MLDeformer::TrainingDataProcessor
+
+#undef LOCTEXT_NAMESPACE
diff --git a/Engine/Plugins/Animation/MLDeformer/MLDeformerFramework/Source/MLDeformerFrameworkEditor/Private/Tools/TrainingDataProcessor/SBoneListWidget.cpp b/Engine/Plugins/Animation/MLDeformer/MLDeformerFramework/Source/MLDeformerFrameworkEditor/Private/Tools/TrainingDataProcessor/SBoneListWidget.cpp
new file mode 100644
--- /dev/null
+++ b/Engine/Plugins/Animation/MLDeformer/MLDeformerFramework/Source/MLDeformerFrameworkEditor/Private/Tools/TrainingDataProcessor/SBoneListWidget.cpp
@@ -0,0 +1,582 @@
+// Copyright Epic Games, Inc. All Rights Reserved.
+
+#include "Tools/TrainingDataProcessor/SBoneListWidget.h"
+#include "MLDeformerTrainingDataProcessorSettings.h"
+#include "MLDeformerModule.h"
+#include "Widgets/Input/SButton.h"
+#include "Widgets/Input/SSearchBox.h"
+#include "Animation/Skeleton.h"
+#include "SMLDeformerBonePickerDialog.h"
+#include "Framework/Commands/UICommandList.h"
+#include "ScopedTransaction.h"
+#include "Editor.h"
+#include "Misc/NotifyHook.h"
+
+#define LOCTEXT_NAMESPACE "MLDeformerTrainingDataProcessorBoneListCustomize"
+
+namespace UE::MLDeformer::TrainingDataProcessor
+{
+	TSharedRef<ITableRow> FBoneTreeWidgetElement::MakeTreeRowWidget(const TSharedRef<STableViewBase>& InOwnerTable,
+	                                                                TSharedRef<FBoneTreeWidgetElement> InTreeElement,
+	                                                                TSharedPtr<SBoneTreeWidget> InTreeWidget)
+	{
+		return SNew(SBoneTreeRowWidget, InOwnerTable, InTreeElement, InTreeWidget);
+	}
+
+	void SBoneTreeWidget::Construct(const FArguments& InArgs)
+	{
+		BoneListWidget = InArgs._BoneListWidget;
+
+		STreeView<TSharedPtr<FBoneTreeWidgetElement>>::FArguments SuperArgs;
+		SuperArgs.TreeItemsSource(&RootElements);
+		SuperArgs.SelectionMode(ESelectionMode::Multi);
+		SuperArgs.OnGenerateRow(this, &SBoneTreeWidget::MakeTableRowWidget);
+		SuperArgs.OnGetChildren_Static(&SBoneTreeWidget::HandleGetChildrenForTree);
+		SuperArgs.OnContextMenuOpening(this, &SBoneTreeWidget::OnContextMenuOpening);
+		SuperArgs.AllowInvisibleItemSelection(true);
+
+		STreeView<TSharedPtr<FBoneTreeWidgetElement>>::Construct(SuperArgs);
+	}
+
+	TSharedPtr<SWidget> SBoneTreeWidget::OnContextMenuOpening() const
+	{
+		const FBoneListWidgetCommands& Actions = FBoneListWidgetCommands::Get();
+		TSharedPtr<FUICommandList> CommandList = BoneListWidget.IsValid() ? BoneListWidget.Pin()->GetCommandList() : TSharedPtr<FUICommandList>();
+		FMenuBuilder Menu(true, CommandList);
+		Menu.BeginSection("BoneActions", LOCTEXT("BoneActionsHeading", "Bone Actions"));
+		{
+			Menu.AddMenuEntry(Actions.AddBones);
+
+			if (!GetSelectedItems().IsEmpty())
+			{
+				Menu.AddMenuEntry(Actions.RemoveBones);
+			}
+
+			if (!RootElements.IsEmpty())
+			{
+				Menu.AddMenuEntry(Actions.ClearBones);
+			}
+		}
+		Menu.EndSection();
+
+		return Menu.MakeWidget();
+	}
+
+	void SBoneTreeWidget::HandleGetChildrenForTree(TSharedPtr<FBoneTreeWidgetElement> InItem, TArray<TSharedPtr<FBoneTreeWidgetElement>>& OutChildren)
+	{
+		OutChildren = InItem.Get()->Children;
+	}
+
+	void SBoneTreeWidget::RecursiveSortElements(const TSharedPtr<FBoneTreeWidgetElement>& Element)
+	{
+		if (Element.IsValid())
+		{
+			Element->Children.Sort(
+				[](const TSharedPtr<FBoneTreeWidgetElement>& ItemA, const TSharedPtr<FBoneTreeWidgetElement>& ItemB)
+				{
+					return (ItemA->Name.ToString() < ItemB->Name.ToString());
+				});
+
+			for (const TSharedPtr<FBoneTreeWidgetElement>& Child : Element->Children)
+			{
+				RecursiveSortElements(Child);
+			}
+		}
+		else // We need to sort the root level.
+		{
+			RootElements.Sort(
+				[](const TSharedPtr<FBoneTreeWidgetElement>& ItemA, const TSharedPtr<FBoneTreeWidgetElement>& ItemB)
+				{
+					return (ItemA->Name.ToString() < ItemB->Name.ToString());
+				});
+
+			for (const TSharedPtr<FBoneTreeWidgetElement>& Child : RootElements)
+			{
+				RecursiveSortElements(Child);
+			}
+		}
+	}
+
+	void SBoneTreeWidget::RefreshElements(const TArray<FName>& BoneNames, const FReferenceSkeleton* RefSkeleton, const FString& FilterText)
+	{
+		RootElements.Reset();
+
+		const FSlateColor ErrorColor = FMLDeformerEditorStyle::Get().GetColor("MLDeformer.InputsWidget.ErrorColor");
+
+		// If we have no reference skeleton, just add everything as flat list as we don't have hierarchy data.
+		if (!RefSkeleton)
+		{
+			for (const FName BoneName : BoneNames)
+			{
+				if (FilterText.IsEmpty() || BoneName.ToString().Contains(FilterText))
+				{
+					TSharedPtr<FBoneTreeWidgetElement> Element = MakeShared<FBoneTreeWidgetElement>();
+					Element->Name = BoneName;
+					Element->TextColor = ErrorColor;
+					RootElements.Add(Element);
+				}
+			}
+		}
+		else
+		{
+			// Add all the bones to some element map so we know the tree element for a given bone name.
+			TMap<FName, TSharedPtr<FBoneTreeWidgetElement>> NameToElementMap;
+			for (const FName BoneName : BoneNames)
+			{
+				if (!BoneName.ToString().Contains(FilterText))
+				{
+					continue;
+				}
+
+				const int32 RefSkelBoneIndex = RefSkeleton->FindBoneIndex(BoneName);
+				TSharedPtr<FBoneTreeWidgetElement> Element = MakeShared<FBoneTreeWidgetElement>();
+				Element->Name = BoneName;
+				Element->TextColor = (RefSkelBoneIndex != INDEX_NONE) ? FSlateColor::UseForeground() : ErrorColor;
+				NameToElementMap.Add(BoneName, Element);
+			}
+
+			// Handle parents and specify root items.
+			for (auto& Element : NameToElementMap)
+			{
+				TSharedPtr<FBoneTreeWidgetElement> ParentElement = FindParentElementForBone(Element.Key, *RefSkeleton, NameToElementMap);
+				if (ParentElement.IsValid())
+				{
+					ParentElement->Children.Add(Element.Value);
+				}
+				else
+				{
+					RootElements.Add(Element.Value);
+				}
+
+				SetItemExpansion(Element.Value, true);
+			}
+		}
+
+		RecursiveSortElements(TSharedPtr<FBoneTreeWidgetElement>());
+	}
+
+	TSharedPtr<FBoneTreeWidgetElement> SBoneTreeWidget::FindParentElementForBone(FName BoneName, const FReferenceSkeleton& RefSkeleton,
+	                                                                             const TMap<FName, TSharedPtr<FBoneTreeWidgetElement>>&
+	                                                                             NameToElementMap)
+	{
+		const int32 BoneIndex = RefSkeleton.FindBoneIndex(BoneName);
+
+		if (BoneIndex != INDEX_NONE)
+		{
+			int32 ParentIndex = RefSkeleton.GetParentIndex(BoneIndex);
+			while (ParentIndex != INDEX_NONE)
+			{
+				const FName ParentName = RefSkeleton.GetBoneName(ParentIndex);
+				const TSharedPtr<FBoneTreeWidgetElement>* ParentElementPtr = NameToElementMap.Find(ParentName);
+				if (ParentElementPtr)
+				{
+					return *ParentElementPtr;
+				}
+
+				ParentIndex = RefSkeleton.GetParentIndex(ParentIndex);
+			}
+		}
+
+		return TSharedPtr<FBoneTreeWidgetElement>();
+	}
+
+	TSharedRef<ITableRow> SBoneTreeWidget::MakeTableRowWidget(TSharedPtr<FBoneTreeWidgetElement> InItem, const TSharedRef<STableViewBase>& OwnerTable)
+	{
+		return InItem->MakeTreeRowWidget(OwnerTable, InItem.ToSharedRef(), SharedThis(this));
+	}
+
+	FReply SBoneTreeWidget::OnKeyDown(const FGeometry& InGeometry, const FKeyEvent& InKeyEvent)
+	{
+		const TSharedPtr<FUICommandList> CommandList = BoneListWidget.IsValid()
+			                                               ? BoneListWidget.Pin()->GetCommandList()
+			                                               : TSharedPtr<FUICommandList>();
+		if (CommandList.IsValid() && CommandList->ProcessCommandBindings(InKeyEvent))
+		{
+			return FReply::Handled();
+		}
+
+		return STreeView<TSharedPtr<FBoneTreeWidgetElement>>::OnKeyDown(InGeometry, InKeyEvent);
+	}
+
+	void SBoneTreeWidget::RecursiveAddNames(const FBoneTreeWidgetElement& Element, TArray<FName>& OutNames)
+	{
+		OutNames.Add(Element.Name);
+		for (const TSharedPtr<FBoneTreeWidgetElement>& ChildElement : Element.Children)
+		{
+			RecursiveAddNames(*ChildElement, OutNames);
+		}
+	}
+
+	TArray<FName> SBoneTreeWidget::ExtractAllElementNames() const
+	{
+		TArray<FName> Names;
+		for (const TSharedPtr<FBoneTreeWidgetElement>& Element : RootElements)
+		{
+			RecursiveAddNames(*Element, Names);
+		}
+		return MoveTemp(Names);
+	}
+
+	void SBoneTreeRowWidget::Construct(const FArguments& InArgs, const TSharedRef<STableViewBase>& OwnerTable,
+	                                   const TSharedRef<FBoneTreeWidgetElement>& InTreeElement, const TSharedPtr<SBoneTreeWidget>& InTreeView)
+	{
+		TreeElement = InTreeElement;
+
+		STableRow<TSharedPtr<FBoneTreeWidgetElement>>::Construct
+		(
+			STableRow<TSharedPtr<FBoneTreeWidgetElement>>::FArguments()
+			.ShowWires(true)
+			.Content()
+			[
+				SNew(STextBlock)
+				.Text(this, &SBoneTreeRowWidget::GetName)
+				.ColorAndOpacity_Lambda
+				(
+					[this]()
+					{
+						return TreeElement.IsValid() ? TreeElement.Pin()->TextColor : FSlateColor::UseForeground();
+					}
+				)
+			],
+			OwnerTable
+		);
+	}
+
+	FText SBoneTreeRowWidget::GetName() const
+	{
+		if (TreeElement.IsValid())
+		{
+			return FText::FromName(TreeElement.Pin()->Name);
+		}
+		return FText();
+	}
+
+	SBoneListWidget::~SBoneListWidget()
+	{
+		GEditor->UnregisterForUndo(this);
+	}
+
+	void SBoneListWidget::Construct(const FArguments& InArgs, FNotifyHook* InNotifyHook)
+	{
+		NotifyHook = InNotifyHook;
+
+		Skeleton = InArgs._Skeleton;
+		UndoObject = InArgs._UndoObject;
+		GetBoneNames = InArgs._GetBoneNames;
+		OnBonesAdded = InArgs._OnBonesAdded;
+		OnBonesRemoved = InArgs._OnBonesRemoved;
+		OnBonesCleared = InArgs._OnBonesCleared;
+
+		FString ErrorMessage;
+		if (!Skeleton.IsValid())
+		{
+			ErrorMessage += LOCTEXT("SkeletonErrorMessage", "Please pass a Skeleton to your SBoneListWidget.\n").ToString();
+		}
+
+		if (!GetBoneNames.IsBound())
+		{
+			ErrorMessage += LOCTEXT("GetBoneNamesMessage", "GetBoneNames has not been bound in your SBoneListWidget.\n").ToString();
+		}
+
+		if (!ErrorMessage.IsEmpty())
+		{
+			ChildSlot
+			[
+				SNew(STextBlock)
+				.Text(FText::FromString(ErrorMessage))
+				.AutoWrapText(true)
+				.ColorAndOpacity(FLinearColor(1.0f, 1.0f, 0.0f))
+			];
+		}
+		else
+		{
+			ChildSlot
+			[
+				SNew(SVerticalBox)
+				+ SVerticalBox::Slot()
+				.Padding(FMargin(0.0f, 4.0f, 0.0f, 2.0f))
+				.AutoHeight()
+				[
+					SNew(SHorizontalBox)
+					+ SHorizontalBox::Slot()
+					.AutoWidth()
+					.FillWidth(1.0f)
+					.VAlign(VAlign_Center)
+					[
+						SNew(SSearchBox)
+						.HintText(LOCTEXT("BonesSearchBoxHint", "Search Bones"))
+						.OnTextChanged(this, &SBoneListWidget::OnFilterTextChanged)
+					]
+					+ SHorizontalBox::Slot()
+					.AutoWidth()
+					.Padding(1.0f, 1.0f)
+					.VAlign(VAlign_Center)
+					[
+						SNew(SButton)
+						.ButtonStyle(FAppStyle::Get(), "SimpleButton")
+						.ToolTipText(LOCTEXT("AddButtonToolTip", "Add bones to the list."))
+						.OnClicked(this, &SBoneListWidget::OnAddBonesButtonClicked)
+						.ContentPadding(FMargin(0))
+						[
+							SNew(SImage)
+							.Image(FAppStyle::Get().GetBrush("Icons.PlusCircle"))
+							.ColorAndOpacity(FSlateColor::UseForeground())
+						]
+					]
+					+ SHorizontalBox::Slot()
+					.AutoWidth()
+					.VAlign(VAlign_Center)
+					.Padding(1.0f, 1.0f)
+					[
+						SNew(SButton)
+						.ButtonStyle(FAppStyle::Get(), "SimpleButton")
+						.ToolTipText(LOCTEXT("ClearButtonToolTip", "Clear the bone list."))
+						.OnClicked(this, &SBoneListWidget::OnClearBonesButtonClicked)
+						.ContentPadding(FMargin(0))
+						[
+							SNew(SImage)
+							.Image(FAppStyle::Get().GetBrush("Icons.Delete"))
+							.ColorAndOpacity(FSlateColor::UseForeground())
+						]
+					]
+				]
+				+ SVerticalBox::Slot()
+				.MinHeight(100.0f)
+				.MaxHeight(300.0f)
+				.Padding(FMargin(0.0f, 0.0f, 4.0f, 4.0f))
+				[
+					SAssignNew(TreeWidget, SBoneTreeWidget)
+					.BoneListWidget(SharedThis(this))
+				]
+			];
+
+			CommandList = MakeShared<FUICommandList>();
+			BindCommands(CommandList);
+			RefreshTree();
+		}
+
+		GEditor->RegisterForUndo(this);
+	}
+
+	FReply SBoneListWidget::OnAddBonesButtonClicked() const
+	{
+		OnAddBones();
+		return FReply::Handled();
+	}
+
+	FReply SBoneListWidget::OnClearBonesButtonClicked() const
+	{
+		OnClearBones();
+		return FReply::Handled();
+	}
+
+	void SBoneListWidget::OnFilterTextChanged(const FText& InFilterText)
+	{
+		FilterText = InFilterText.ToString();
+		RefreshTree();
+	}
+
+	void SBoneListWidget::BindCommands(const TSharedPtr<FUICommandList>& InCommandList)
+	{
+		const FBoneListWidgetCommands& Commands = FBoneListWidgetCommands::Get();
+		InCommandList->MapAction(Commands.AddBones, FExecuteAction::CreateSP(this, &SBoneListWidget::OnAddBones));
+		InCommandList->MapAction(Commands.RemoveBones, FExecuteAction::CreateSP(this, &SBoneListWidget::OnRemoveBones));
+		InCommandList->MapAction(Commands.ClearBones, FExecuteAction::CreateSP(this, &SBoneListWidget::OnClearBones));
+	}
+
+	void SBoneListWidget::OnAddBones() const
+	{
+		if (!GetBoneNames.IsBound())
+		{
+			UE_LOG(LogMLDeformer, Warning, TEXT("Please set a GetBoneNames to your SBoneListWidget when creating your SBoneListWidget."));
+			return;
+		}
+
+		if (!Skeleton.IsValid())
+		{
+			UE_LOG(LogMLDeformer, Warning, TEXT("No skeleton is available to pick bones from"));
+			return;
+		}
+
+		const FSlateColor HighlightColor = FMLDeformerEditorStyle::Get().GetColor("MLDeformer.InputsWidget.HighlightColor");
+		const FReferenceSkeleton* RefSkeleton = const_cast<FReferenceSkeleton*>(&Skeleton->GetReferenceSkeleton());
+
+		TArray<FName>* BoneNames = GetBoneNames.Execute();
+		if (!BoneNames)
+		{
+			UE_LOG(LogMLDeformer, Warning, TEXT("There are no bones that can be added."));
+			return;
+		}
+
+		TSharedPtr<SMLDeformerBonePickerDialog> Dialog =
+			SNew(SMLDeformerBonePickerDialog)
+			.RefSkeleton(RefSkeleton)
+			.AllowMultiSelect(true)
+			.HighlightBoneNamesColor(HighlightColor)
+			.HighlightBoneNames(*BoneNames);
+
+		Dialog->ShowModal();
+
+		const TArray<FName>& PickedBoneNames = Dialog->GetPickedBoneNames();
+		if (!PickedBoneNames.IsEmpty())
+		{
+			FScopedTransaction Transaction(LOCTEXT("AddBonesText", "Add Bones"));
+			if (UndoObject.IsValid())
+			{
+				UndoObject->Modify();
+			}
+
+			TArray<FName> BonesAdded;
+			for (const FName PickedBoneName : PickedBoneNames)
+			{
+				if (!BoneNames->Contains(PickedBoneName))
+				{
+					BoneNames->Add(PickedBoneName);
+					BonesAdded.Add(PickedBoneName);
+				}
+			}
+
+			// Broadcast that bones got added.
+			if (!BonesAdded.IsEmpty())
+			{
+				OnBonesAdded.ExecuteIfBound(BonesAdded);
+				NotifyPropertyChanged();
+			}
+
+			RefreshTree();
+		}
+	}
+
+	void SBoneListWidget::OnClearBones() const
+	{
+		if (!GetBoneNames.IsBound())
+		{
+			UE_LOG(LogMLDeformer, Warning, TEXT("Please set a GetBoneNames to your SBoneListWidget when creating your SBoneListWidget."));
+			return;
+		}
+
+		FScopedTransaction Transaction(LOCTEXT("ClearBonesText", "Clear Bones"));
+		if (UndoObject.IsValid())
+		{
+			UndoObject->Modify();
+		}
+
+		TArray<FName>* BoneNames = GetBoneNames.Execute();
+		if (BoneNames && !BoneNames->IsEmpty())
+		{
+			BoneNames->Empty();
+			OnBonesCleared.ExecuteIfBound();
+			RefreshTree();
+			TreeWidget->ClearSelection();
+			NotifyPropertyChanged();
+		}
+	}
+
+	void SBoneListWidget::OnRemoveBones() const
+	{
+		if (!GetBoneNames.IsBound())
+		{
+			UE_LOG(LogMLDeformer, Warning, TEXT("Please set a GetBoneNames to your SBoneListWidget when creating your SBoneListWidget."));
+			return;
+		}
+
+		const TArray<TSharedPtr<FBoneTreeWidgetElement>> SelectedItems = TreeWidget->GetSelectedItems();
+		if (SelectedItems.IsEmpty())
+		{
+			return;
+		}
+
+		FScopedTransaction Transaction(LOCTEXT("RemoveBonesText", "Remove Bones"));
+		if (UndoObject.IsValid())
+		{
+			UndoObject->Modify();
+		}
+
+		TArray<FName>* BoneNames = GetBoneNames.Execute();
+		if (BoneNames)
+		{
+			TArray<FName> BoneNamesRemoved;
+			for (const TSharedPtr<FBoneTreeWidgetElement>& Item : SelectedItems)
+			{
+				BoneNames->Remove(Item->Name);
+				BoneNamesRemoved.Add(Item->Name);
+			}
+
+			if (!BoneNamesRemoved.IsEmpty())
+			{
+				OnBonesRemoved.ExecuteIfBound(BoneNamesRemoved);
+				NotifyPropertyChanged();
+			}
+		}
+
+		RefreshTree();
+		TreeWidget->ClearSelection();
+	}
+
+	void SBoneListWidget::RefreshTree() const
+	{
+		const FReferenceSkeleton* RefSkeleton = Skeleton.Get() ? &Skeleton->GetReferenceSkeleton() : nullptr;
+
+		if (GetBoneNames.IsBound() && GetBoneNames.Execute())
+		{
+			const TArray<FName>* BoneNames = GetBoneNames.Execute();
+			TreeWidget->RefreshElements(*BoneNames, RefSkeleton, FilterText);
+		}
+		else
+		{
+			TreeWidget->RefreshElements(TArray<FName>(), RefSkeleton, FilterText);
+		}
+		TreeWidget->RequestTreeRefresh();
+	}
+
+	FBoneListWidgetCommands::FBoneListWidgetCommands()
+		: TCommands<FBoneListWidgetCommands>("Bone List", LOCTEXT("BoneListDesc", "Modify Bone List"), NAME_None,
+		                                     FMLDeformerEditorStyle::Get().GetStyleSetName())
+	{
+	}
+
+	void FBoneListWidgetCommands::RegisterCommands()
+	{
+		UI_COMMAND(AddBones, "Add Bones", "Add bones to the list.", EUserInterfaceActionType::Button, FInputChord(EKeys::Insert));
+		UI_COMMAND(RemoveBones, "Delete Selected", "Deletes the selected bones from the list.", EUserInterfaceActionType::Button,
+		           FInputChord(EKeys::Delete));
+		UI_COMMAND(ClearBones, "Clear List", "Clears the entire list.", EUserInterfaceActionType::Button, FInputChord());
+	}
+
+	void SBoneListWidget::Refresh() const
+	{
+		RefreshTree();
+	}
+
+	TSharedPtr<SBoneTreeWidget> SBoneListWidget::GetTreeWidget() const
+	{
+		return TreeWidget;
+	}
+
+	TSharedPtr<FUICommandList> SBoneListWidget::GetCommandList() const
+	{
+		return CommandList;
+	}
+
+	void SBoneListWidget::NotifyPropertyChanged() const
+	{
+		if (NotifyHook)
+		{
+			FProperty* BoneNamesProperty = FindFieldChecked<FProperty>(FMLDeformerTrainingDataProcessorBoneList::StaticStruct(),
+			                                                           GET_MEMBER_NAME_CHECKED(FMLDeformerTrainingDataProcessorBoneList, BoneNames));
+			NotifyHook->NotifyPostChange(FPropertyChangedEvent(BoneNamesProperty, EPropertyChangeType::ValueSet), BoneNamesProperty);
+		}
+	}
+
+	void SBoneListWidget::PostUndo(bool bSuccess)
+	{
+		RefreshTree();
+	}
+
+	void SBoneListWidget::PostRedo(bool bSuccess)
+	{
+		RefreshTree();
+	}
+} // namespace UE::MLDeformer::TrainingDataProcessor
+
+#undef LOCTEXT_NAMESPACE
diff --git a/Engine/Plugins/Animation/MLDeformer/MLDeformerFramework/Source/MLDeformerFrameworkEditor/Private/Tools/TrainingDataProcessor/STrainingDataProcessorMainWidget.cpp b/Engine/Plugins/Animation/MLDeformer/MLDeformerFramework/Source/MLDeformerFrameworkEditor/Private/Tools/TrainingDataProcessor/STrainingDataProcessorMainWidget.cpp
new file mode 100644
--- /dev/null
+++ b/Engine/Plugins/Animation/MLDeformer/MLDeformerFramework/Source/MLDeformerFrameworkEditor/Private/Tools/TrainingDataProcessor/STrainingDataProcessorMainWidget.cpp
@@ -0,0 +1,201 @@
+// Copyright Epic Games, Inc. All Rights Reserved.
+
+#include "Tools/TrainingDataProcessor/STrainingDataProcessorMainWidget.h"
+#include "MLDeformerModel.h"
+#include "MLDeformerTrainingDataProcessorSettings.h"
+#include "Widgets/Text/STextBlock.h"
+#include "Widgets/Layout/SBox.h"
+#include "Widgets/Input/SButton.h"
+#include "Modules/ModuleManager.h"
+#include "PropertyEditorModule.h"
+#include "DetailsViewArgs.h"
+#include "Editor.h"
+#include "Framework/Application/SlateApplication.h"
+#include "Misc/MessageDialog.h"
+#include "Tools/TrainingDataProcessor/TrainingDataProcessor.h"
+#include "Engine/SkeletalMesh.h"
+#include "Animation/Skeleton.h"
+#include "UObject/UObjectGlobals.h"
+
+#define LOCTEXT_NAMESPACE "TrainingDataProcessorMainWidget"
+
+namespace UE::MLDeformer::TrainingDataProcessor
+{
+	void STrainingDataProcessorMainWidget::Construct(const FArguments& InArgs)
+	{
+		Model = InArgs._Model;
+
+		// Create the details view.
+		FPropertyEditorModule& PropertyModule = FModuleManager::LoadModuleChecked<FPropertyEditorModule>("PropertyEditor");
+		FDetailsViewArgs Args;
+		Args.NameAreaSettings = FDetailsViewArgs::HideNameArea;
+		Args.bAllowSearch = true;
+		Args.bShowObjectLabel = false;
+		Args.bShowPropertyMatrixButton = false;
+		Args.bShowOptions = false;
+		Args.bShowObjectLabel = false;
+		Args.NotifyHook = this;
+
+		DetailsView = PropertyModule.CreateDetailView(Args);
+		UMLDeformerTrainingDataProcessorSettings* TrainingDataProcessorSettings = Model ? Model->GetTrainingDataProcessorSettings() : nullptr;
+		DetailsView->SetObject(TrainingDataProcessorSettings);
+
+		// Start listening for changes in the ML Deformer model, as we should validate against the skeleton setup there.
+		// We use this to trigger re-inits of the UI and components, when one of our input assets got modified (by property changes, or reimports, etc.).
+		// This listens to changes in ALL objects in the engine, but we apply a filter in the IsObjectOfInterest to only react on certain objects.
+		ObjectModifiedHandle = FCoreUObjectDelegates::OnObjectModified.AddLambda(
+			[this](UObject* Object)
+			{
+				if (IsObjectOfInterest(Object))
+				{
+					OnAssetModified(Object);
+				}
+			});
+
+		ChildSlot
+		[
+			SNew(SVerticalBox)
+			+ SVerticalBox::Slot()
+			.FillHeight(1.0f)
+			[
+				DetailsView.ToSharedRef()
+			]
+			+ SVerticalBox::Slot()
+			.Padding(FMargin(10.0f, 4.0f))
+			.AutoHeight()
+			[
+				SNew(SButton)
+				.HAlign(EHorizontalAlignment::HAlign_Center)
+				.VAlign(EVerticalAlignment::VAlign_Center)
+				.Text(LOCTEXT("GenerateButtonText", "Generate Training Data"))
+				.ButtonStyle(&FAppStyle::Get().GetWidgetStyle<FButtonStyle>("PrimaryButton"))
+				.OnClicked(this, &STrainingDataProcessorMainWidget::OnGenerateButtonClicked)
+				.IsEnabled_Lambda([this]() { return IsValidConfiguration(); })
+			]
+		];
+
+		GEditor->RegisterForUndo(this);
+	}
+
+	STrainingDataProcessorMainWidget::~STrainingDataProcessorMainWidget()
+	{
+		GEditor->UnregisterForUndo(this);
+
+		if (ObjectModifiedHandle.IsValid())
+		{
+			FCoreUObjectDelegates::OnObjectModified.Remove(ObjectModifiedHandle);
+		}
+	}
+
+	void STrainingDataProcessorMainWidget::PostUndo(bool bSuccess)
+	{
+		FEditorUndoClient::PostUndo(bSuccess);
+		Refresh();
+	}
+
+	void STrainingDataProcessorMainWidget::PostRedo(bool bSuccess)
+	{
+		FEditorUndoClient::PostRedo(bSuccess);
+		Refresh();
+	}
+
+	void STrainingDataProcessorMainWidget::Refresh() const
+	{
+		if (DetailsView.IsValid())
+		{
+			DetailsView->ForceRefresh();
+		}
+	}
+
+	void STrainingDataProcessorMainWidget::NotifyPostChange(const FPropertyChangedEvent& PropertyChangedEvent, FProperty* PropertyThatChanged)
+	{
+		FNotifyHook::NotifyPostChange(PropertyChangedEvent, PropertyThatChanged);
+
+		if (!PropertyThatChanged)
+		{
+			return;
+		}
+
+		// We need to refresh if our output animation changes.
+		if (PropertyThatChanged->GetName() == GET_MEMBER_NAME_CHECKED(UMLDeformerTrainingDataProcessorSettings, OutputAnimSequence))
+		{
+			//DetailsView->Invalidate(EInvalidateWidgetReason::None);
+			Refresh();
+		}
+	}
+
+	bool STrainingDataProcessorMainWidget::IsValidConfiguration() const
+	{
+		const UMLDeformerTrainingDataProcessorSettings* TrainingDataProcessorSettings = Model ? Model->GetTrainingDataProcessorSettings() : nullptr;
+		if (!TrainingDataProcessorSettings)
+		{
+			return false;
+		}
+
+		const USkeleton* Skeleton = (Model && Model->GetSkeletalMesh() && Model->GetSkeletalMesh()->GetSkeleton())
+			                            ? Model->GetSkeletalMesh()->GetSkeleton()
+			                            : nullptr;
+		return TrainingDataProcessorSettings->IsValid(Skeleton);
+	}
+
+	bool STrainingDataProcessorMainWidget::IsObjectOfInterest(UObject* Object) const
+	{
+		if (!Model)
+		{
+			return false;
+		}
+
+		// If we modify our training data processor settings, we aren't interested.
+		// The reason for this is that it will trigger a UI refresh when changing any property in these settings, which we don't want.
+		if (Object == Model->GetTrainingDataProcessorSettings())
+		{
+			return false;
+		}
+
+		// Make sure we only trigger when our skeletal mesh or skeleton changes.
+		const USkeletalMesh* SkeletalMesh = Model->GetSkeletalMesh();
+		const USkeleton* Skeleton = SkeletalMesh ? SkeletalMesh->GetSkeleton() : nullptr;
+		if (Object == Model || Object == SkeletalMesh || Object == Skeleton)
+		{
+			return true;
+		}
+
+		return false;
+	}
+
+	void STrainingDataProcessorMainWidget::OnAssetModified(UObject* Object) const
+	{
+		if (DetailsView.IsValid())
+		{
+			DetailsView->ForceRefresh();
+		}
+	}
+
+	FReply STrainingDataProcessorMainWidget::OnGenerateButtonClicked() const
+	{
+		const UMLDeformerTrainingDataProcessorSettings* TrainingDataProcessorSettings = Model ? Model->GetTrainingDataProcessorSettings() : nullptr;
+		if (!TrainingDataProcessorSettings)
+		{
+			return FReply::Handled();
+		}
+
+		// Run the actual algorithm to process the training data.
+		// This will update the output animation sequence.
+		FTrainingDataProcessor ProcessorAlgo;
+		const USkeleton* Skeleton = Model && Model->GetSkeletalMesh() ? Model->GetSkeletalMesh()->GetSkeleton() : nullptr;
+		if (!ProcessorAlgo.Execute(*TrainingDataProcessorSettings, Skeleton))
+		{
+			FMessageDialog::Open(EAppMsgCategory::Error, EAppMsgType::Ok, LOCTEXT("FailMessage", "Operation failed or canceled by user."),
+			                     LOCTEXT("FailedTitle", "Training Data Processor"));
+		}
+		else
+		{
+			FMessageDialog::Open(EAppMsgCategory::Success, EAppMsgType::Ok, LOCTEXT("SuccessMessage", "Animation Sequence generated successfully."),
+			                     LOCTEXT("SuccessTitle", "Training Data Processor"));
+		}
+
+		return FReply::Handled();
+	}
+} // namespace UE::MLDeformer::TrainingDataProcessor
+
+#undef LOCTEXT_NAMESPACE
diff --git a/Engine/Plugins/Animation/MLDeformer/MLDeformerFramework/Source/MLDeformerFrameworkEditor/Private/Tools/TrainingDataProcessor/TrainingDataProcessor.cpp b/Engine/Plugins/Animation/MLDeformer/MLDeformerFramework/Source/MLDeformerFrameworkEditor/Private/Tools/TrainingDataProcessor/TrainingDataProcessor.cpp
new file mode 100644
--- /dev/null
+++ b/Engine/Plugins/Animation/MLDeformer/MLDeformerFramework/Source/MLDeformerFrameworkEditor/Private/Tools/TrainingDataProcessor/TrainingDataProcessor.cpp
@@ -0,0 +1,500 @@
+// Copyright Epic Games, Inc. All Rights Reserved.
+
+#include "Tools/TrainingDataProcessor/TrainingDataProcessor.h"
+#include "MLDeformerModule.h"
+#include "ReferenceSkeleton.h"
+#include "Algo/MaxElement.h"
+#include "Misc/ScopedSlowTask.h"
+#include "Animation/AnimSequence.h"
+#include "Animation/Skeleton.h"
+#include "BonePose.h"
+#include "Engine/SkeletalMesh.h"
+#include "Animation/AttributesRuntime.h"
+#include "Async/ParallelFor.h"
+#include "Math/RandomStream.h"
+
+#define LOCTEXT_NAMESPACE "TrainingDataProcessorAlgo"
+
+namespace UE::MLDeformer::TrainingDataProcessor
+{
+	bool FTrainingDataProcessor::Execute(const UMLDeformerTrainingDataProcessorSettings& Settings, const USkeleton* Skeleton)
+	{
+		if (!IsValid(Skeleton))
+		{
+			UE_LOG(LogMLDeformer, Error, TEXT("The Training Data Processor settings are not valid."));
+			return false;
+		}
+
+		// Remove existing frame data.
+		Clear();
+
+		// Now iterate over all enabled and valid input animations, and sample their frames.
+		if (!SampleFrames(Settings, *Skeleton))
+		{
+			UE_LOG(LogMLDeformer, Warning, TEXT("Failed to sample frames, user probably cancelled."));
+			return false;
+		}
+
+		// Perform pose remixing, if desired.
+		// This basically randomizes keyframes for groups of bones, as defined in the bone groups that we set up.
+		const FReferenceSkeleton& RefSkeleton = Skeleton->GetReferenceSkeleton();
+		if (Settings.bRemixPoses && !Settings.BoneGroups.Groups.IsEmpty())
+		{
+			if (!RemixPoses(Settings, RefSkeleton))
+			{
+				UE_LOG(LogMLDeformer, Warning, TEXT("No best frames found, user probably cancelled."));
+				return false;
+			}
+		}
+
+		// Reduce the number of frames, which essentially only keeps the frames that are as different as possible from each other.
+		// It is not the optimal solution as it is a greedy algorithm, but it should pick a good set of diverse poses.
+		TArray<int32> FramesToInclude;
+		if (Settings.bReduceFrames)
+		{
+			FramesToInclude = FindBestFrameIndices(Settings.NumOutputFrames, RefSkeleton);
+			if (FramesToInclude.IsEmpty())
+			{
+				UE_LOG(LogMLDeformer, Warning, TEXT("No best frames found, user probably cancelled."));
+				return false;
+			}
+		}
+
+		// Now transfer our transforms data into the animation sequence we selected as output anim sequence.
+		// If FramesToInclude is empty it will just include all the frames.
+		if (!SaveAnimationDataInAnimSequence(Settings, RefSkeleton, FramesToInclude))
+		{
+			UE_LOG(LogMLDeformer, Warning, TEXT("Failed to save results of the training data processor into the output animation sequence."));
+			return false;
+		}
+
+		return true;
+	}
+
+	bool FTrainingDataProcessor::SampleFrames(const UMLDeformerTrainingDataProcessorSettings& Settings, const USkeleton& Skeleton)
+	{
+		// Make sure the animations are loaded.
+		int32 TotalNumFrames = 0;
+		{
+			FScopedSlowTask Task(Settings.AnimList.Num(), LOCTEXT("AnimPreloadText", "Loading animations"));
+			Task.MakeDialog();
+			for (const FMLDeformerTrainingDataProcessorAnim& Anim : Settings.AnimList)
+			{
+				if (Anim.bEnabled)
+				{
+					UAnimSequence* AnimSequence = Anim.AnimSequence.LoadSynchronous();
+					if (AnimSequence)
+					{
+						if (AnimSequence->IsCompressedDataOutOfDate())
+						{
+							AnimSequence->WaitOnExistingCompression(true);
+						}
+
+						TotalNumFrames += AnimSequence->GetDataModel()->GetNumberOfFrames();
+					}
+				}
+				Task.EnterProgressFrame();
+			}
+		}
+
+		// Build a list of existing bone names.
+		AnimFrameData.BoneNames.Reset();
+		AnimFrameData.BoneIndices.Reset();
+		const FReferenceSkeleton& RefSkeleton = Skeleton.GetReferenceSkeleton();
+		for (const FName BoneName : Settings.BoneList.BoneNames)
+		{
+			const int32 BoneIndex = RefSkeleton.FindBoneIndex(BoneName);
+			if (BoneIndex != INDEX_NONE)
+			{
+				AnimFrameData.BoneNames.Add(BoneName);
+				AnimFrameData.BoneIndices.Add(BoneIndex);
+			}
+		}
+
+		// Allocate space for the frames.
+		AnimFrameData.NumRefSkelBones = RefSkeleton.GetNum();
+		const int32 NumFrameElements = AnimFrameData.NumRefSkelBones * TotalNumFrames;
+		AnimFrameData.Transforms.SetNum(NumFrameElements);
+		check(GetNumFrames() == TotalNumFrames);
+
+		// Sample the ref pose and store it.
+		const TArray<FTransform>& RefBonePose = RefSkeleton.GetRefBonePose();
+		AnimFrameData.RefPoseTransforms.SetNum(AnimFrameData.NumRefSkelBones);
+		for (int32 Index = 0; Index < AnimFrameData.NumRefSkelBones; ++Index)
+		{
+			AnimFrameData.RefPoseTransforms[Index] = FTransform3f(RefBonePose[Index]);
+		}
+
+		// Sample all the frames.
+		{
+			FMemMark Mark(FMemStack::Get());
+			FScopedSlowTask Task(TotalNumFrames, LOCTEXT("AnimSamplingText", "Sampling animations"));
+			Task.MakeDialogDelayed(1.0f, true, false);
+
+			TArray<FBoneIndexType> RequiredBoneIndexArray;
+			RequiredBoneIndexArray.AddUninitialized(RefSkeleton.GetNum());
+			for (int32 BoneIndex = 0; BoneIndex < RequiredBoneIndexArray.Num(); ++BoneIndex)
+			{
+				RequiredBoneIndexArray[BoneIndex] = BoneIndex;
+			}
+
+			int32 AnimSeqFrameOffset = 0;
+			for (const FMLDeformerTrainingDataProcessorAnim& Anim : Settings.AnimList)
+			{
+				if (!Anim.bEnabled)
+				{
+					continue;
+				}
+
+				const UAnimSequence* AnimSequence = Anim.AnimSequence.Get();
+				if (!AnimSequence)
+				{
+					continue;
+				}
+
+				const IAnimationDataModel* AnimDataModel = AnimSequence->GetDataModel();
+				const int32 NumFramesInAnimSequence = AnimDataModel->GetNumberOfFrames();
+				const FFrameRate FrameRate = AnimSequence->GetSamplingFrameRate();
+
+				FBoneContainer RequiredBones;
+				RequiredBones.InitializeTo(RequiredBoneIndexArray, UE::Anim::ECurveFilterMode::DisallowAll, Skeleton);
+				RequiredBones.SetUseRAWData(false);
+
+				// Sample all frames of this animation sequence.
+				ParallelFor(NumFramesInAnimSequence, [&](int32 FrameNumber)
+				{
+					const double SampleTime = FrameNumber / FrameRate.AsDecimal();
+
+					FCompactPose Pose;
+					Pose.SetBoneContainer(&RequiredBones);
+					Pose.ResetToRefPose(RequiredBones);
+
+					FBlendedCurve TempCurve;
+					UE::Anim::FStackAttributeContainer TempAttributes;
+
+					FAnimExtractContext ExtractionContext(SampleTime);
+					ExtractionContext.bExtractWithRootMotionProvider = false;
+					FAnimationPoseData AnimPoseData(Pose, TempCurve, TempAttributes);
+					AnimSequence->GetAnimationPose(AnimPoseData, ExtractionContext);
+
+					const FCompactPose& CompactPose = AnimPoseData.GetPose();
+					const TArray<FTransform, TMemStackAllocator<>>& PoseBoneTransforms = CompactPose.GetBones();
+					check(PoseBoneTransforms.Num() == RefSkeleton.GetNum());
+
+					// Store the transforms for the entire reference skeleton.
+					TArrayView<FTransform3f> FrameTransforms = GetFrameTransforms(FrameNumber + AnimSeqFrameOffset);
+					FTransform BoneTransform;
+					for (int32 Index = 0; Index < AnimFrameData.NumRefSkelBones; ++Index)
+					{
+						const FCompactPoseBoneIndex CompactIndex = RequiredBones.GetCompactPoseIndexFromSkeletonIndex(Index);
+						if (CompactIndex != INDEX_NONE)
+						{
+							const FTransform& LocalTransform = PoseBoneTransforms[CompactIndex.GetInt()];
+							BoneTransform = LocalTransform;
+						}
+						else
+						{
+							BoneTransform = RefBonePose[Index];
+						}
+
+						// Store the transform in the frame data.
+						FrameTransforms[Index] = FTransform3f(BoneTransform);
+					}
+				}); // For all frames in this anim sequence.
+
+				Task.EnterProgressFrame();
+				if (Task.ShouldCancel())
+				{
+					Clear();
+					return false;
+				}
+
+				AnimSeqFrameOffset += NumFramesInAnimSequence;
+			} // For all animations.
+
+			const SIZE_T NumBytes = AnimFrameData.Transforms.NumBytes();
+			UE_LOG(LogMLDeformer, Display, TEXT("Sampled frame data: %zu Bytes (%.2f MB)"), NumBytes, NumBytes/(double)(1000*1000));
+		}
+
+		return true;
+	}
+
+	FTransform3f FTrainingDataProcessor::GetRefPoseTransform(const FReferenceSkeleton& RefSkeleton, FName BoneName)
+	{
+		const int32 BoneIndex = RefSkeleton.GetRawRefBoneNames().Find(BoneName);
+		if (BoneIndex != INDEX_NONE)
+		{
+			return FTransform3f(RefSkeleton.GetRefBonePose()[BoneIndex]);
+		}
+
+		return FTransform3f::Identity;
+	}
+
+	void FTrainingDataProcessor::Clear()
+	{
+		AnimFrameData.Transforms.Empty();
+		AnimFrameData.BoneNames.Empty();
+		AnimFrameData.BoneIndices.Empty();
+		AnimFrameData.RefPoseTransforms.Empty();
+		AnimFrameData.NumRefSkelBones = 0;
+	}
+
+	int32 FTrainingDataProcessor::GetNumFrames() const
+	{
+		if (AnimFrameData.NumRefSkelBones == 0)
+		{
+			return 0;
+		}
+		check(AnimFrameData.Transforms.Num() % AnimFrameData.NumRefSkelBones == 0);
+		return AnimFrameData.Transforms.Num() / AnimFrameData.NumRefSkelBones;
+	}
+
+	TArrayView<FTransform3f> FTrainingDataProcessor::GetFrameTransforms(int32 FrameIndex)
+	{
+		const int32 NumBones = AnimFrameData.NumRefSkelBones;
+		check(FrameIndex < GetNumFrames());
+		check(NumBones > 0);
+		return TArrayView<FTransform3f>(&AnimFrameData.Transforms[NumBones * FrameIndex], NumBones);
+	}
+
+	TConstArrayView<FTransform3f> FTrainingDataProcessor::GetFrameTransforms(const TArray<FTransform3f>& Transforms, int32 FrameIndex) const
+	{
+		const int32 NumBones = AnimFrameData.NumRefSkelBones;
+		check(Transforms.Num() == AnimFrameData.Transforms.Num());
+		check(FrameIndex < GetNumFrames());
+		check(NumBones > 0);
+		return TConstArrayView<FTransform3f>(&Transforms[NumBones * FrameIndex], NumBones);
+	}
+
+	TConstArrayView<FTransform3f> FTrainingDataProcessor::GetFrameTransforms(int32 FrameIndex) const
+	{
+		const int32 NumBones = AnimFrameData.NumRefSkelBones;
+		check(FrameIndex < GetNumFrames());
+		check(NumBones > 0);
+		return TConstArrayView<FTransform3f>(&AnimFrameData.Transforms[NumBones * FrameIndex], NumBones);
+	}
+
+	TConstArrayView<FTransform3f> FTrainingDataProcessor::GetRefPoseTransforms() const
+	{
+		const int32 NumBones = AnimFrameData.NumRefSkelBones;
+		check(NumBones > 0);
+		return TConstArrayView<FTransform3f>(&AnimFrameData.Transforms[0], NumBones);
+	}
+
+	double FTrainingDataProcessor::CalculateMeanError(const TConstArrayView<FTransform3f> PoseA, const TConstArrayView<FTransform3f> PoseB) const
+	{
+		check(PoseA.Num() == PoseB.Num());
+		check(PoseA.Num() == AnimFrameData.NumRefSkelBones);
+
+		double Sum = 0.0;
+		for (int32 Index = 0; Index < AnimFrameData.BoneIndices.Num(); ++Index)
+		{
+			const int32 BoneIndex = AnimFrameData.BoneIndices[Index];
+			const FQuat RotA = FQuat(PoseA[BoneIndex].GetRotation());
+			const FQuat RotB = FQuat(PoseB[BoneIndex].GetRotation());
+			Sum += FMath::Square(RotA.X - RotB.X);
+			Sum += FMath::Square(RotA.Y - RotB.Y);
+			Sum += FMath::Square(RotA.Z - RotB.Z);
+			Sum += FMath::Square(RotA.W - RotB.W);
+		}
+
+		const double Mean = Sum / static_cast<double>(AnimFrameData.BoneIndices.Num() * 4); // Times 4 because of the 4 quat values.
+		return Mean;
+	}
+
+	double FTrainingDataProcessor::CalculateMeanError(int32 PreviousPoseIndex, int32 CurrentPoseIndex) const
+	{
+		if (PreviousPoseIndex != -1)
+		{
+			return CalculateMeanError(GetFrameTransforms(PreviousPoseIndex), GetFrameTransforms(CurrentPoseIndex));
+		}
+
+		return CalculateMeanError(GetRefPoseTransforms(), GetFrameTransforms(CurrentPoseIndex));
+	}
+
+	// Find the best MaxNumFrames number of frames.
+	// With best frames we mean frames that are most far apart from each other.
+	TArray<int32> FTrainingDataProcessor::FindBestFrameIndices(int32 MaxNumFrames, const FReferenceSkeleton& RefSkeleton) const
+	{
+		TArray<int32> BestFrames;
+
+		// Now find the best frames.
+		const int32 NumSampledFrames = GetNumFrames();
+		const int32 NumFramesToFind = FMath::Min(MaxNumFrames, NumSampledFrames);
+		BestFrames.Reserve(NumFramesToFind);
+		int32 PreviousBestFrameIndex = -1; // The ref pose.
+
+		TArray<double> FrameMeanErrors; // One error per frame.
+		FrameMeanErrors.SetNum(NumSampledFrames);
+
+		TArray<int32> FramesToCheck;
+		FramesToCheck.SetNum(NumSampledFrames);
+		for (int32 FrameIndex = 0; FrameIndex < NumSampledFrames; FrameIndex++)
+		{
+			FramesToCheck[FrameIndex] = FrameIndex;
+		}
+
+		FScopedSlowTask Task(NumFramesToFind, LOCTEXT("FrameReductionText", "Reducing the number of frames"));
+		Task.MakeDialogDelayed(1.0f, true);
+
+		for (int32 FrameToFindIndex = 0; FrameToFindIndex < NumFramesToFind; ++FrameToFindIndex)
+		{
+			// Calculate the mean errors to the previous best pose.
+			check(FramesToCheck.Num() == FrameMeanErrors.Num())
+			ParallelFor(TEXT("MLDeformer::TrainingDataProcessorAlgo::FindBestFrameIndices"), FramesToCheck.Num(), 5, [&](int32 Index)
+			{
+				const int32 FrameIndex = FramesToCheck[Index];
+				FrameMeanErrors[Index] = CalculateMeanError(PreviousBestFrameIndex, FrameIndex);
+			});
+
+			// Find the frame with the highest error, which indicates it diverges most from the previous pose.
+			const double* MaxDistancePtr = Algo::MaxElement(FrameMeanErrors);
+			const int32 MaxErrorFrameIndex = MaxDistancePtr - FrameMeanErrors.GetData();
+			const int32 BestSampledFrameIndex = FramesToCheck[MaxErrorFrameIndex];
+
+			// Store this frame as best frame.
+			BestFrames.Add(BestSampledFrameIndex);
+			FramesToCheck.RemoveAt(MaxErrorFrameIndex);
+			FrameMeanErrors.Pop(); // We can just pop here instead of RemoveAt, as we reinit all values later, and Pop is more efficient.
+			PreviousBestFrameIndex = BestSampledFrameIndex;
+
+			Task.EnterProgressFrame();
+			if (Task.ShouldCancel())
+			{
+				BestFrames.Empty();
+				break;
+			}
+		}
+
+		return MoveTemp(BestFrames);
+	}
+
+	namespace
+	{
+		TArray<int32> GenerateShuffledFrameArray(int32 NumFrames, const FRandomStream& RandomStream)
+		{
+			// First initialize the list to the ordered list of frame numbers.
+			TArray<int32> ShuffledFrames;
+			ShuffledFrames.SetNum(NumFrames);
+			for (int32 FrameIndex = 0; FrameIndex < NumFrames; ++FrameIndex)
+			{
+				ShuffledFrames[FrameIndex] = FrameIndex;
+			}
+
+			// Perform the shuffle.
+			for (int32 Index = NumFrames - 1; Index > 0; --Index)
+			{
+				const int32 RandomFrameIndex = RandomStream.RandRange(0, Index);
+				ShuffledFrames.Swap(Index, RandomFrameIndex);
+			}
+
+			return MoveTemp(ShuffledFrames);
+		}
+	}
+
+	bool FTrainingDataProcessor::RemixPoses(const UMLDeformerTrainingDataProcessorSettings& Settings,
+	                                            const FReferenceSkeleton& RefSkeleton)
+	{
+		// We need a backup of our frame data for this to work correctly.
+		const TArray<FTransform3f> OriginalTransforms = AnimFrameData.Transforms;
+
+		const FRandomStream RandomStream(Settings.RandomSeed);
+		for (int32 BoneGroupIndex = 0; BoneGroupIndex < Settings.BoneGroups.Groups.Num(); ++BoneGroupIndex)
+		{
+			const FMLDeformerTrainingDataProcessorBoneGroup& BoneGroup = Settings.BoneGroups.Groups[BoneGroupIndex];
+
+			// Generate a shuffled frame order for this group. This just reorders the frame numbers.
+			// For example frames [0, 1, 2, 3, 4] could be shuffled to turn into [2, 4, 0, 1, 3].
+			const TArray<int32> ShuffledFrames = GenerateShuffledFrameArray(GetNumFrames(), RandomStream);
+
+			// For every frame.
+			for (int32 Index = 0; Index < ShuffledFrames.Num(); ++Index)
+			{
+				// Get the shuffled and current frame transform buffers.
+				// With shuffled we mean the transform of another frame number.
+				const int32 ShuffledFrameIndex = ShuffledFrames[Index];
+				const TConstArrayView<FTransform3f> ShuffledFrameTransforms = GetFrameTransforms(OriginalTransforms, ShuffledFrameIndex);
+				TArrayView<FTransform3f> FrameTransforms = GetFrameTransforms(Index);
+
+				// Replace the transform for each bone in this bone group with the shuffled version.
+				for (const FName BoneName : BoneGroup.BoneNames)
+				{
+					const int32 RefSkelBoneIndex = RefSkeleton.FindBoneIndex(BoneName);
+					if (RefSkelBoneIndex != INDEX_NONE)
+					{
+						FrameTransforms[RefSkelBoneIndex] = ShuffledFrameTransforms[RefSkelBoneIndex];
+					}
+				}
+			}
+		}
+
+		return true;
+	}
+
+	bool FTrainingDataProcessor::SaveAnimationDataInAnimSequence(const UMLDeformerTrainingDataProcessorSettings& Settings,
+	                                                                 const FReferenceSkeleton& RefSkeleton,
+	                                                                 const TArray<int32>& FramesToInclude)
+	{
+		UAnimSequence* OutputSequence = Settings.OutputAnimSequence.LoadSynchronous();
+		if (!OutputSequence)
+		{
+			return false;
+		}
+
+		(void)OutputSequence->MarkPackageDirty();
+
+		check(RefSkeleton.GetNum() == AnimFrameData.NumRefSkelBones);
+		const int32 NumBones = AnimFrameData.NumRefSkelBones;
+
+		FScopedSlowTask Task(NumBones, LOCTEXT("SavingMessage", "Generating Animation Sequence"));
+		Task.MakeDialogDelayed(1.0f, true);
+
+		// Clear the current animation data, so all bone and curve tracks etc.
+		IAnimationDataController& Controller = OutputSequence->GetController();
+		Controller.InitializeModel();
+		OutputSequence->ResetAnimation();
+		const int32 NumOutputFrames = FramesToInclude.IsEmpty() ? GetNumFrames() : FramesToInclude.Num();
+		Controller.SetNumberOfFrames(NumOutputFrames, false);
+		Controller.SetFrameRate(FFrameRate(30, 1), false);
+
+		// Preallocate the pos/rot/scale buffers, as we will decompose the FTransforms.
+		TArray<FVector3f> PosKeys;
+		TArray<FQuat4f> RotKeys;
+		TArray<FVector3f> ScaleKeys;
+		PosKeys.SetNumUninitialized(NumOutputFrames, EAllowShrinking::No);
+		RotKeys.SetNumUninitialized(NumOutputFrames, EAllowShrinking::No);
+		ScaleKeys.SetNumUninitialized(NumOutputFrames, EAllowShrinking::No);
+
+		// For all bones, generate a key track.
+		Controller.OpenBracket(LOCTEXT("CreateNewAnimBracket", "Create Anim Sequence"));
+		for (int32 BoneIndex = 0; BoneIndex < NumBones; ++BoneIndex)
+		{
+			const FName BoneName = RefSkeleton.GetBoneName(BoneIndex);
+
+			// Decompose the transforms into arrays for pos/rot/scale.
+			for (int32 Index = 0; Index < NumOutputFrames; ++Index)
+			{
+				const int32 FrameIndex = FramesToInclude.IsEmpty() ? Index : FramesToInclude[Index];
+				const TConstArrayView<FTransform3f> FrameTransforms = GetFrameTransforms(FrameIndex);
+				const FTransform3f& Transform = FrameTransforms[BoneIndex];
+
+				PosKeys[Index] = AnimFrameData.RefPoseTransforms[BoneIndex].GetTranslation();
+				RotKeys[Index] = Transform.GetRotation();
+				ScaleKeys[Index] = AnimFrameData.RefPoseTransforms[BoneIndex].GetScale3D();
+			}
+
+			Controller.AddBoneCurve(BoneName, false);
+			Controller.SetBoneTrackKeys(BoneName, PosKeys, RotKeys, ScaleKeys, false);
+
+			Task.EnterProgressFrame();
+		}
+		Controller.CloseBracket();
+
+		Controller.NotifyPopulated();
+		OutputSequence->RefreshCacheData();
+
+		return true;
+	}
+} // namespace UE::MLDeformer::TrainingDataProcessor
+
+#undef LOCTEXT_NAMESPACE
diff --git a/Engine/Plugins/Animation/MLDeformer/MLDeformerFramework/Source/MLDeformerFrameworkEditor/Private/Tools/TrainingDataProcessor/TrainingDataProcessorSettingsDetailCustomization.cpp b/Engine/Plugins/Animation/MLDeformer/MLDeformerFramework/Source/MLDeformerFrameworkEditor/Private/Tools/TrainingDataProcessor/TrainingDataProcessorSettingsDetailCustomization.cpp
new file mode 100644
--- /dev/null
+++ b/Engine/Plugins/Animation/MLDeformer/MLDeformerFramework/Source/MLDeformerFrameworkEditor/Private/Tools/TrainingDataProcessor/TrainingDataProcessorSettingsDetailCustomization.cpp
@@ -0,0 +1,378 @@
+// Copyright Epic Games, Inc. All Rights Reserved.
+
+#include "Tools/TrainingDataProcessor/TrainingDataProcessorSettingsDetailCustomization.h"
+#include "ContentBrowserModule.h"
+#include "DetailCategoryBuilder.h"
+#include "DetailLayoutBuilder.h"
+#include "DetailWidgetRow.h"
+#include "IContentBrowserSingleton.h"
+#include "MLDeformerGeomCacheModel.h"
+#include "MLDeformerTrainingDataProcessorSettings.h"
+#include "Animation/AnimSequence.h"
+#include "Animation/Skeleton.h"
+#include "AssetRegistry/AssetRegistryModule.h"
+#include "Engine/SkeletalMesh.h"
+#include "Framework/Application/SlateApplication.h"
+#include "Misc/MessageDialog.h"
+#include "UObject/SavePackage.h"
+#include "Widgets/Input/SButton.h"
+#include "Widgets/Layout/SBox.h"
+#include "PropertyCustomizationHelpers.h"
+#include "SWarningOrErrorBox.h"
+
+#define LOCTEXT_NAMESPACE "TrainingDataProcessorSettingsDetailCustomization"
+
+namespace UE::MLDeformer::TrainingDataProcessor
+{
+	TSharedRef<IDetailCustomization> FTrainingDataProcessorSettingsDetailCustomization::MakeInstance()
+	{
+		return MakeShareable(new FTrainingDataProcessorSettingsDetailCustomization());
+	}
+
+	namespace
+	{
+		bool GetFolderFromPackagePath(const FString& PackagePath, FString& OutFolderPath)
+		{
+			int32 LastSlashIndex = INDEX_NONE;
+			if (PackagePath.FindLastChar('/', LastSlashIndex) && LastSlashIndex != INDEX_NONE)
+			{
+				OutFolderPath = PackagePath.Left(LastSlashIndex);
+				return true;
+			}
+
+			return false;
+		}
+	}
+
+	void FTrainingDataProcessorSettingsDetailCustomization::CustomizeDetails(IDetailLayoutBuilder& DetailBuilder)
+	{
+		IDetailCategoryBuilder& OutputCategory = DetailBuilder.EditCategory("Output");
+		IDetailCategoryBuilder& InputAnimCategory = DetailBuilder.EditCategory("Input Animations");
+		IDetailCategoryBuilder& FrameReductionCategory = DetailBuilder.EditCategory("Frame Reduction");
+
+		// Get the training data processor settings and the skeleton.
+		TArray<TWeakObjectPtr<UObject>> Objects;
+		DetailBuilder.GetObjectsBeingCustomized(Objects);
+		if (Objects.Num() == 1 && Objects[0] != nullptr)
+		{
+			TrainingDataProcessorSettings = Cast<UMLDeformerTrainingDataProcessorSettings>(Objects[0].Get());
+			check(TrainingDataProcessorSettings.Get());
+		}
+
+		// Hide the default output anim sequence property.
+		const FName PropertyName = GET_MEMBER_NAME_CHECKED(UMLDeformerTrainingDataProcessorSettings, OutputAnimSequence);
+		TSharedPtr<IPropertyHandle> OutputAnimSequencePropertyHandle = DetailBuilder.GetProperty(PropertyName);
+		DetailBuilder.HideProperty(OutputAnimSequencePropertyHandle);
+
+		OutputCategory.AddCustomRow(FText::FromString("SkeletonMismatchErrorRow"))
+		              .Visibility(TAttribute<EVisibility>::CreateSP(
+			              this, &FTrainingDataProcessorSettingsDetailCustomization::GetSkeletonMismatchErrorVisibility))
+		              .WholeRowContent()
+		[
+			SNew(SBox)
+			.Padding(FMargin(0.0f, 4.0f))
+			[
+				SNew(SWarningOrErrorBox)
+				.MessageStyle(EMessageStyle::Error)
+				.Message(LOCTEXT("SkeletonMismatchError", "The output anim sequence skeleton does not match the model's skeleton."))
+			]
+		];
+
+		// Create our own version of the property.
+		// Insert some Create New button behind the anim sequence property.
+		OutputCategory.AddCustomRow(FText::FromName(PropertyName))
+		              .NameContent()
+			[
+				OutputAnimSequencePropertyHandle->CreatePropertyNameWidget()
+			]
+			.ValueContent()
+			[
+				SNew(SHorizontalBox)
+				+ SHorizontalBox::Slot()
+				.AutoWidth()
+				[
+					SNew(SObjectPropertyEntryBox)
+					.PropertyHandle(OutputAnimSequencePropertyHandle)
+					.AllowedClass(UAnimSequence::StaticClass())
+					.ObjectPath_Lambda([this]()
+					{
+						return TrainingDataProcessorSettings->OutputAnimSequence.LoadSynchronous()
+							       ? TrainingDataProcessorSettings->OutputAnimSequence->GetPathName()
+							       : FString();
+					})
+					.ThumbnailPool(DetailBuilder.GetThumbnailPool())
+					.OnShouldFilterAsset(this, &FTrainingDataProcessorSettingsDetailCustomization::FilterAnimSequences)
+				]
+				+ SHorizontalBox::Slot()
+				.Padding(FMargin(8.0f, 2.0f, 2.0f, 2.0f))
+				.MaxWidth(200)
+				.VAlign(VAlign_Center)
+				[
+					SNew(SButton)
+					.Text(FText::FromString("Create New"))
+					.HAlign(HAlign_Center)
+					.VAlign(VAlign_Center)
+					.OnClicked(this, &FTrainingDataProcessorSettingsDetailCustomization::OnCreateNewButtonClicked)
+				]
+			];
+
+		// Show a warning when there are no input frames.
+		InputAnimCategory.AddCustomRow(FText::FromString("NoInputAnimsWarning"))
+		                 .Visibility(TAttribute<EVisibility>::CreateSP(
+			                 this, &FTrainingDataProcessorSettingsDetailCustomization::GetNoFramesWarningVisibility))
+		                 .WholeRowContent()
+		[
+			SNew(SBox)
+			.Padding(FMargin(0.0f, 4.0f))
+			[
+				SNew(SWarningOrErrorBox)
+				.MessageStyle(EMessageStyle::Warning)
+				.Message(LOCTEXT("NoInputFramesError", "Please add some input animations to sample frames from."))
+			]
+		];
+
+		// Add the animation list property.
+		InputAnimCategory.AddProperty(
+			                 GET_MEMBER_NAME_CHECKED(UMLDeformerTrainingDataProcessorSettings, AnimList),
+			                 UMLDeformerTrainingDataProcessorSettings::StaticClass())
+		                 .DisplayName(FText::Format(LOCTEXT("InputAnimsFrameCountString", "Animation List ({0} Frames)"), GetTotalNumInputFrames()));
+
+		// Get the property handle for the animation inputs, and make it listen to changes.
+		TSharedPtr<IPropertyHandle> InputAnimsArrayPropertyHandle = DetailBuilder.GetProperty(
+			GET_MEMBER_NAME_CHECKED(UMLDeformerTrainingDataProcessorSettings, AnimList));
+		if (InputAnimsArrayPropertyHandle.IsValid())
+		{
+			InputAnimsArrayPropertyHandle->SetOnPropertyValueChanged(
+				FSimpleDelegate::CreateStatic(&FTrainingDataProcessorSettingsDetailCustomization::Refresh, &DetailBuilder));
+		}
+
+		// Show a warning when there are no input bones while frame reduction is enabled.
+		FrameReductionCategory.AddCustomRow(FText::FromString("NoInputBonesWarning"))
+		                      .Visibility(TAttribute<EVisibility>::CreateSP(
+			                      this, &FTrainingDataProcessorSettingsDetailCustomization::GetNoInputBonesWarningVisibility))
+		                      .WholeRowContent()
+		[
+			SNew(SBox)
+			.Padding(FMargin(0.0f, 4.0f))
+			[
+				SNew(SWarningOrErrorBox)
+				.MessageStyle(EMessageStyle::Warning)
+				.Message(LOCTEXT("NoInputBonesError", "Please add input bones."))
+			]
+		];
+	}
+
+	EVisibility FTrainingDataProcessorSettingsDetailCustomization::GetNoInputBonesWarningVisibility() const
+	{
+		if (TrainingDataProcessorSettings.IsValid())
+		{
+			return (TrainingDataProcessorSettings->bReduceFrames && TrainingDataProcessorSettings->BoneList.BoneNames.IsEmpty())
+				       ? EVisibility::Visible
+				       : EVisibility::Collapsed;
+		}
+		return EVisibility::Collapsed;
+	}
+
+	EVisibility FTrainingDataProcessorSettingsDetailCustomization::GetNoFramesWarningVisibility() const
+	{
+		return GetTotalNumInputFrames() > 0 ? EVisibility::Collapsed : EVisibility::Visible;
+	}
+
+	int32 FTrainingDataProcessorSettingsDetailCustomization::GetTotalNumInputFrames() const
+	{
+		const FAssetRegistryModule& AssetRegistryModule = FModuleManager::LoadModuleChecked<FAssetRegistryModule>("AssetRegistry");
+		int32 TotalNumFrames = 0;
+		if (TrainingDataProcessorSettings.IsValid())
+		{
+			for (const FMLDeformerTrainingDataProcessorAnim& Anim : TrainingDataProcessorSettings->AnimList)
+			{
+				if (Anim.bEnabled)
+				{
+					const FAssetData AssetData = AssetRegistryModule.Get().GetAssetByObjectPath(Anim.AnimSequence.ToSoftObjectPath());
+					if (AssetData.IsValid())
+					{
+						const FAssetTagValueRef Tag = AssetData.TagsAndValues.FindTag(TEXT("Number Of Frames"));
+						if (Tag.IsSet())
+						{
+							TotalNumFrames += FCString::Atoi(*Tag.GetValue());
+						}
+					}
+				}
+			}
+		}
+		return TotalNumFrames;
+	}
+
+	void FTrainingDataProcessorSettingsDetailCustomization::Refresh(IDetailLayoutBuilder* DetailBuilder)
+	{
+		if (DetailBuilder)
+		{
+			DetailBuilder->ForceRefreshDetails();
+		}
+	}
+
+	FString FTrainingDataProcessorSettingsDetailCustomization::FindDefaultAnimSequencePath(const UMLDeformerModel* Model)
+	{
+		// Try to find a good path to create a new animation in.
+		// Do this by checking if we already have training animations set up, if so, use that same path.
+		if (Model)
+		{
+			FString Folder;
+			const UMLDeformerGeomCacheModel* GeomCacheModel = Cast<UMLDeformerGeomCacheModel>(Model);
+			if (GeomCacheModel)
+			{
+				// Try to use the training animations first.
+				for (const FMLDeformerGeomCacheTrainingInputAnim& Anim : GeomCacheModel->GetTrainingInputAnims())
+				{
+					if (Anim.GetAnimSequenceSoftObjectPtr().ToSoftObjectPath().IsValid())
+					{
+						if (GetFolderFromPackagePath(Anim.GetAnimSequenceSoftObjectPtr().GetLongPackageName(), Folder))
+						{
+							return Folder;
+						}
+					}
+				}
+			}
+
+			// Since we didn't find a good training animation path, try the input animations to the training data processor settings.
+			const UMLDeformerTrainingDataProcessorSettings* Settings = Model->GetTrainingDataProcessorSettings();
+			if (Settings)
+			{
+				for (const FMLDeformerTrainingDataProcessorAnim& Anim : Settings->AnimList)
+				{
+					if (GetFolderFromPackagePath(Anim.AnimSequence.GetLongPackageName(), Folder))
+					{
+						return Folder;
+					}
+				}
+			}
+
+			// If that also failed, use the model path.
+			if (GetFolderFromPackagePath(Model->GetPathName(nullptr), Folder))
+			{
+				return Folder;
+			}
+		}
+
+		return FString("/Game");
+	}
+
+	FReply FTrainingDataProcessorSettingsDetailCustomization::OnCreateNewButtonClicked() const
+	{
+		// Get the skeleton from our model.
+		TStrongObjectPtr<UMLDeformerTrainingDataProcessorSettings> StrongSettings = TrainingDataProcessorSettings.Pin();
+		check(StrongSettings);
+		USkeleton* Skeleton = nullptr;
+		UMLDeformerModel* Model = Cast<UMLDeformerModel>(StrongSettings->GetOuter());
+		if (Model)
+		{
+			Skeleton = Model->GetSkeletalMesh() ? Model->GetSkeletalMesh()->GetSkeleton() : nullptr;
+		}
+
+		// We need a skeleton.
+		const FText MessageTitle = LOCTEXT("MessageBoxTitle", "Training Data Processor");
+		if (!Skeleton)
+		{
+			FMessageDialog::Open(EAppMsgCategory::Error, EAppMsgType::Ok,
+			                     LOCTEXT("FailMessage",
+			                             "Cannot create new anim sequence as we don't have a skeleton to assign.\nPlease make sure you select a skeletal mesh in your model."),
+			                     MessageTitle);
+			return FReply::Handled();
+		}
+
+		// Create the save asset dialog, don't allow picking existing assets.
+		FContentBrowserModule& ContentBrowserModule = FModuleManager::LoadModuleChecked<FContentBrowserModule>("ContentBrowser");
+		FSaveAssetDialogConfig SaveAssetConfig;
+		SaveAssetConfig.DialogTitleOverride = LOCTEXT("CreateWindowTitle", "Create New Animation Sequence Asset");
+		SaveAssetConfig.DefaultPath = FindDefaultAnimSequencePath(Model);
+		SaveAssetConfig.DefaultAssetName = TEXT("NewAnimSequence");
+		SaveAssetConfig.ExistingAssetPolicy = ESaveAssetDialogExistingAssetPolicy::Type::Disallow;
+		const FString SaveObjectPath = ContentBrowserModule.Get().CreateModalSaveAssetDialog(SaveAssetConfig);
+		if (SaveObjectPath.IsEmpty())
+		{
+			UE_LOG(LogMLDeformer, Warning, TEXT("User canceled the save dialog."));
+			return FReply::Handled();
+		}
+
+		// Create the package.
+		FString PackageName;
+		FString AssetName;
+		SaveObjectPath.Split(TEXT("."), &PackageName, &AssetName);
+		UPackage* Package = CreatePackage(*PackageName);
+		if (!Package)
+		{
+			FMessageDialog::Open(EAppMsgCategory::Error, EAppMsgType::Ok, LOCTEXT("CreatePackageFailMessage", "Failed to create the package."),
+			                     MessageTitle);
+			return FReply::Handled();
+		}
+		Package->FullyLoad();
+
+		// Create the animation asset.
+		UAnimSequence* NewAnimSequence = NewObject<UAnimSequence>(Package, FName(*AssetName), RF_Public | RF_Standalone);
+		if (NewAnimSequence)
+		{
+			// Set up the animation sequence.
+			NewAnimSequence->SetSkeleton(Skeleton);
+			NewAnimSequence->SetPreviewMesh(Model->GetSkeletalMesh());
+			NewAnimSequence->GetController().InitializeModel();
+			// Needed, as otherwise there will be issues and crashes in the anim data model validation code.
+			NewAnimSequence->RefreshCacheData();
+			(void)Package->MarkPackageDirty();
+			FAssetRegistryModule::AssetCreated(NewAnimSequence);
+
+			// Update the animation sequence inside our training data processor settings UObject to point to the newly created one.
+			// That will also refresh the UI to show our newly created anim sequence when returning from this dialog.
+			StrongSettings->OutputAnimSequence = NewAnimSequence;
+			UE_LOG(LogMLDeformer, Display, TEXT("Successfully created AnimSequence asset at: %s"), *SaveObjectPath);
+		}
+		else
+		{
+			FMessageDialog::Open(EAppMsgCategory::Error, EAppMsgType::Ok,
+			                     LOCTEXT("AnimCreationErrorMessage", "Failed to create new Anim Sequence. Please check the log for more details."),
+			                     MessageTitle);
+			UE_LOG(LogMLDeformer, Error, TEXT("Failed to create the anim instance asset at: %s"), *SaveObjectPath);
+		}
+
+		return FReply::Handled();
+	}
+
+	bool FTrainingDataProcessorSettingsDetailCustomization::FilterAnimSequences(const FAssetData& AssetData) const
+	{
+		const USkeleton* Skeleton = TrainingDataProcessorSettings.IsValid() ? TrainingDataProcessorSettings->FindSkeleton() : nullptr;
+		if (Skeleton && Skeleton->IsCompatibleForEditor(AssetData))
+		{
+			return false;
+		}
+
+		return true;
+	}
+
+	EVisibility FTrainingDataProcessorSettingsDetailCustomization::GetSkeletonMismatchErrorVisibility() const
+	{
+		// Get the skeleton from our model.
+		if (TrainingDataProcessorSettings.IsValid())
+		{
+			const USkeleton* ModelSkeleton = TrainingDataProcessorSettings->FindSkeleton();
+
+			// Get the skeleton used by the output animation.
+			USkeleton* OutputAnimSkeleton = nullptr;
+			if (TrainingDataProcessorSettings->OutputAnimSequence.IsValid())
+			{
+				if (TrainingDataProcessorSettings->OutputAnimSequence.LoadSynchronous())
+				{
+					OutputAnimSkeleton = TrainingDataProcessorSettings->OutputAnimSequence->GetSkeleton();
+				}
+			}
+
+			// If both skeletons are known and they mismatch, show the error.
+			if (ModelSkeleton && OutputAnimSkeleton && ModelSkeleton != OutputAnimSkeleton)
+			{
+				return EVisibility::Visible;
+			}
+		}
+
+		return EVisibility::Collapsed;
+	}
+} // namespace UE::MLDeformer::TrainingDataProcessor
+
+#undef LOCTEXT_NAMESPACE
diff --git a/Engine/Plugins/Animation/MLDeformer/MLDeformerFramework/Source/MLDeformerFrameworkEditor/Private/Tools/TrainingDataProcessor/TrainingDataProcessorTool.cpp b/Engine/Plugins/Animation/MLDeformer/MLDeformerFramework/Source/MLDeformerFrameworkEditor/Private/Tools/TrainingDataProcessor/TrainingDataProcessorTool.cpp
new file mode 100644
--- /dev/null
+++ b/Engine/Plugins/Animation/MLDeformer/MLDeformerFramework/Source/MLDeformerFrameworkEditor/Private/Tools/TrainingDataProcessor/TrainingDataProcessorTool.cpp
@@ -0,0 +1,160 @@
+// Copyright Epic Games, Inc. All Rights Reserved.
+
+#include "Tools/TrainingDataProcessor/TrainingDataProcessorTool.h"
+#include "Tools/TrainingDataProcessor/STrainingDataProcessorMainWidget.h"
+#include "MLDeformerEditorToolkit.h"
+#include "MLDeformerTrainingDataProcessorSettings.h"
+#include "WorkflowOrientedApp/WorkflowTabFactory.h"
+#include "IDocumentation.h"
+#include "Widgets/Docking/SDockTab.h"
+#include "Engine/SkeletalMesh.h"
+#include "Animation/Skeleton.h"
+#include "PropertyHandle.h"
+
+#define LOCTEXT_NAMESPACE "TrainingDataProcessorTool"
+
+namespace UE::MLDeformer::TrainingDataProcessor
+{
+	static FName TrainingDataProcessorToolName = FName("Training Data Processor");
+	static FText TrainingDataProcessorToolTip = LOCTEXT("TrainingDataProcessorToolToolTip",
+	                                                    "The training data processor tool, which allows us to generate training data from a set of animations.");
+
+	// The tab summoner for the tool.
+	class FTrainingDataProcessorToolTabSummoner
+		: public FWorkflowTabFactory
+	{
+	public:
+		FTrainingDataProcessorToolTabSummoner(const TSharedRef<FMLDeformerEditorToolkit>& InEditor)
+			: FWorkflowTabFactory(TrainingDataProcessorToolName, InEditor)
+		{
+			bIsSingleton = true;
+			TabRole = ETabRole::NomadTab;
+			TabLabel = LOCTEXT("TabLabel", "Training Data");
+			ViewMenuDescription = TrainingDataProcessorToolTip;
+			ViewMenuTooltip = TrainingDataProcessorToolTip;
+		}
+
+		virtual ~FTrainingDataProcessorToolTabSummoner() override = default;
+
+		virtual TSharedRef<SWidget> CreateTabBody(const FWorkflowTabSpawnInfo& Info) const override
+		{
+			TSharedPtr<FMLDeformerEditorToolkit> Toolkit = StaticCastSharedPtr<FMLDeformerEditorToolkit>(HostingApp.Pin());
+			if (!Toolkit.IsValid())
+			{
+				return SNullWidget::NullWidget;
+			}
+
+			TSharedPtr<FMLDeformerEditorModel> EditorModel = Toolkit->GetActiveModelPointer().Pin();
+			if (!EditorModel.IsValid())
+			{
+				return SNullWidget::NullWidget;
+			}
+
+			// Create the training data processor settings if it doesn't exist yet.
+			UMLDeformerModel* Model = EditorModel->GetModel();
+			if (!Model->GetTrainingDataProcessorSettings())
+			{
+				UMLDeformerTrainingDataProcessorSettings* TrainingDataProcessorSettings = NewObject<UMLDeformerTrainingDataProcessorSettings>(
+					Model, NAME_None, RF_Transactional);
+				Model->SetTrainingDataProcessorSettings(TrainingDataProcessorSettings);
+			}
+
+			return SNew(STrainingDataProcessorMainWidget)
+				.Model(Model);
+		}
+
+		virtual TSharedPtr<SToolTip> CreateTabToolTipWidget(const FWorkflowTabSpawnInfo& Info) const override
+		{
+			return IDocumentation::Get()->CreateToolTip(
+				TrainingDataProcessorToolTip,
+				nullptr,
+				TEXT("Shared/Editors/Persona"),
+				"TrainingDataProcessorTool_Window");
+		}
+	};
+
+
+	// The tool menu extender, which extends the tools menu inside the ML Deformer asset editor.
+	class FMLDeformerTrainingDataProcessorToolsMenuExtender
+		: public FToolsMenuExtender
+	{
+	public:
+		virtual FMenuEntryParams GetMenuEntry(FMLDeformerEditorToolkit& Toolkit) const override
+		{
+			FMenuEntryParams Params;
+			Params.DirectActions = FUIAction(
+				FExecuteAction::CreateLambda([this, &Toolkit]()
+				{
+					Toolkit.GetAssociatedTabManager()->TryInvokeTab(TrainingDataProcessorToolName);
+				}),
+				FCanExecuteAction::CreateLambda([]()
+				{
+					return true;
+				})
+			);
+			Params.LabelOverride = FText::FromName(TrainingDataProcessorToolName);
+			Params.ToolTipOverride = TrainingDataProcessorToolTip;
+			return Params;
+		}
+
+		virtual TSharedPtr<FWorkflowTabFactory> GetTabSummoner(const TSharedRef<FMLDeformerEditorToolkit>& Toolkit) const override
+		{
+			return MakeShared<FTrainingDataProcessorToolTabSummoner>(Toolkit);
+		}
+	};
+
+	// Register to the ML Deformer asset editor tools menu.
+	void RegisterTool()
+	{
+		FMLDeformerEditorToolkit::AddToolsMenuExtender(MakeUnique<FMLDeformerTrainingDataProcessorToolsMenuExtender>());
+	}
+
+	USkeleton* FindSkeletonForProperty(const TSharedPtr<IPropertyHandle>& PropertyHandle)
+	{
+		if (!PropertyHandle.IsValid())
+		{
+			return nullptr;
+		}
+
+		TArray<UObject*> Objects;
+		PropertyHandle->GetOuterObjects(Objects);
+
+		auto FindSkeletonForObject = [&PropertyHandle](UObject* InObject) -> USkeleton* {
+			for (; InObject; InObject = InObject->GetOuter())
+			{
+				if (USkeletalMesh* SkeletalMesh = Cast<USkeletalMesh>(InObject))
+				{
+					return SkeletalMesh->GetSkeleton();
+				}
+
+				if (const UAnimationAsset* AnimationAsset = Cast<UAnimationAsset>(InObject))
+				{
+					if (AnimationAsset->IsAsset())
+					{
+						return AnimationAsset->GetSkeleton();
+					}
+				}
+
+				if (IBoneReferenceSkeletonProvider* SkeletonProvider = Cast<IBoneReferenceSkeletonProvider>(InObject))
+				{
+					bool bInvalidSkeletonIsError = false;
+					return SkeletonProvider->GetSkeleton(bInvalidSkeletonIsError, PropertyHandle.Get());
+				}
+			}
+
+			return nullptr;
+		};
+
+		for (UObject* Object : Objects)
+		{
+			if (USkeleton* Skeleton = FindSkeletonForObject(Object))
+			{
+				return Skeleton;
+			}
+		}
+
+		return nullptr;
+	}
+} // namespace UE::MLDeformer::TrainingDataProcessor
+
+#undef LOCTEXT_NAMESPACE
diff --git a/Engine/Plugins/Animation/MLDeformer/MLDeformerFramework/Source/MLDeformerFrameworkEditor/Public/MLDeformerEditorModel.h b/Engine/Plugins/Animation/MLDeformer/MLDeformerFramework/Source/MLDeformerFrameworkEditor/Public/MLDeformerEditorModel.h
--- a/Engine/Plugins/Animation/MLDeformer/MLDeformerFramework/Source/MLDeformerFrameworkEditor/Public/MLDeformerEditorModel.h
+++ b/Engine/Plugins/Animation/MLDeformer/MLDeformerFramework/Source/MLDeformerFrameworkEditor/Public/MLDeformerEditorModel.h
@@ -10,7 +10,6 @@
 #include "MLDeformerEditorActor.h"
 #include "Misc/FrameTime.h"
 #include "MLDeformerVizSettings.h"
-#include "MLDeformerModule.h"
 #include "MLDeformerModel.h"
 #include "MLDeformerSampler.h"
 #include "MeshDescription.h"
@@ -45,13 +44,13 @@ enum class ETrainingResult : uint8
 	/** The user has aborted the training process. */
 	Aborted,
 
-	/** The user has aborted the training process and we can't use the resulting network. */
+	/** The user has aborted the training process, and we can't use the resulting network. */
 	AbortedCantUse,
 
 	/** The input or output data to the network has issues, which means we cannot train. */
 	FailOnData,
 
-	/** The python script has some error (see output log). Or it's missing some required model training class. */
+	/** The python script has some error (see output log). Or it's missing some required model training class. Or an actual python code syntax error or failed imports. */
 	FailPythonError,
 
 	/** Any other error, for example when some exception got thrown during training. */
@@ -91,7 +90,7 @@ namespace UE::MLDeformer
 		 * The destructor, which unregisters the instance of this editor model from the model registry.
 		 * It also deletes the editor actors and unregisters from the OnPostEditChangeProperty delegate.
 		 */
-		virtual ~FMLDeformerEditorModel();
+		virtual ~FMLDeformerEditorModel() override;
 
 		// FGCObject overrides.
 		virtual void AddReferencedObjects(FReferenceCollector& Collector) override;
@@ -108,7 +107,7 @@ namespace UE::MLDeformer
 		/**
 		 * Change the skeletal mesh on a specific skeletal mesh component.
 		 * This internally will not only call SkelMeshComponent->SetSkeletalMesh(Mesh), but also will reassign the materials.
-		 * The reason for this is because if there are material instances set already, the skeletal mesh component will not override them to preserve changes done by the user.
+		 * The reason for this is that if there are material instances set already, the skeletal mesh component will not override them to preserve changes done by the user.
 		 * This helper function will force these to be updated as well. This was needed in case of Metahumans.
 		 * Metahumans use a PostProcessAnimBP with a ControlRig node in it, which drives material curves.
 		 * In order to drive this, material instances are needed instead of regular materials. But as mentioned before, when switching the mesh the 
@@ -321,7 +320,7 @@ namespace UE::MLDeformer
 
 		/**
 		 * Check if this model is compatible with a specific deformer asset.
-		 * For example it will check whether the other deformer asset's skeletal mesh is the same as the one used
+		 * For example, it will check whether the other deformer asset's skeletal mesh is the same as the one used
 		 * by this model.
 		 * @param Deformer The deformer asset to check against.
 		 * @return Returns true if the specified asset is compatible with this one, otherwise false is returned.
@@ -398,6 +397,22 @@ namespace UE::MLDeformer
 		 */
 		virtual void OnPropertyChanged(FPropertyChangedEvent& PropertyChangedEvent);
 
+		/**
+		 * This is called before undo is called on the model.
+		 */
+		virtual void OnPreEditUndo() {}
+
+		/**
+		 * This is called after undo is performed on the model.
+		 */
+		virtual void OnPostEditUndo() {}
+
+		/**
+		 * This is called whenever a transaction happened on the model.
+		 * @param Event Information about the event.
+		 */
+		virtual void OnPostTransacted(const FTransactionObjectEvent& Event);
+
 		/**
 		 * Executed when the user presses the play button in the timeline.
 		 * On default this will start playing the anim sequence related to whether we are in training or testing mode.
@@ -406,7 +421,7 @@ namespace UE::MLDeformer
 
 		/**
 		 * This is called just before we start the training process.
-		 * You can use it to backup any specific values that might be modified by the training process.
+		 * You can use it to back up any specific values that might be modified by the training process.
 		 * Sometimes you want to restore specific values when the user aborts the training, while the training already modified those values.
 		 */
 		virtual void OnPreTraining() {}
@@ -421,9 +436,9 @@ namespace UE::MLDeformer
 
 		/**
 		 * This is executed when the training process got aborted by the user.
-		 * You can use this method to restore any values you had to backup during OnPreTraining for example.
-		 * @param bUsePartiallyTrainedWhenAborted When the user aborts training, they can get offered to use the partially trained neural network. This boolean 
-		 *                                        specifies whether they picked "Yes" (true) or "No" (false) to that question.
+		 * You can use this method to restore any values you had to back up during OnPreTraining for example.
+		 * @param bUsePartiallyTrainedData When the user aborts training, they can get offered to use the partially trained neural network. This boolean 
+		 *                                 specifies whether they picked "Yes" (true) or "No" (false) to that question.
 		 */
 		virtual void OnTrainingAborted(bool bUsePartiallyTrainedData) {}
 
@@ -533,7 +548,7 @@ namespace UE::MLDeformer
 
 		/**
 		 * Update the state that is used to decide whether the Train button in the UI is enabled or not.
-		 * One criteria could be to check whether we have all required inputs, such as a skeletal mesh, target mesh and anim sequence.
+		 * A criteria could be to check whether we have all required inputs, such as a skeletal mesh, target mesh and anim sequence.
 		 * This should set the bIsReadyForTraining bool to true when it is ready, or false if not.
 		 */
 		virtual void UpdateIsReadyForTrainingState() { bIsReadyForTraining = false; }
@@ -696,6 +711,7 @@ namespace UE::MLDeformer
 		 * Generate the normals for a given morph target.
 		 * @param LOD The LOD level.
 		 * @param SkelMesh The skeletal mesh to get the mesh data from.
+		 * @param MorphTargetIndex The index of the morph target to generate the normals for.
 		 * @param Deltas The per vertex deltas. The number of elements in this array is NumMorphTargets * NumImportedVertices.
 		 * @param BaseVertexPositions The positions of the base/neutral mesh, basically the unskinned vertices.
 		 * @param BaseNormals The normals of the base mesh.
@@ -935,7 +951,7 @@ namespace UE::MLDeformer
 		/**
 		 * Get the editor's input info object.
 		 * This is the information the user currently has setup in the UI, so not the input info of the runtime model.
-		 * The runtime model's input info contains the data inputs the model has been trained on, which can be different than what the user currently has setup in the UI.
+		 * The runtime model's input info contains the data inputs the model has been trained on, which can be different from what the user currently has set up in the UI.
 		 * @return A pointer to the editor's input info object.
 		 */
 		UMLDeformerInputInfo* GetEditorInputInfo() const { return EditorInputInfo.Get(); }	
@@ -948,7 +964,7 @@ namespace UE::MLDeformer
 		/**
 		 * This will call OnInputAssetsChanged first, followed by OnPostInputAssetChanged.
 		 * It then will update the editor only flags for the assets, and then forces a refresh of the model details panel and optionally also the visualization settings panel.
-		 * @param RefreshVizSettings Set to true if you want the visualization settings details panel to also refresh.
+		 * @param bRefreshVizSettings Set to true if you want the visualization settings details panel to also refresh.
 		 */
 		void TriggerInputAssetChanged(bool bRefreshVizSettings=false);
 
@@ -986,7 +1002,7 @@ namespace UE::MLDeformer
 		/**
 		 * Get the currently desired training frame number.
 		 * You can call CheckTrainingFrameChanged in order to make the desired frame the actual frame.
-		 * @return The currently desired training frame number, which might not the the same as the real training frame number.
+		 * @return The currently desired training frame number, which might not be the same as the real training frame number.
 		 */
 		int32 GetCurrentTrainingFrame() const { return CurrentTrainingFrame; }
 
@@ -1018,15 +1034,15 @@ namespace UE::MLDeformer
 
 		/**
 		 * Correct floating point errors that can cause issues when sampling animation using step timing.
-		 * @param FrameNumber The desired frame number
-		 * @param FrameTimeForFrameNumber The frame time computed for the frame number, which may have conversion errors
-		 * @Param FrameRate The frame rate for conversion
-		 * @return A corrected frame time, which is equal to FrameTimeForFrameNumber unless a floating point error means the floor is not equal to the FrameNumber
+		 * @param FrameNumber The desired frame number.
+		 * @param FrameTimeForFrameNumber The frame time computed for the frame number, which may have conversion errors.
+		 * @param FrameRate The frame rate for conversion.
+		 * @return A corrected frame time, which is equal to FrameTimeForFrameNumber unless a floating point error means the floor is not equal to the FrameNumber.
 		 */
 		static float CorrectedFrameTime(int32 FrameNumber, float FrameTimeForFrameNumber, FFrameRate FrameRate);
 
-		TSharedPtr<SMLDeformerInputWidget> GetInputWidget() const			{ return InputWidget; }
-		void SetInputWidget(TSharedPtr<SMLDeformerInputWidget> Widget)		{ InputWidget = Widget; }
+		TSharedPtr<SMLDeformerInputWidget> GetInputWidget() const				{ return InputWidget; }
+		void SetInputWidget(const TSharedPtr<SMLDeformerInputWidget>& Widget)	{ InputWidget = Widget; }
 
 		FMLDeformerSampler* GetSamplerForActiveAnim() const;
 
@@ -1064,6 +1080,12 @@ namespace UE::MLDeformer
 		 */
 		void OnPostEditChangeProperty(FPropertyChangedEvent& PropertyChangedEvent);
 
+		/**
+		 * Trigger a PostEditChangeProperty using a ValueSet event type, on the model, which will also trigger the OnPostEditChangeProperty.
+		 * @param PropertyName The name of the property to trigger it for.
+		 */
+		void BroadcastPropertyChanged(const FName PropertyName);
+
 		/**
 		 * Called whenever an engine object is being modified.
 		 * Use this to set the reinitialize assets flag. For example when a skeletal mesh or anim sequence change (being modified or reimported), we should 
@@ -1108,8 +1130,6 @@ namespace UE::MLDeformer
 		 *        recomputing normals can lead to lower quality results, in trade for faster performance.
 		 * @param MaskChannel The weight mask mode, which specifies what channel to get the weight data from. Such channel allows the user to define what areas the deformer should for example not be active in.
 		 * @param bInvertMaskChannel Specifies whether the weight mask should be inverted or not.
-		 * @param MaskBuffer An optional mask buffer that contains 'Model->GetNumBaseMeshVerts() * (MorphModel->GetNumMorphTargets()- 1)' number of floats. Deltas will be multiplied by this value. 
-		 *        When the mask buffer is an empty array, it will be ignored.
 		 */
 		UE_DEPRECATED(5.3, "Please call FMLDeformerMorphModelEditorModel::CreateMorphTargets instead.")
 		void CreateEngineMorphTargets(
@@ -1145,22 +1165,22 @@ namespace UE::MLDeformer
 
 		/**
 		 * Get the base actor depending on the currently active mode (train or test mode).
-		 * The base actor is the linear skinned one.
+		 * The base actor is the linear-skinned one.
 		 * @return A pointer to the linear skinned actor in either test or training mode.
 		 */
 		FMLDeformerEditorActor* GetVisualizationModeBaseActor() const;
 
 		/**
 		 * Get the animation sequence, either the test or training one, depending on in which mode we are in.
-		 * This can return a nullptr if no anim sequence has been setup yet.
+		 * This can return a nullptr if no anim sequence has been set up yet.
 		 * @return A pointer to the test sequence when we are in test mode, or the training anim sequence if we are in training mode.
 		 */
 		const UAnimSequence* GetAnimSequence() const;
 
 		/**
-		 * Calculate the current time line position offset, of the currently selected frame.
+		 * Calculate the current timeline position offset, of the currently selected frame.
 		 * This depends on whether we are in test or training mode.
-		 * @return The time line offset, in seconds.
+		 * @return The timeline offset, in seconds.
 		 */
 		double CalcTimelinePosition() const;
 
@@ -1246,6 +1266,15 @@ namespace UE::MLDeformer
 		/** The delegate handle to the post edit property event. */
 		FDelegateHandle PostEditPropertyDelegateHandle;
 
+		/** The delegate handle linked to PostTransacted of the model. */
+		FDelegateHandle PostTransactedDelegateHandle;
+
+		/** The delegate handle linked to PreEditUndo of the model. */
+		FDelegateHandle PreEditUndoDelegateHandle;
+
+		/** The delegate handle linked to PostEditUndo of the model. */
+		FDelegateHandle PostEditUndoDelegateHandle;
+
 		/** The current training frame. */
 		int32 CurrentTrainingFrame = -1;
 
diff --git a/Engine/Plugins/Animation/MLDeformer/MLDeformerFramework/Source/MLDeformerFrameworkEditor/Public/MLDeformerEditorModule.h b/Engine/Plugins/Animation/MLDeformer/MLDeformerFramework/Source/MLDeformerFrameworkEditor/Public/MLDeformerEditorModule.h
--- a/Engine/Plugins/Animation/MLDeformer/MLDeformerFramework/Source/MLDeformerFrameworkEditor/Public/MLDeformerEditorModule.h
+++ b/Engine/Plugins/Animation/MLDeformer/MLDeformerFramework/Source/MLDeformerFrameworkEditor/Public/MLDeformerEditorModule.h
@@ -8,6 +8,11 @@
 
 namespace UE::MLDeformer
 {
+	namespace TrainingDataProcessor
+	{
+		class FTrainingDataProcessorTool;
+	}
+
 	/**
 	 * The ML Deformer editor module.
 	 * This registers the editor mode and custom property customizations.
@@ -18,8 +23,8 @@ namespace UE::MLDeformer
 	{
 	public:
 		// IModuleInterface overrides.
-		void StartupModule() override;
-		void ShutdownModule() override;
+		virtual void StartupModule() override;
+		virtual void ShutdownModule() override;
 		// ~END IModuleInterface overrides.
 
 		/**
diff --git a/Engine/Plugins/Animation/MLDeformer/MLDeformerFramework/Source/MLDeformerFrameworkEditor/Public/MLDeformerEditorToolkit.h b/Engine/Plugins/Animation/MLDeformer/MLDeformerFramework/Source/MLDeformerFrameworkEditor/Public/MLDeformerEditorToolkit.h
--- a/Engine/Plugins/Animation/MLDeformer/MLDeformerFramework/Source/MLDeformerFrameworkEditor/Public/MLDeformerEditorToolkit.h
+++ b/Engine/Plugins/Animation/MLDeformer/MLDeformerFramework/Source/MLDeformerFrameworkEditor/Public/MLDeformerEditorToolkit.h
@@ -87,6 +87,11 @@ namespace UE::MLDeformer
 		virtual TStatId GetStatId() const override;
 		// ~END FTickableEditorObject overrides.
 
+		// FEditorUndoClient overrides.
+		void PostUndo(bool bSuccess) override;
+		void PostRedo(bool bSuccess) override;
+		// ~END FEditorUndoClient overrides.
+
 		// IHasPersonaToolkit overrides.
 		virtual TSharedRef<IPersonaToolkit> GetPersonaToolkit() const override;
 		IPersonaToolkit* GetPersonaToolkitPointer() const							{ return PersonaToolkit.Get(); }
diff --git a/Engine/Plugins/Animation/MLDeformer/MLDeformerFramework/Source/MLDeformerFrameworkEditor/Public/SMLDeformerBonePickerDialog.h b/Engine/Plugins/Animation/MLDeformer/MLDeformerFramework/Source/MLDeformerFrameworkEditor/Public/SMLDeformerBonePickerDialog.h
--- a/Engine/Plugins/Animation/MLDeformer/MLDeformerFramework/Source/MLDeformerFrameworkEditor/Public/SMLDeformerBonePickerDialog.h
+++ b/Engine/Plugins/Animation/MLDeformer/MLDeformerFramework/Source/MLDeformerFrameworkEditor/Public/SMLDeformerBonePickerDialog.h
@@ -82,7 +82,7 @@ namespace UE::MLDeformer
 		friend class SMLDeformerBonePickerTreeWidget;
 
 		SLATE_BEGIN_ARGS(SMLDeformerBonePickerDialog) {}
-		SLATE_ARGUMENT(FReferenceSkeleton*, RefSkeleton)		// The reference skeleton to grab the bone names and hierarchy from.
+		SLATE_ARGUMENT(FReferenceSkeleton const*, RefSkeleton)		// The reference skeleton to grab the bone names and hierarchy from.
 		SLATE_ARGUMENT(TArray<FName>, HighlightBoneNames)		// Highlight bones in this list.
 		SLATE_ARGUMENT(TArray<FName>, InitialSelectedBoneNames)	// Select all bones in this list.
 		SLATE_ARGUMENT(FSlateColor, HighlightBoneNamesColor)	// Bones inside the HighlightBoneNames list will use this color.
@@ -107,7 +107,7 @@ namespace UE::MLDeformer
 		void OnOkClicked();
 		void OnCancelClicked();
 		void OnFilterTextChanged(const FText& InFilterText);
-		void RefreshBoneTreeElements();
+		void RefreshBoneTreeElements() const;
 		void SelectInitialItems();
 		void SelectInitialItemsRecursive(const TSharedPtr<FMLDeformerBonePickerTreeElement>& Item);
 
@@ -117,7 +117,7 @@ namespace UE::MLDeformer
 		TArray<FName> HighlightBoneNames;
 		TArray<FName> InitialSelectedBoneNames;
 		TArray<FName> IncludeList;
-		FReferenceSkeleton* RefSkeleton = nullptr;
+		FReferenceSkeleton const* RefSkeleton = nullptr;
 		TArray<FName> PickedBoneNames;
 		FSlateColor HighlightColor;
 		bool bAllowMultiSelect = true;
diff --git a/Engine/Plugins/Animation/MLDeformer/MLDeformerFramework/Source/MLDeformerFrameworkEditor/Public/SMLDeformerCurvePickerDialog.h b/Engine/Plugins/Animation/MLDeformer/MLDeformerFramework/Source/MLDeformerFrameworkEditor/Public/SMLDeformerCurvePickerDialog.h
--- a/Engine/Plugins/Animation/MLDeformer/MLDeformerFramework/Source/MLDeformerFrameworkEditor/Public/SMLDeformerCurvePickerDialog.h
+++ b/Engine/Plugins/Animation/MLDeformer/MLDeformerFramework/Source/MLDeformerFrameworkEditor/Public/SMLDeformerCurvePickerDialog.h
@@ -100,7 +100,7 @@ namespace UE::MLDeformer
 		void OnOkClicked();
 		void OnCancelClicked();
 		void OnFilterTextChanged(const FText& InFilterText);
-		void RefreshListElements();
+		void RefreshListElements() const;
 
 	private:
 		TSharedPtr<SMLDeformerCurvePickerListView> CurveListView;
diff --git a/Engine/Plugins/Animation/MLDeformer/MLDeformerFramework/Source/MLDeformerFrameworkEditor/Public/SMLDeformerInputBonesWidget.h b/Engine/Plugins/Animation/MLDeformer/MLDeformerFramework/Source/MLDeformerFrameworkEditor/Public/SMLDeformerInputBonesWidget.h
--- a/Engine/Plugins/Animation/MLDeformer/MLDeformerFramework/Source/MLDeformerFrameworkEditor/Public/SMLDeformerInputBonesWidget.h
+++ b/Engine/Plugins/Animation/MLDeformer/MLDeformerFramework/Source/MLDeformerFrameworkEditor/Public/SMLDeformerInputBonesWidget.h
@@ -84,6 +84,7 @@ namespace UE::MLDeformer
 		void RecursiveSortElements(TSharedPtr<FMLDeformerInputBoneTreeElement> Element);
 		void RecursiveAddNames(const FMLDeformerInputBoneTreeElement& Element, TArray<FName>& OutNames) const;
 		void OnSelectionChanged(TSharedPtr<FMLDeformerInputBoneTreeElement> Selection, ESelectInfo::Type SelectInfo);
+		TSharedPtr<FMLDeformerInputBoneTreeElement> FindParentElementForBone(FName BoneName, const FReferenceSkeleton& RefSkeleton, const TMap<FName, TSharedPtr<FMLDeformerInputBoneTreeElement>>& NameToElementMap) const;
 
 	private:
 		TArray<TSharedPtr<FMLDeformerInputBoneTreeElement>> RootElements;
diff --git a/Engine/Plugins/Animation/MLDeformer/MLDeformerFramework/Source/MLDeformerFrameworkEditor/Public/SMLDeformerInputWidget.h b/Engine/Plugins/Animation/MLDeformer/MLDeformerFramework/Source/MLDeformerFrameworkEditor/Public/SMLDeformerInputWidget.h
--- a/Engine/Plugins/Animation/MLDeformer/MLDeformerFramework/Source/MLDeformerFrameworkEditor/Public/SMLDeformerInputWidget.h
+++ b/Engine/Plugins/Animation/MLDeformer/MLDeformerFramework/Source/MLDeformerFrameworkEditor/Public/SMLDeformerInputWidget.h
@@ -67,6 +67,7 @@ namespace UE::MLDeformer
 
 		// Registers the commands that are being executed by the different widgets.
 		static void RegisterCommands();
+		static void UnregisterCommands();
 
 	protected:
 		void AddSectionSeparator();
diff --git a/Engine/Plugins/Animation/MLDeformer/MLDeformerFramework/Source/MLDeformerFrameworkEditor/Public/Tools/TrainingDataProcessor/AnimCustomization.h b/Engine/Plugins/Animation/MLDeformer/MLDeformerFramework/Source/MLDeformerFrameworkEditor/Public/Tools/TrainingDataProcessor/AnimCustomization.h
new file mode 100644
--- /dev/null
+++ b/Engine/Plugins/Animation/MLDeformer/MLDeformerFramework/Source/MLDeformerFrameworkEditor/Public/Tools/TrainingDataProcessor/AnimCustomization.h
@@ -0,0 +1,35 @@
+// Copyright Epic Games, Inc. All Rights Reserved.
+
+#pragma once
+
+#include "IPropertyTypeCustomization.h"
+#include "PropertyHandle.h"
+#include "DetailWidgetRow.h"
+
+class UMLDeformerTrainingDataProcessorSettings;
+
+namespace UE::MLDeformer::TrainingDataProcessor
+{
+	/**
+	 * The property detail customization for a single animation input that can be enabled or disabled.
+	 * This is a customization for the type FMLDeformerTrainingDataProcessorAnim.
+	 */
+	class MLDEFORMERFRAMEWORKEDITOR_API FAnimCustomization final : public IPropertyTypeCustomization
+	{
+	public:
+		static TSharedRef<IPropertyTypeCustomization> MakeInstance();
+
+		//~ Begin IPropertyTypeCustomization overrides.
+		virtual void CustomizeHeader(TSharedRef<IPropertyHandle> StructPropertyHandle, class FDetailWidgetRow& HeaderRow, IPropertyTypeCustomizationUtils& StructCustomizationUtils) override;
+		virtual void CustomizeChildren(TSharedRef<IPropertyHandle> StructPropertyHandle, class IDetailChildrenBuilder& ChildBuilder, IPropertyTypeCustomizationUtils& StructCustomizationUtils) override;
+		EVisibility GetAnimErrorVisibility(const TSharedRef<IPropertyHandle> StructPropertyHandle, int32 AnimIndex) const;
+		// ~End IPropertyTypeCustomization overrides.
+
+	private:
+		void RefreshProperties() const;
+		static UMLDeformerTrainingDataProcessorSettings* FindSettings(const TSharedRef<IPropertyHandle>& StructPropertyHandle);
+
+	private:
+		TSharedPtr<IPropertyUtilities> PropertyUtilities;
+	};
+}	// namespace UE::MLDeformer::TrainingDataProcessor
diff --git a/Engine/Plugins/Animation/MLDeformer/MLDeformerFramework/Source/MLDeformerFrameworkEditor/Public/Tools/TrainingDataProcessor/BoneGroupsListCustomization.h b/Engine/Plugins/Animation/MLDeformer/MLDeformerFramework/Source/MLDeformerFrameworkEditor/Public/Tools/TrainingDataProcessor/BoneGroupsListCustomization.h
new file mode 100644
--- /dev/null
+++ b/Engine/Plugins/Animation/MLDeformer/MLDeformerFramework/Source/MLDeformerFrameworkEditor/Public/Tools/TrainingDataProcessor/BoneGroupsListCustomization.h
@@ -0,0 +1,37 @@
+// Copyright Epic Games, Inc. All Rights Reserved.
+
+#pragma once
+
+#include "IPropertyTypeCustomization.h"
+#include "PropertyHandle.h"
+#include "MLDeformerTrainingDataProcessorSettings.h"
+
+class UMLDeformerTrainingDataProcessorSettings;
+class USkeleton;
+
+namespace UE::MLDeformer::TrainingDataProcessor
+{
+	/**
+	 * The property detail customization for a list of bone groups.
+	 * This is a customization for the type FMLDeformerTrainingDataProcessorBoneGroupsList.
+	 */
+	class MLDEFORMERFRAMEWORKEDITOR_API FBoneGroupsListCustomization : public IPropertyTypeCustomization
+	{
+	public:
+		static TSharedRef<IPropertyTypeCustomization> MakeInstance();
+
+		//~ Begin IPropertyTypeCustomization overrides.
+		virtual void CustomizeHeader(TSharedRef<IPropertyHandle> StructPropertyHandle, class FDetailWidgetRow& HeaderRow,
+		                             IPropertyTypeCustomizationUtils& StructCustomizationUtils) override;
+		virtual void CustomizeChildren(TSharedRef<IPropertyHandle> StructPropertyHandle, class IDetailChildrenBuilder& ChildBuilder,
+		                               IPropertyTypeCustomizationUtils& StructCustomizationUtils) override {}
+		//~ End IPropertyTypeCustomization overrides.
+
+	private:
+		TArray<FMLDeformerTrainingDataProcessorBoneGroup>* GetBoneGroups() const;
+
+	private:
+		TSharedPtr<IPropertyUtilities> PropertyUtilities;
+		TSharedPtr<IPropertyHandle> StructProperty;
+	};
+} // namespace UE::MLDeformer::TrainingDataProcessor
diff --git a/Engine/Plugins/Animation/MLDeformer/MLDeformerFramework/Source/MLDeformerFrameworkEditor/Public/Tools/TrainingDataProcessor/BoneListCustomization.h b/Engine/Plugins/Animation/MLDeformer/MLDeformerFramework/Source/MLDeformerFrameworkEditor/Public/Tools/TrainingDataProcessor/BoneListCustomization.h
new file mode 100644
--- /dev/null
+++ b/Engine/Plugins/Animation/MLDeformer/MLDeformerFramework/Source/MLDeformerFrameworkEditor/Public/Tools/TrainingDataProcessor/BoneListCustomization.h
@@ -0,0 +1,36 @@
+// Copyright Epic Games, Inc. All Rights Reserved.
+
+#pragma once
+
+#include "IPropertyTypeCustomization.h"
+#include "PropertyHandle.h"
+
+class UMLDeformerTrainingDataProcessorSettings;
+class USkeleton;
+
+namespace UE::MLDeformer::TrainingDataProcessor
+{
+	/**
+	 * The property detail customization for a list of bones.
+	 * This is a detail customization for the type FMLDeformerTrainingDataProcessorBoneList.
+	 */
+	class MLDEFORMERFRAMEWORKEDITOR_API FBoneListCustomization : public IPropertyTypeCustomization
+	{
+	public:
+		static TSharedRef<IPropertyTypeCustomization> MakeInstance();
+
+		//~ Begin IPropertyTypeCustomization overrides.
+		virtual void CustomizeHeader(TSharedRef<IPropertyHandle> StructPropertyHandle, class FDetailWidgetRow& HeaderRow,
+		                             IPropertyTypeCustomizationUtils& StructCustomizationUtils) override;
+		virtual void CustomizeChildren(TSharedRef<IPropertyHandle> StructPropertyHandle, class IDetailChildrenBuilder& ChildBuilder,
+		                               IPropertyTypeCustomizationUtils& StructCustomizationUtils) override { }
+		// ~End IPropertyTypeCustomization overrides.
+
+	private:
+		TArray<FName>* GetBoneNames() const;
+
+	private:
+		TSharedPtr<IPropertyUtilities> PropertyUtilities;
+		TSharedPtr<IPropertyHandle> StructProperty;
+	};
+} // namespace UE::MLDeformer::TrainingDataProcessor
diff --git a/Engine/Plugins/Animation/MLDeformer/MLDeformerFramework/Source/MLDeformerFrameworkEditor/Public/Tools/TrainingDataProcessor/SBoneGroupsListWidget.h b/Engine/Plugins/Animation/MLDeformer/MLDeformerFramework/Source/MLDeformerFrameworkEditor/Public/Tools/TrainingDataProcessor/SBoneGroupsListWidget.h
new file mode 100644
--- /dev/null
+++ b/Engine/Plugins/Animation/MLDeformer/MLDeformerFramework/Source/MLDeformerFrameworkEditor/Public/Tools/TrainingDataProcessor/SBoneGroupsListWidget.h
@@ -0,0 +1,166 @@
+// Copyright Epic Games, Inc. All Rights Reserved.
+
+#pragma once
+
+#include "CoreTypes.h"
+#include "Templates/SharedPointer.h"
+#include "Widgets/SCompoundWidget.h"
+#include "Widgets/Views/STreeView.h"
+#include "Framework/Commands/Commands.h"
+#include "MLDeformerTrainingDataProcessorSettings.h"
+#include "Styling/SlateColor.h"
+#include "EditorUndoClient.h"
+
+class USkeleton;
+class FUICommandList;
+class FNotifyHook;
+struct FReferenceSkeleton;
+
+namespace UE::MLDeformer::TrainingDataProcessor
+{
+	class SBoneGroupsTreeWidget;
+	class SBoneGroupsListWidget;
+	
+	class MLDEFORMERFRAMEWORKEDITOR_API FBoneGroupsListWidgetCommands final : public TCommands<FBoneGroupsListWidgetCommands>
+	{
+	public:
+		FBoneGroupsListWidgetCommands();
+		virtual void RegisterCommands() override;
+
+	public:
+		TSharedPtr<FUICommandInfo> CreateGroup;
+		TSharedPtr<FUICommandInfo> DeleteSelectedItems;
+		TSharedPtr<FUICommandInfo> ClearGroups;
+		TSharedPtr<FUICommandInfo> AddBoneToGroup;
+	};
+
+
+	class MLDEFORMERFRAMEWORKEDITOR_API FBoneGroupTreeElement final : public TSharedFromThis<FBoneGroupTreeElement>
+	{
+	public:
+		static TSharedRef<ITableRow> MakeTreeRowWidget(const TSharedRef<STableViewBase>& InOwnerTable,
+		                                               const TSharedRef<FBoneGroupTreeElement>& InTreeElement,
+		                                               const TSharedPtr<SBoneGroupsTreeWidget>& InTreeWidget);
+
+		bool IsGroup() const { return GroupIndex != INDEX_NONE; }
+
+	public:
+		FString Name;
+		TArray<TSharedPtr<FBoneGroupTreeElement>> Children;
+		TWeakPtr<FBoneGroupTreeElement> ParentGroup;
+		FSlateColor TextColor;
+		int32 GroupIndex = INDEX_NONE;
+		int32 GroupBoneIndex = INDEX_NONE;
+	};
+
+
+	class MLDEFORMERFRAMEWORKEDITOR_API SBoneGroupTreeRowWidget final : public STableRow<TSharedPtr<FBoneGroupTreeElement>>
+	{
+	public:
+		void Construct(const FArguments& InArgs, const TSharedRef<STableViewBase>& OwnerTable, const TSharedRef<FBoneGroupTreeElement>& InTreeElement,
+		               const TSharedPtr<SBoneGroupsTreeWidget>& InTreeView);
+
+	private:
+		FText GetName() const;
+
+	private:
+		TWeakPtr<FBoneGroupTreeElement> WeakTreeElement;
+		friend class SBoneGroupsTreeWidget;
+	};
+
+
+	class MLDEFORMERFRAMEWORKEDITOR_API SBoneGroupsTreeWidget final : public STreeView<TSharedPtr<FBoneGroupTreeElement>>
+	{
+	public:
+		friend class SBoneGroupsListWidget;
+
+		SLATE_BEGIN_ARGS(SBoneGroupsTreeWidget)
+			{
+			}
+
+			SLATE_ARGUMENT(TSharedPtr<SBoneGroupsListWidget>, BoneGroupsWidget)
+		SLATE_END_ARGS()
+
+		void Construct(const FArguments& InArgs);
+		void Refresh();
+		int32 GetNumSelectedGroups() const;
+
+	private:
+		//~ Begin STreeView overrides.
+		virtual FReply OnKeyDown(const FGeometry& InGeometry, const FKeyEvent& InKeyEvent) override;
+		//~ End STreeView overrides.
+
+		void AddElement(const TSharedPtr<FBoneGroupTreeElement>& Element, const TSharedPtr<FBoneGroupTreeElement>& ParentElement);
+		const TArray<TSharedPtr<FBoneGroupTreeElement>>& GetRootElements() const;
+		TSharedRef<ITableRow> MakeTableRowWidget(TSharedPtr<FBoneGroupTreeElement> InItem, const TSharedRef<STableViewBase>& OwnerTable);
+		TSharedPtr<SWidget> CreateContextMenuWidget() const;
+		void UpdateTreeElements();
+		void RefreshTree();
+
+		static void HandleGetChildrenForTree(TSharedPtr<FBoneGroupTreeElement> InItem, TArray<TSharedPtr<FBoneGroupTreeElement>>& OutChildren);
+		static TSharedPtr<SWidget> CreateContextWidget();
+
+	private:
+		TArray<TSharedPtr<FBoneGroupTreeElement>> RootElements;
+		TWeakPtr<SBoneGroupsListWidget> BoneGroupsWidget;
+	};
+
+
+	DECLARE_DELEGATE_RetVal(TArray<FMLDeformerTrainingDataProcessorBoneGroup>*, FBoneGroupsListWidgetGetBoneGroups)
+
+	/**
+	 * A widget that shows a set of bone groups, and allows you to manage them by creating, removing and editing of groups.
+	 * We see a bone group as a list of bone names. Multiple bone groups can exist. If you need only one list of bones
+	 * then you can use the bone SBoneListWidget instead.
+	 */
+	class MLDEFORMERFRAMEWORKEDITOR_API SBoneGroupsListWidget final : public SCompoundWidget, public FEditorUndoClient
+	{
+		friend class SBoneGroupsTreeWidget;
+
+		SLATE_BEGIN_ARGS(SBoneGroupsListWidget)
+			{
+			}
+
+			SLATE_ARGUMENT(TWeakObjectPtr<USkeleton>, Skeleton)
+			SLATE_ARGUMENT(TWeakObjectPtr<UObject>, UndoObject)
+			SLATE_EVENT(FBoneGroupsListWidgetGetBoneGroups, GetBoneGroups)
+		SLATE_END_ARGS()
+
+	public:
+		virtual ~SBoneGroupsListWidget() override;
+
+		void Construct(const FArguments& InArgs, FNotifyHook* InNotifyHook);
+
+	private:
+		//~ Begin FEditorUndoClient overrides.
+		virtual void PostUndo(bool bSuccess) override;
+		virtual void PostRedo(bool bSuccess) override;
+		//~ End FEditorUndoClient overrides.
+
+		void BindCommands(const TSharedPtr<FUICommandList>& InCommandList);
+		void OnFilterTextChanged(const FText& InFilterText);
+		void RefreshTree() const;
+		FReply OnAddButtonClicked() const;
+		FReply OnClearButtonClicked() const;
+
+		TSharedPtr<SBoneGroupsTreeWidget> GetTreeWidget() const;
+		TSharedPtr<FUICommandList> GetCommandList() const;
+		TArray<FMLDeformerTrainingDataProcessorBoneGroup>* GetBoneGroupsValues() const;
+		TWeakObjectPtr<USkeleton> GetSkeleton() const;
+		const FString& GetFilterText() const;
+
+		void OnCreateBoneGroup() const;
+		void OnDeleteSelectedItems() const;
+		void OnClearBoneGroups() const;
+		void OnAddBoneToGroup() const;
+
+	private:
+		TSharedPtr<SBoneGroupsTreeWidget> TreeWidget;
+		TWeakObjectPtr<USkeleton> Skeleton;
+		TWeakObjectPtr<UObject> UndoObject;
+		TSharedPtr<FUICommandList> CommandList;
+		FBoneGroupsListWidgetGetBoneGroups GetBoneGroups;
+		FString FilterText;
+		FNotifyHook* NotifyHook = nullptr;
+	};
+} // namespace UE::MLDeformer::TrainingDataProcessor
diff --git a/Engine/Plugins/Animation/MLDeformer/MLDeformerFramework/Source/MLDeformerFrameworkEditor/Public/Tools/TrainingDataProcessor/SBoneListWidget.h b/Engine/Plugins/Animation/MLDeformer/MLDeformerFramework/Source/MLDeformerFrameworkEditor/Public/Tools/TrainingDataProcessor/SBoneListWidget.h
new file mode 100644
--- /dev/null
+++ b/Engine/Plugins/Animation/MLDeformer/MLDeformerFramework/Source/MLDeformerFrameworkEditor/Public/Tools/TrainingDataProcessor/SBoneListWidget.h
@@ -0,0 +1,149 @@
+// Copyright Epic Games, Inc. All Rights Reserved.
+
+#pragma once
+
+#include "CoreMinimal.h"
+#include "Widgets/SCompoundWidget.h"
+#include "Widgets/Views/SListView.h"
+#include "Widgets/Views/STreeView.h"
+#include "Framework/Commands/Commands.h"
+#include "Framework/MultiBox/MultiBoxBuilder.h"
+#include "UObject/NameTypes.h"
+#include "Styling/SlateColor.h"
+#include "EditorUndoClient.h"
+
+class USkeleton;
+class FUICommandList;
+class FNotifyHook;
+struct FReferenceSkeleton;
+
+namespace UE::MLDeformer::TrainingDataProcessor
+{
+	class SBoneTreeWidget;
+	class SBoneListWidget;
+
+	class MLDEFORMERFRAMEWORKEDITOR_API FBoneListWidgetCommands final : public TCommands<FBoneListWidgetCommands>
+	{
+	public:
+		FBoneListWidgetCommands();
+		virtual void RegisterCommands() override;
+
+	public:
+		TSharedPtr<FUICommandInfo> AddBones;
+		TSharedPtr<FUICommandInfo> RemoveBones;
+		TSharedPtr<FUICommandInfo> ClearBones;
+	};
+
+
+	class MLDEFORMERFRAMEWORKEDITOR_API FBoneTreeWidgetElement final : public TSharedFromThis<FBoneTreeWidgetElement>
+	{
+	public:
+		TSharedRef<ITableRow> MakeTreeRowWidget(const TSharedRef<STableViewBase>& InOwnerTable, TSharedRef<FBoneTreeWidgetElement> InTreeElement,
+		                                        TSharedPtr<SBoneTreeWidget> InTreeWidget);
+
+	public:
+		FName Name;
+		TArray<TSharedPtr<FBoneTreeWidgetElement>> Children;
+		FSlateColor TextColor;
+	};
+
+
+	class MLDEFORMERFRAMEWORKEDITOR_API SBoneTreeRowWidget final : public STableRow<TSharedPtr<FBoneTreeWidgetElement>>
+	{
+	public:
+		void Construct(const FArguments& InArgs, const TSharedRef<STableViewBase>& OwnerTable, const TSharedRef<FBoneTreeWidgetElement>& InTreeElement,
+		               const TSharedPtr<SBoneTreeWidget>& InTreeView);
+
+	private:
+		TWeakPtr<FBoneTreeWidgetElement> TreeElement;
+		FText GetName() const;
+	};
+
+
+	class MLDEFORMERFRAMEWORKEDITOR_API SBoneTreeWidget final : public STreeView<TSharedPtr<FBoneTreeWidgetElement>>
+	{
+		SLATE_BEGIN_ARGS(SBoneTreeWidget) { }
+			SLATE_ARGUMENT(TSharedPtr<SBoneListWidget>, BoneListWidget)
+		SLATE_END_ARGS()
+
+	public:
+		void Construct(const FArguments& InArgs);
+		void RefreshElements(const TArray<FName>& BoneNames, const FReferenceSkeleton* RefSkeleton, const FString& FilterText);
+		TArray<FName> ExtractAllElementNames() const;
+		const TArray<TSharedPtr<FBoneTreeWidgetElement>>& GetRootElements() const { return RootElements; }
+
+	private:
+		virtual FReply OnKeyDown(const FGeometry& InGeometry, const FKeyEvent& InKeyEvent) override;
+		
+		static void RecursiveAddNames(const FBoneTreeWidgetElement& Element, TArray<FName>& OutNames);
+		TSharedRef<ITableRow> MakeTableRowWidget(TSharedPtr<FBoneTreeWidgetElement> InItem, const TSharedRef<STableViewBase>& OwnerTable);
+		TSharedPtr<SWidget> OnContextMenuOpening() const;
+		static TSharedPtr<FBoneTreeWidgetElement> FindParentElementForBone(FName BoneName, const FReferenceSkeleton& RefSkeleton, const TMap<FName, TSharedPtr<FBoneTreeWidgetElement>>& NameToElementMap);
+		void RecursiveSortElements(const TSharedPtr<FBoneTreeWidgetElement>& Element);
+		static void HandleGetChildrenForTree(TSharedPtr<FBoneTreeWidgetElement> InItem, TArray<TSharedPtr<FBoneTreeWidgetElement>>& OutChildren);
+
+	private:
+		TArray<TSharedPtr<FBoneTreeWidgetElement>> RootElements;
+		TWeakPtr<SBoneListWidget> BoneListWidget;
+	};
+
+
+	DECLARE_DELEGATE_OneParam(FOnBoneListWidgetBonesAdded, const TArray<FName>& BoneNames)
+	DECLARE_DELEGATE_OneParam(FOnBoneListWidgetBonesRemoved, const TArray<FName>& BoneNames)
+	DECLARE_DELEGATE(FOnBoneListWidgetBonesCleared)
+	DECLARE_DELEGATE_RetVal(TArray<FName>*, FBoneListWidgetGetBoneNames)
+
+	/**
+	 * The bone list widget, which displays a list of bones (in a hierarchy) and allows you to add and remove bones.
+	 * It works directly on a TArray<FName> as source. You can use the GetBoneNames event to provide this array.
+	 * You can use some of the delegates to listen for changes to the array.
+	 */
+	class MLDEFORMERFRAMEWORKEDITOR_API SBoneListWidget final : public SCompoundWidget, public FEditorUndoClient
+	{
+		SLATE_BEGIN_ARGS(SBoneListWidget) {}
+			SLATE_ARGUMENT(TWeakObjectPtr<USkeleton>, Skeleton)
+			SLATE_ARGUMENT(TWeakObjectPtr<UObject>, UndoObject)
+			SLATE_EVENT(FOnBoneListWidgetBonesAdded, OnBonesAdded)
+			SLATE_EVENT(FOnBoneListWidgetBonesRemoved, OnBonesRemoved)
+			SLATE_EVENT(FOnBoneListWidgetBonesCleared, OnBonesCleared)
+			SLATE_EVENT(FBoneListWidgetGetBoneNames, GetBoneNames)
+		SLATE_END_ARGS()
+
+	public:
+		virtual ~SBoneListWidget() override;
+
+		void Construct(const FArguments& InArgs, FNotifyHook* InNotifyHook);
+		void Refresh() const;
+		TSharedPtr<SBoneTreeWidget> GetTreeWidget() const;
+		TSharedPtr<FUICommandList> GetCommandList() const;
+
+	private:
+		// FEditorUndoClient overrides.
+		virtual void PostUndo(bool bSuccess) override;
+		virtual void PostRedo(bool bSuccess) override;
+		// ~END FEditorUndoClient overrides.
+
+		void BindCommands(const TSharedPtr<FUICommandList>& InCommandList);
+		void OnFilterTextChanged(const FText& InFilterText);
+		void RefreshTree() const;
+		FReply OnAddBonesButtonClicked() const;
+		FReply OnClearBonesButtonClicked() const;
+		void NotifyPropertyChanged() const;
+
+		void OnAddBones() const;
+		void OnRemoveBones() const;
+		void OnClearBones() const;
+
+	private:
+		TSharedPtr<SBoneTreeWidget> TreeWidget;
+		TWeakObjectPtr<USkeleton> Skeleton;
+		TWeakObjectPtr<UObject> UndoObject;
+		TSharedPtr<FUICommandList> CommandList;
+		FNotifyHook* NotifyHook = nullptr;
+		FString FilterText;
+		FOnBoneListWidgetBonesAdded OnBonesAdded;
+		FOnBoneListWidgetBonesRemoved OnBonesRemoved;
+		FOnBoneListWidgetBonesCleared OnBonesCleared;
+		FBoneListWidgetGetBoneNames GetBoneNames;
+	};
+} // namespace UE::MLDeformer::TrainingDataProcessor
diff --git a/Engine/Plugins/Animation/MLDeformer/MLDeformerFramework/Source/MLDeformerFrameworkEditor/Public/Tools/TrainingDataProcessor/STrainingDataProcessorMainWidget.h b/Engine/Plugins/Animation/MLDeformer/MLDeformerFramework/Source/MLDeformerFrameworkEditor/Public/Tools/TrainingDataProcessor/STrainingDataProcessorMainWidget.h
new file mode 100644
--- /dev/null
+++ b/Engine/Plugins/Animation/MLDeformer/MLDeformerFramework/Source/MLDeformerFrameworkEditor/Public/Tools/TrainingDataProcessor/STrainingDataProcessorMainWidget.h
@@ -0,0 +1,57 @@
+// Copyright Epic Games, Inc. All Rights Reserved.
+
+#pragma once
+
+#include "CoreMinimal.h"
+#include "EditorUndoClient.h"
+#include "MLDeformerModel.h"
+#include "Misc/NotifyHook.h"
+#include "Templates/SharedPointer.h"
+#include "Widgets/SCompoundWidget.h"
+
+class IDetailsView;
+
+namespace UE::MLDeformer::TrainingDataProcessor
+{
+	/**
+	 * The main widget for the training data processor tool.
+	 * This widget is basically what's inside the tab when this tool opens.
+	 * It contains a detail view and generate button.
+	 */
+	class MLDEFORMERFRAMEWORKEDITOR_API STrainingDataProcessorMainWidget final : public SCompoundWidget, public FEditorUndoClient, public FNotifyHook
+	{
+	public:
+		SLATE_BEGIN_ARGS(STrainingDataProcessorMainWidget) { }
+			SLATE_ARGUMENT(TObjectPtr<UMLDeformerModel>, Model)
+		SLATE_END_ARGS()
+
+		virtual ~STrainingDataProcessorMainWidget() override;
+		void Construct(const FArguments& InArgs);
+
+		//~ Begin FEditorUndoClient overrides.
+		virtual void PostUndo(bool bSuccess) override;
+		virtual void PostRedo(bool bSuccess) override;
+		//~ End FEditorUndoClient overrides.
+
+		//~ Begin FNotifyHook overrides.
+		virtual void NotifyPostChange(const FPropertyChangedEvent& PropertyChangedEvent, FProperty* PropertyThatChanged) override;
+		//~ End FNotifyHook overrides.
+
+	private:
+		FReply OnGenerateButtonClicked() const;
+		bool IsValidConfiguration() const;
+		bool IsObjectOfInterest(UObject* Object) const;
+		void OnAssetModified(UObject* Object) const;
+		void Refresh() const;
+
+	private:
+		/** The details view that shows the properties of our UMLDeformerTrainingDataProcessorSettings. */
+		TSharedPtr<IDetailsView> DetailsView;
+
+		/** A pointer to our model. */
+		TObjectPtr<UMLDeformerModel> Model;
+
+		/** The delegate that handles when an object got modified (any object). */
+		FDelegateHandle ObjectModifiedHandle;
+	};
+} // namespace UE::MLDeformer::TrainingDataProcessor
diff --git a/Engine/Plugins/Animation/MLDeformer/MLDeformerFramework/Source/MLDeformerFrameworkEditor/Public/Tools/TrainingDataProcessor/TrainingDataProcessor.h b/Engine/Plugins/Animation/MLDeformer/MLDeformerFramework/Source/MLDeformerFrameworkEditor/Public/Tools/TrainingDataProcessor/TrainingDataProcessor.h
new file mode 100644
--- /dev/null
+++ b/Engine/Plugins/Animation/MLDeformer/MLDeformerFramework/Source/MLDeformerFrameworkEditor/Public/Tools/TrainingDataProcessor/TrainingDataProcessor.h
@@ -0,0 +1,76 @@
+// Copyright Epic Games, Inc. All Rights Reserved.
+
+#pragma once
+
+#include "CoreMinimal.h"
+#include "MLDeformerTrainingDataProcessorSettings.h"
+
+struct FReferenceSkeleton;
+class USkeleton;
+
+namespace UE::MLDeformer::TrainingDataProcessor
+{
+	/**
+	 * The training data processor algorithm, which is executed when you press the Generate button.
+	 * This will take a set of animation sequences as input, then find the 'best' N number of frames.
+	 * After that it can remix the poses based on a set of bone groups. Basically the keyframes of the
+	 * bones inside the groups get shuffled. So all the bones in the bone group will change together.
+	 */
+	class MLDEFORMERFRAMEWORKEDITOR_API FTrainingDataProcessor final
+	{
+	public:
+		/**
+		 * Run the algorithm using specific settings.
+		 * The output of this process is a UAnimSequence that is specified inside the Settings.
+		 * @param Settings The settings to use during execution. This could come from the UMLDeformerModel::GetTrainingDataProcessorSettings() for example.
+		 * @param Skeleton The skeleton to use when sampling animations and apply the transforms on in the output animation.
+		 * @return Returns true if successful, or false if not. Errors will be logged in case false is returned. 
+		 */
+		bool Execute(const UMLDeformerTrainingDataProcessorSettings& Settings, const USkeleton* Skeleton);
+
+	private:
+		void Clear();
+		bool SampleFrames(const UMLDeformerTrainingDataProcessorSettings& Settings, const USkeleton& Skeleton);
+		TArray<int32> FindBestFrameIndices(int32 MaxNumFrames, const FReferenceSkeleton& RefSkeleton) const;
+		bool RemixPoses(const UMLDeformerTrainingDataProcessorSettings& Settings, const FReferenceSkeleton& RefSkeleton);
+		bool SaveAnimationDataInAnimSequence(const UMLDeformerTrainingDataProcessorSettings& Settings, const FReferenceSkeleton& RefSkeleton, const TArray<int32>& FramesToInclude);
+		int32 GetNumFrames() const;
+		TArrayView<FTransform3f> GetFrameTransforms(int32 FrameIndex);
+		TConstArrayView<FTransform3f> GetFrameTransforms(int32 FrameIndex) const;
+		TConstArrayView<FTransform3f> GetFrameTransforms(const TArray<FTransform3f>& Transforms, int32 FrameIndex) const;
+		TConstArrayView<FTransform3f> GetRefPoseTransforms() const;
+		double CalculateMeanError(TConstArrayView<FTransform3f> PoseA, TConstArrayView<FTransform3f> PoseB) const;
+		double CalculateMeanError(int32 PreviousPoseIndex, int32 CurrentPoseIndex) const;
+		static FTransform3f GetRefPoseTransform(const FReferenceSkeleton& RefSkeleton, FName BoneName);
+		
+	private:
+		/**
+		 * The animation data represented as buffers of vectors and quaternions.
+		 * The size of the arrays are the number of bones in the reference skeleton, multiplied by the number of frames.
+		 * So the layout for a two bone skeleton would be like this:
+		 * [(Bone1, Bone2), (Bone1, Bone2), (Bone1, Bone2), ...] where each (Bone1, Bone2) represents a single frame.
+		 * In case of positions and scales, the value for Bone1 and Bone2 would be a FVector.
+		 * For the rotations it is an FQuat.
+		 */
+		struct FAnimFrameData
+		{
+			// NumRefSkelBones * GetNumFrames().
+			TArray<FTransform3f> Transforms;
+
+			// NumRefSkelBones FTransforms.			
+			TArray<FTransform3f> RefPoseTransforms;
+
+			// Bone names in the bones list, that actually exist.
+			TArray<FName> BoneNames;
+
+			// For each entry in the BoneNames array, this contains an index into the FReferenceSkeleton.
+			TArray<int32> BoneIndices;
+
+			// Number of bones in the reference skeleton.
+			int32 NumRefSkelBones = 0;
+		};
+
+		/** The animation data as flat float arrays. @see FAnimFrameData for more information. */
+		FAnimFrameData AnimFrameData;
+	};
+} // namespace UE::MLDeformer::TrainingDataProcessor
diff --git a/Engine/Plugins/Animation/MLDeformer/MLDeformerFramework/Source/MLDeformerFrameworkEditor/Public/Tools/TrainingDataProcessor/TrainingDataProcessorSettingsDetailCustomization.h b/Engine/Plugins/Animation/MLDeformer/MLDeformerFramework/Source/MLDeformerFrameworkEditor/Public/Tools/TrainingDataProcessor/TrainingDataProcessorSettingsDetailCustomization.h
new file mode 100644
--- /dev/null
+++ b/Engine/Plugins/Animation/MLDeformer/MLDeformerFramework/Source/MLDeformerFrameworkEditor/Public/Tools/TrainingDataProcessor/TrainingDataProcessorSettingsDetailCustomization.h
@@ -0,0 +1,41 @@
+// Copyright Epic Games, Inc. All Rights Reserved.
+
+#pragma once
+
+#include "IDetailCustomization.h"
+#include "Input/Reply.h"
+#include "Layout/Visibility.h"
+
+class UMLDeformerModel;
+class UMLDeformerTrainingDataProcessorSettings;
+class USkeleton;
+struct FAssetData;
+
+namespace UE::MLDeformer::TrainingDataProcessor
+{
+	/**
+	 * The detail customization class for the UMLDeformerTrainingDataProcessorSettings class.
+	 */
+	class MLDEFORMERFRAMEWORKEDITOR_API FTrainingDataProcessorSettingsDetailCustomization final : public IDetailCustomization
+	{
+	public:
+		static TSharedRef<IDetailCustomization> MakeInstance();
+
+		virtual void CustomizeDetails(IDetailLayoutBuilder& DetailBuilder) override;
+
+	private:
+		EVisibility GetNoInputBonesWarningVisibility() const;
+		EVisibility GetNoFramesWarningVisibility() const;
+		EVisibility GetSkeletonMismatchErrorVisibility() const;
+		
+		FReply OnCreateNewButtonClicked() const;
+		int32 GetTotalNumInputFrames() const;
+		bool FilterAnimSequences(const FAssetData& AssetData) const;
+		
+		static void Refresh(IDetailLayoutBuilder* DetailBuilder);
+		static FString FindDefaultAnimSequencePath(const UMLDeformerModel* Model);
+
+	private:
+		TWeakObjectPtr<UMLDeformerTrainingDataProcessorSettings> TrainingDataProcessorSettings;
+	};
+} // namespace UE::MLDeformer::TrainingDataProcessor
diff --git a/Engine/Plugins/Animation/MLDeformer/MLDeformerFramework/Source/MLDeformerFrameworkEditor/Public/Tools/TrainingDataProcessor/TrainingDataProcessorTool.h b/Engine/Plugins/Animation/MLDeformer/MLDeformerFramework/Source/MLDeformerFrameworkEditor/Public/Tools/TrainingDataProcessor/TrainingDataProcessorTool.h
new file mode 100644
--- /dev/null
+++ b/Engine/Plugins/Animation/MLDeformer/MLDeformerFramework/Source/MLDeformerFrameworkEditor/Public/Tools/TrainingDataProcessor/TrainingDataProcessorTool.h
@@ -0,0 +1,25 @@
+// Copyright Epic Games, Inc. All Rights Reserved.
+
+#pragma once
+
+#include "CoreMinimal.h"
+
+class USkeleton;
+class IPropertyHandle;
+
+namespace UE::MLDeformer::TrainingDataProcessor
+{
+	/**
+	 * Registers the Training Data Processor tool to the ML Deformer asset editor tools menu.
+	 */
+	MLDEFORMERFRAMEWORKEDITOR_API void RegisterTool();
+
+	/**
+	 * Find a skeleton for a given property.
+	 * It does this by iterating over the outer objects of the property that's passed as parameter. It then checks the type of the object
+	 * against a set of known types for which we know how to get the skeleton for.
+	 * @param PropertyHandle The property handle to try to get a skeleton for.
+	 * @return A pointer to the skeleton, or nullptr if not found.
+	 */
+	MLDEFORMERFRAMEWORKEDITOR_API USkeleton* FindSkeletonForProperty(const TSharedPtr<IPropertyHandle>& PropertyHandle);
+}
diff --git a/Engine/Plugins/Animation/MLDeformer/NearestNeighborModel/Source/NearestNeighborModel/Private/NearestNeighborModel.cpp b/Engine/Plugins/Animation/MLDeformer/NearestNeighborModel/Source/NearestNeighborModel/Private/NearestNeighborModel.cpp
--- a/Engine/Plugins/Animation/MLDeformer/NearestNeighborModel/Source/NearestNeighborModel/Private/NearestNeighborModel.cpp
+++ b/Engine/Plugins/Animation/MLDeformer/NearestNeighborModel/Source/NearestNeighborModel/Private/NearestNeighborModel.cpp
@@ -1288,7 +1288,7 @@ void UNearestNeighborModel::FinalizeMorphTargets()
 
 UMLDeformerInputInfo* UNearestNeighborModel::CreateInputInfo()
 {
-	return NewObject<UNearestNeighborModelInputInfo>(this);
+	return NewObject<UNearestNeighborModelInputInfo>(this, NAME_None, RF_Transactional);
 }
 
 FString UNearestNeighborModel::GetDisplayName() const
@@ -1315,7 +1315,7 @@ FString UNearestNeighborModel::GetDefaultDeformerGraphAssetPath() const
 
 UMLDeformerModelInstance* UNearestNeighborModel::CreateModelInstance(UMLDeformerComponent* Component)
 {
-	return NewObject<UNearestNeighborModelInstance>(Component);
+	return NewObject<UNearestNeighborModelInstance>(Component, NAME_None, RF_Transactional);
 }
 
 void UNearestNeighborModel::Serialize(FArchive& Archive)
diff --git a/Engine/Plugins/Animation/MLDeformer/NeuralMorphModel/Source/NeuralMorphModel/Private/NeuralMorphModel.cpp b/Engine/Plugins/Animation/MLDeformer/NeuralMorphModel/Source/NeuralMorphModel/Private/NeuralMorphModel.cpp
--- a/Engine/Plugins/Animation/MLDeformer/NeuralMorphModel/Source/NeuralMorphModel/Private/NeuralMorphModel.cpp
+++ b/Engine/Plugins/Animation/MLDeformer/NeuralMorphModel/Source/NeuralMorphModel/Private/NeuralMorphModel.cpp
@@ -43,15 +43,14 @@ UNeuralMorphModel::UNeuralMorphModel(const FObjectInitializer& ObjectInitializer
 
 UMLDeformerModelInstance* UNeuralMorphModel::CreateModelInstance(UMLDeformerComponent* Component)
 {
-	return NewObject<UNeuralMorphModelInstance>(Component);
+	return NewObject<UNeuralMorphModelInstance>(Component, NAME_None, RF_Transactional);
 }
 
 UMLDeformerInputInfo* UNeuralMorphModel::CreateInputInfo()
 {
-	return NewObject<UNeuralMorphInputInfo>(this);
+	return NewObject<UNeuralMorphInputInfo>(this, NAME_None, RF_Transactional);
 }
 
-
 namespace
 {
 	int32 FindVirtualParentIndex(const FReferenceSkeleton& RefSkel, int32 BoneIndex, const TArray<FName>& IncludedBoneNames)
diff --git a/Engine/Plugins/Animation/MLDeformer/NeuralMorphModel/Source/NeuralMorphModelEditor/Private/NeuralMorphEditorModel.cpp b/Engine/Plugins/Animation/MLDeformer/NeuralMorphModel/Source/NeuralMorphModelEditor/Private/NeuralMorphEditorModel.cpp
--- a/Engine/Plugins/Animation/MLDeformer/NeuralMorphModel/Source/NeuralMorphModelEditor/Private/NeuralMorphEditorModel.cpp
+++ b/Engine/Plugins/Animation/MLDeformer/NeuralMorphModel/Source/NeuralMorphModelEditor/Private/NeuralMorphEditorModel.cpp
@@ -23,6 +23,7 @@
 #include "Rendering/SkeletalMeshLODModel.h"
 #include "ReferenceSkeleton.h"
 #include "Widgets/SBoxPanel.h"
+#include "Misc/TransactionObjectEvent.h"
 #include "SNeuralMorphBoneGroupsWidget.h"
 #include "SNeuralMorphCurveGroupsWidget.h"
 #include "SNeuralMorphInputWidget.h"
diff --git a/Engine/Plugins/Animation/MLDeformer/NeuralMorphModel/Source/NeuralMorphModelEditor/Private/SNeuralMorphBoneGroupsWidget.cpp b/Engine/Plugins/Animation/MLDeformer/NeuralMorphModel/Source/NeuralMorphModelEditor/Private/SNeuralMorphBoneGroupsWidget.cpp
--- a/Engine/Plugins/Animation/MLDeformer/NeuralMorphModel/Source/NeuralMorphModelEditor/Private/SNeuralMorphBoneGroupsWidget.cpp
+++ b/Engine/Plugins/Animation/MLDeformer/NeuralMorphModel/Source/NeuralMorphModelEditor/Private/SNeuralMorphBoneGroupsWidget.cpp
@@ -9,6 +9,7 @@
 #include "MLDeformerEditorStyle.h"
 #include "Framework/MultiBox/MultiBoxBuilder.h"
 #include "Engine/SkeletalMesh.h"
+#include "ScopedTransaction.h"
 
 #define LOCTEXT_NAMESPACE "NeuralMorphBoneGroupsWidget"
 
@@ -132,6 +133,9 @@ namespace UE::NeuralMorphModel
 			return;
 		}
 
+		FScopedTransaction Transaction(TEXT("SNeuralMorphBoneGroupsWidget"), LOCTEXT("CreateBoneGroupText", "Create Bone Group"), NeuralMorphModel);
+		NeuralMorphModel->Modify();
+
 		NeuralMorphModel->BoneGroups.AddDefaulted();
 		FNeuralMorphBoneGroup& BoneGroup = NeuralMorphModel->BoneGroups.Last();
 		BoneGroup.GroupName = EditorModel->GenerateUniqueBoneGroupName();
@@ -169,6 +173,9 @@ namespace UE::NeuralMorphModel
 
 		UNeuralMorphModel* NeuralMorphModel = EditorModel->GetNeuralMorphModel();
 
+		FScopedTransaction Transaction(TEXT("SNeuralMorphBoneGroupsWidget"), LOCTEXT("RemoveBoneGroupItemsText", "Remove Bone Group Items"), NeuralMorphModel);
+		NeuralMorphModel->Modify();
+
 		// Remove all selected bones.
 		if (!SelectedBones.IsEmpty())
 		{
@@ -220,6 +227,9 @@ namespace UE::NeuralMorphModel
 	void SNeuralMorphBoneGroupsWidget::OnClearBoneGroups()
 	{
 		UNeuralMorphModel* NeuralMorphModel = EditorModel->GetNeuralMorphModel();
+		FScopedTransaction Transaction(TEXT("SNeuralMorphBoneGroupsWidget"), LOCTEXT("ClearBoneGroupItemsText", "Clear Bone Groups"), NeuralMorphModel);
+		NeuralMorphModel->Modify();
+
 		NeuralMorphModel->BoneGroups.Empty();
 		NeuralMorphModel->BoneGroupMaskInfoMap.Empty();
 		RefreshTree(true);
@@ -269,6 +279,9 @@ namespace UE::NeuralMorphModel
 		const TArray<FName>& BoneNames = Dialog->GetPickedBoneNames();
 		if (!BoneNames.IsEmpty())
 		{
+			FScopedTransaction Transaction(TEXT("SNeuralMorphBoneGroupsWidget"), LOCTEXT("AddBoneGroupItemsText", "Add Bones To Group"), NeuralMorphModel);
+			NeuralMorphModel->Modify();
+
 			for (const FName BoneName : BoneNames)
 			{
 				BoneGroup.BoneNames.AddUnique(BoneName);
diff --git a/Engine/Plugins/Animation/MLDeformer/NeuralMorphModel/Source/NeuralMorphModelEditor/Private/SNeuralMorphCurveGroupsWidget.cpp b/Engine/Plugins/Animation/MLDeformer/NeuralMorphModel/Source/NeuralMorphModelEditor/Private/SNeuralMorphCurveGroupsWidget.cpp
--- a/Engine/Plugins/Animation/MLDeformer/NeuralMorphModel/Source/NeuralMorphModelEditor/Private/SNeuralMorphCurveGroupsWidget.cpp
+++ b/Engine/Plugins/Animation/MLDeformer/NeuralMorphModel/Source/NeuralMorphModelEditor/Private/SNeuralMorphCurveGroupsWidget.cpp
@@ -9,6 +9,8 @@
 #include "MLDeformerEditorStyle.h"
 #include "Framework/MultiBox/MultiBoxBuilder.h"
 #include "Engine/SkeletalMesh.h"
+#include "ScopedTransaction.h"
+
 
 #define LOCTEXT_NAMESPACE "NeuralMorphCurveGroupsWidget"
 
@@ -121,6 +123,9 @@ namespace UE::NeuralMorphModel
 			return;
 		}
 
+		FScopedTransaction Transaction(TEXT("SNeuralMorphCurveGroupsWidget"), LOCTEXT("CreateCurveGroupText", "Create Curve Group"), NeuralMorphModel);
+		NeuralMorphModel->Modify();
+
 		NeuralMorphModel->CurveGroups.AddDefaulted();
 		FNeuralMorphCurveGroup& CurveGroup = NeuralMorphModel->CurveGroups.Last();
 		CurveGroup.GroupName = EditorModel->GenerateUniqueCurveGroupName();
@@ -157,6 +162,9 @@ namespace UE::NeuralMorphModel
 
 		UNeuralMorphModel* NeuralMorphModel = EditorModel->GetNeuralMorphModel();
 
+		FScopedTransaction Transaction(TEXT("SNeuralMorphCurveGroupsWidget"), LOCTEXT("RemoveCurveGroupItemsText", "Remove Curve Group Items"), NeuralMorphModel);
+		NeuralMorphModel->Modify();
+
 		// Remove all selected bones.
 		if (!SelectedCurves.IsEmpty())
 		{
@@ -202,6 +210,10 @@ namespace UE::NeuralMorphModel
 	void SNeuralMorphCurveGroupsWidget::OnClearCurveGroups()
 	{
 		UNeuralMorphModel* NeuralMorphModel = EditorModel->GetNeuralMorphModel();
+
+		FScopedTransaction Transaction(TEXT("SNeuralMorphCurveGroupsWidget"), LOCTEXT("ClearCurveGroupText", "Clear Curve Groups"), NeuralMorphModel);
+		NeuralMorphModel->Modify();
+
 		NeuralMorphModel->CurveGroups.Empty();
 		RefreshTree(true);
 		EditorModel->RebuildEditorMaskInfo();
@@ -248,6 +260,9 @@ namespace UE::NeuralMorphModel
 		const TArray<FName>& CurveNames = Dialog->GetPickedCurveNames();
 		if (!CurveNames.IsEmpty())
 		{
+			FScopedTransaction Transaction(TEXT("SNeuralMorphCurveGroupsWidget"), LOCTEXT("AddCurveToGroupText", "Add Curves To Group"), NeuralMorphModel);
+			NeuralMorphModel->Modify();
+
 			for (const FName CurveName : CurveNames)
 			{
 				CurveGroup.CurveNames.Add(CurveName);
diff --git a/Engine/Plugins/Animation/MLDeformer/NeuralMorphModel/Source/NeuralMorphModelEditor/Private/SNeuralMorphInputWidget.cpp b/Engine/Plugins/Animation/MLDeformer/NeuralMorphModel/Source/NeuralMorphModelEditor/Private/SNeuralMorphInputWidget.cpp
--- a/Engine/Plugins/Animation/MLDeformer/NeuralMorphModel/Source/NeuralMorphModelEditor/Private/SNeuralMorphInputWidget.cpp
+++ b/Engine/Plugins/Animation/MLDeformer/NeuralMorphModel/Source/NeuralMorphModelEditor/Private/SNeuralMorphInputWidget.cpp
@@ -27,6 +27,7 @@
 #include "Engine/SkeletalMesh.h"
 #include "MeshDescription.h"
 #include "IDetailsView.h"
+#include "ScopedTransaction.h"
 
 #define LOCTEXT_NAMESPACE "NeuralMorphInputWidget"
 
@@ -123,7 +124,13 @@ namespace UE::NeuralMorphModel
 			return;
 		}
 
+
 		FNeuralMorphEditorModel* NeuralEditorModel = static_cast<FNeuralMorphEditorModel*>(EditorModel);
+
+		UNeuralMorphModel* NeuralMorphModel = NeuralEditorModel->GetNeuralMorphModel();		
+		FScopedTransaction Transaction(TEXT("SNeuralMorphInputWidget"), LOCTEXT("ResetAllBoneMaskText", "Reset All Bone Masks"), NeuralMorphModel);
+		NeuralMorphModel->Modify();
+
 		HierarchyDepth = 1;
 		NeuralEditorModel->GenerateBoneMaskInfos(HierarchyDepth);
 		NeuralEditorModel->RebuildEditorMaskInfo();
@@ -139,12 +146,16 @@ namespace UE::NeuralMorphModel
 			return;
 		}
 
+		FNeuralMorphEditorModel* NeuralEditorModel = static_cast<FNeuralMorphEditorModel*>(EditorModel);
+		UNeuralMorphModel* NeuralMorphModel = NeuralEditorModel->GetNeuralMorphModel();		
+		FScopedTransaction Transaction(TEXT("SNeuralMorphInputWidget"), LOCTEXT("ResetSelectedBoneMaskText", "Reset Selected Bone Masks"), NeuralMorphModel);
+		NeuralMorphModel->Modify();
+
 		// Get the selected input bones.
 		check(InputBonesWidget.IsValid());
 		check(InputBonesWidget->GetTreeWidget().IsValid());
 		TArray<TSharedPtr<FMLDeformerInputBoneTreeElement>> SelectedItems = InputBonesWidget->GetTreeWidget()->GetSelectedItems();
 
-		FNeuralMorphEditorModel* NeuralEditorModel = static_cast<FNeuralMorphEditorModel*>(EditorModel);
 		UNeuralMorphInputInfo* EditorInputInfo = Cast<UNeuralMorphInputInfo>(NeuralEditorModel->GetEditorInputInfo());
 		check(EditorInputInfo);
 
@@ -212,7 +223,11 @@ namespace UE::NeuralMorphModel
 		// If we picked some bones.
 		const TArray<FName>& PickedBoneNames = Dialog->GetPickedBoneNames();
 		if (!PickedBoneNames.IsEmpty())
-		{			
+		{	
+			UNeuralMorphModel* NeuralMorphModel = NeuralEditorModel->GetNeuralMorphModel();		
+			FScopedTransaction Transaction(TEXT("SNeuralMorphInputWidget"), LOCTEXT("EditBoneMaskText", "Edit Bone Mask"), NeuralMorphModel);
+			NeuralMorphModel->Modify();
+
 			// For all bones we have selected, add the picked mask bones to the mask info.
 			for (const TSharedPtr<FMLDeformerInputBoneTreeElement>& Item : SelectedItems)
 			{
@@ -222,10 +237,10 @@ namespace UE::NeuralMorphModel
 				// Add the picked bone name to the mask info of the selected bone.
 				const int32 BoneIndex = EditorInputInfo->GetBoneNames().Find(SelectedInputBoneName);
 				check(BoneIndex != INDEX_NONE);
-				FMLDeformerMaskInfo* MaskInfo = NeuralEditorModel->GetNeuralMorphModel()->BoneMaskInfoMap.Find(SelectedInputBoneName);
+				FMLDeformerMaskInfo* MaskInfo = NeuralMorphModel->BoneMaskInfoMap.Find(SelectedInputBoneName);
 				if (MaskInfo == nullptr)
 				{
-					MaskInfo = &NeuralEditorModel->GetNeuralMorphModel()->BoneMaskInfoMap.Add(SelectedInputBoneName, FMLDeformerMaskInfo());
+					MaskInfo = &NeuralMorphModel->BoneMaskInfoMap.Add(SelectedInputBoneName, FMLDeformerMaskInfo());
 				}
 
 				check(MaskInfo);
@@ -252,6 +267,10 @@ namespace UE::NeuralMorphModel
 		}
 
 		FNeuralMorphEditorModel* NeuralEditorModel = static_cast<FNeuralMorphEditorModel*>(EditorModel);
+		UNeuralMorphModel* NeuralMorphModel = NeuralEditorModel->GetNeuralMorphModel();
+		FScopedTransaction Transaction(TEXT("SNeuralMorphInputWidget"), LOCTEXT("ConfigureBoneMaskText", "Configure Bone Mask"), NeuralMorphModel);
+		NeuralMorphModel->Modify();
+		
 		HierarchyDepth = 1;
 		NeuralEditorModel->GenerateBoneGroupMaskInfos(HierarchyDepth);
 		NeuralEditorModel->RebuildEditorMaskInfo();
@@ -275,6 +294,10 @@ namespace UE::NeuralMorphModel
 		UNeuralMorphInputInfo* EditorInputInfo = Cast<UNeuralMorphInputInfo>(NeuralEditorModel->GetEditorInputInfo());
 		check(EditorInputInfo);
 
+		UNeuralMorphModel* NeuralMorphModel = NeuralEditorModel->GetNeuralMorphModel();
+		FScopedTransaction Transaction(TEXT("SNeuralMorphInputWidget"), LOCTEXT("ResetSelectedBoneGroupMask", "Reset Selected Bone Group Mask"), NeuralMorphModel);
+		NeuralMorphModel->Modify();
+
 		// For all bone groups we have selected, reset the mask info.
 		HierarchyDepth = 1;
 		for (const TSharedPtr<FNeuralMorphBoneGroupsTreeElement>& Item : SelectedItems)
@@ -322,6 +345,7 @@ namespace UE::NeuralMorphModel
 
 		// If we only selected one item, we can highlight the bones already in the mask info.
 		FNeuralMorphEditorModel* NeuralEditorModel = static_cast<FNeuralMorphEditorModel*>(EditorModel);
+		UNeuralMorphModel* NeuralMorphModel = NeuralEditorModel->GetNeuralMorphModel();
 		UNeuralMorphInputInfo* EditorInputInfo = Cast<UNeuralMorphInputInfo>(NeuralEditorModel->GetEditorInputInfo());
 		check(EditorInputInfo);
 		TArray<FName> HighlightedBones;
@@ -330,7 +354,7 @@ namespace UE::NeuralMorphModel
 			check(LastSelectedGroup->IsGroup());
 			const int32 GroupIndex = LastSelectedGroup->GroupIndex;
 			check(GroupIndex!= INDEX_NONE);
-			const FMLDeformerMaskInfo* MaskInfo = NeuralEditorModel->GetNeuralMorphModel()->BoneGroupMaskInfoMap.Find(LastSelectedGroup->Name);
+			const FMLDeformerMaskInfo* MaskInfo = NeuralMorphModel->BoneGroupMaskInfoMap.Find(LastSelectedGroup->Name);
 			if (MaskInfo)
 			{
 				HighlightedBones = MaskInfo->BoneNames;
@@ -351,7 +375,10 @@ namespace UE::NeuralMorphModel
 		// If we picked some bones.
 		const TArray<FName>& PickedBoneNames = Dialog->GetPickedBoneNames();
 		if (!PickedBoneNames.IsEmpty())
-		{			
+		{	
+			FScopedTransaction Transaction(TEXT("SNeuralMorphInputWidget"), LOCTEXT("EditBoneGroupMaskText", "Edit Bone Group Mask"), NeuralMorphModel);
+			NeuralMorphModel->Modify();
+
 			// For all groups we selected.
 			for (const TSharedPtr<FNeuralMorphBoneGroupsTreeElement>& Item : SelectedItems)
 			{
@@ -363,10 +390,10 @@ namespace UE::NeuralMorphModel
 				}
 
 				// Add the picked bone name to the mask info of the selected group.
-				FMLDeformerMaskInfo* MaskInfo = NeuralEditorModel->GetNeuralMorphModel()->BoneGroupMaskInfoMap.Find(Item->Name);
+				FMLDeformerMaskInfo* MaskInfo = NeuralMorphModel->BoneGroupMaskInfoMap.Find(Item->Name);
 				if (MaskInfo == nullptr)
 				{
-					MaskInfo = &NeuralEditorModel->GetNeuralMorphModel()->BoneGroupMaskInfoMap.Add(Item->Name, FMLDeformerMaskInfo());
+					MaskInfo = &NeuralMorphModel->BoneGroupMaskInfoMap.Add(Item->Name, FMLDeformerMaskInfo());
 				}
 				check(MaskInfo);
 				MaskInfo->BoneNames.Reset();
@@ -818,6 +845,7 @@ namespace UE::NeuralMorphModel
 	void SNeuralMorphInputWidget::OnAddInputCurves(const TArray<FName>& Names)
 	{
 		FNeuralMorphEditorModel* NeuralEditorModel = static_cast<FNeuralMorphEditorModel*>(EditorModel);
+		UNeuralMorphModel* NeuralMorphModel = NeuralEditorModel->GetNeuralMorphModel();
 		NeuralEditorModel->RebuildEditorMaskInfo();
 	}
 
@@ -851,6 +879,9 @@ namespace UE::NeuralMorphModel
 			NeuralMaskInfo = *NeuralModel->BoneMaskInfoMap.Find(BoneName);
 		}
 
+		FScopedTransaction Transaction(TEXT("SNeuralMorphInputWidget"), LOCTEXT("ConfigureBoneMaskText", "Configure Bone Mask"), NeuralModel);
+		NeuralModel->Modify();
+
 		// Show the mask configuration dialog.
 		TSharedPtr<SMLDeformerMaskConfigDialog> Dialog = SNew(SMLDeformerMaskConfigDialog, EditorModel)
 			.InitialMaskInfo(NeuralMaskInfo)
@@ -914,6 +945,7 @@ namespace UE::NeuralMorphModel
 
 		// If we only selected one item, we can highlight the bones already in the mask info.
 		FNeuralMorphEditorModel* NeuralEditorModel = static_cast<FNeuralMorphEditorModel*>(EditorModel);
+		UNeuralMorphModel* NeuralMorphModel = NeuralEditorModel->GetNeuralMorphModel();
 		UNeuralMorphInputInfo* EditorInputInfo = Cast<UNeuralMorphInputInfo>(NeuralEditorModel->GetEditorInputInfo());
 		check(EditorInputInfo);
 
@@ -923,7 +955,7 @@ namespace UE::NeuralMorphModel
 			check(LastSelectedGroup->IsGroup());
 			const int32 GroupIndex = LastSelectedGroup->GroupIndex;
 			check(GroupIndex!= INDEX_NONE);
-			GroupMaskInfo = NeuralEditorModel->GetNeuralMorphModel()->BoneGroupMaskInfoMap.Find(LastSelectedGroup->Name);
+			GroupMaskInfo = NeuralMorphModel->BoneGroupMaskInfoMap.Find(LastSelectedGroup->Name);
 			if (!GroupMaskInfo)
 			{
 				return;
@@ -937,6 +969,9 @@ namespace UE::NeuralMorphModel
 		check(GroupMaskInfo);
 		const FName GroupName = LastSelectedGroup->Name;
 
+		FScopedTransaction Transaction(TEXT("SNeuralMorphInputWidget"), LOCTEXT("ConfigureBoneGroupMaskText", "Configure Bone Group Mask"), NeuralMorphModel);
+		NeuralMorphModel->Modify();
+
 		// Show the mask configuration dialog.
 		TSharedPtr<SMLDeformerMaskConfigDialog> Dialog = SNew(SMLDeformerMaskConfigDialog, EditorModel)
 			.InitialMaskInfo(*GroupMaskInfo)
diff --git a/Engine/Plugins/Animation/MLDeformer/NeuralMorphModel/Source/NeuralMorphModelEditor/Public/NeuralMorphEditorModel.h b/Engine/Plugins/Animation/MLDeformer/NeuralMorphModel/Source/NeuralMorphModelEditor/Public/NeuralMorphEditorModel.h
--- a/Engine/Plugins/Animation/MLDeformer/NeuralMorphModel/Source/NeuralMorphModelEditor/Public/NeuralMorphEditorModel.h
+++ b/Engine/Plugins/Animation/MLDeformer/NeuralMorphModel/Source/NeuralMorphModelEditor/Public/NeuralMorphEditorModel.h
@@ -29,8 +29,6 @@ namespace UE::NeuralMorphModel
 
 		// FGCObject overrides.
 		virtual FString GetReferencerName() const override		{ return TEXT("FNeuralMorphEditorModel"); }
-		virtual bool LoadTrainedNetwork() const override;
-		virtual void UpdateIsReadyForTrainingState() override;
 		// ~END FGCObject overrides.
 	
 		// FMLDeformerEditorModel overrides.
@@ -45,6 +43,8 @@ namespace UE::NeuralMorphModel
 		virtual bool GetSupportsPerTrainingInputAnimVertexMask() const override	{ return true; }
 		virtual void UpdateTrainingDeviceList() override;
 		virtual UMeshDeformer* GetActiveHeatMapDeformer() const override;
+		virtual void UpdateIsReadyForTrainingState() override;
+		virtual bool LoadTrainedNetwork() const override;
 		// ~END FMLDeformerEditorModel overrides.
 
 		// FMLDeformerMorphModelEditorModel overrides.
diff --git a/Engine/Plugins/Animation/MLDeformer/VertexDeltaModel/Source/VertexDeltaModel/Private/VertexDeltaModel.cpp b/Engine/Plugins/Animation/MLDeformer/VertexDeltaModel/Source/VertexDeltaModel/Private/VertexDeltaModel.cpp
--- a/Engine/Plugins/Animation/MLDeformer/VertexDeltaModel/Source/VertexDeltaModel/Private/VertexDeltaModel.cpp
+++ b/Engine/Plugins/Animation/MLDeformer/VertexDeltaModel/Source/VertexDeltaModel/Private/VertexDeltaModel.cpp
@@ -51,7 +51,7 @@ void UVertexDeltaModel::SetNNEModelData(TObjectPtr<UNNEModelData> ModelData)
 
 UMLDeformerModelInstance* UVertexDeltaModel::CreateModelInstance(UMLDeformerComponent* Component)
 {
-	return NewObject<UVertexDeltaModelInstance>(Component);
+	return NewObject<UVertexDeltaModelInstance>(Component, NAME_None, RF_Transactional);
 }
 
 FString UVertexDeltaModel::GetDefaultDeformerGraphAssetPath() const 
diff --git a/Engine/Plugins/Animation/MLDeformer/VertexDeltaModel/Source/VertexDeltaModel/Private/VertexDeltaModelInstance.cpp b/Engine/Plugins/Animation/MLDeformer/VertexDeltaModel/Source/VertexDeltaModel/Private/VertexDeltaModelInstance.cpp
--- a/Engine/Plugins/Animation/MLDeformer/VertexDeltaModel/Source/VertexDeltaModel/Private/VertexDeltaModelInstance.cpp
+++ b/Engine/Plugins/Animation/MLDeformer/VertexDeltaModel/Source/VertexDeltaModel/Private/VertexDeltaModelInstance.cpp
@@ -161,8 +161,8 @@ void UVertexDeltaModelInstance::CreateNNEModel()
 		UVertexDeltaModel* VertexDeltaModel = Cast<UVertexDeltaModel>(Model);
 		if (VertexDeltaModel)
 		{
-			TWeakInterfacePtr<INNERuntime> Runtime = UE::NNE::GetRuntime<INNERuntime>(VertexDeltaModel->GetNNERuntimeName());
-			TWeakInterfacePtr<INNERuntimeRDG> RuntimeRDG = UE::NNE::GetRuntime<INNERuntimeRDG>(VertexDeltaModel->GetNNERuntimeName());
+			const TWeakInterfacePtr<INNERuntime> Runtime = UE::NNE::GetRuntime<INNERuntime>(VertexDeltaModel->GetNNERuntimeName());
+			const TWeakInterfacePtr<INNERuntimeRDG> RuntimeRDG = UE::NNE::GetRuntime<INNERuntimeRDG>(VertexDeltaModel->GetNNERuntimeName());
 			if (!Runtime.IsValid())
 			{
 				UE_LOG(LogVertexDeltaModel, Error, TEXT("Can't get NNE runtime: %s"), *VertexDeltaModel->GetNNERuntimeName());
@@ -170,18 +170,18 @@ void UVertexDeltaModelInstance::CreateNNEModel()
 			}
 
 			// If we can create the model from its data.
-			TObjectPtr<UNNEModelData> ModelData = VertexDeltaModel->NNEModel;
+			const TObjectPtr<UNNEModelData> ModelData = VertexDeltaModel->NNEModel;
 			if (ModelData && RuntimeRDG.IsValid() && RuntimeRDG->CanCreateModelRDG(ModelData) == INNERuntimeRDG::ECanCreateModelRDGStatus::Ok)
 			{
 				// Create the model.
-				TSharedPtr<UE::NNE::IModelRDG> ModelRDG = RuntimeRDG->CreateModelRDG(ModelData);
+				const TSharedPtr<UE::NNE::IModelRDG> ModelRDG = RuntimeRDG->CreateModelRDG(ModelData);
 				if (ModelRDG.IsValid())
 				{
 					ModelInstanceRDG = RuntimeRDG->CreateModelRDG(ModelData)->CreateModelInstanceRDG();
 					if (ModelInstanceRDG)
 					{
 						// Setup inputs.
-						TConstArrayView<UE::NNE::FTensorDesc> InputTensorDescs = ModelInstanceRDG->GetInputTensorDescs();
+						const TConstArrayView<UE::NNE::FTensorDesc> InputTensorDescs = ModelInstanceRDG->GetInputTensorDescs();
 						UE::NNE::FTensorShape InputTensorShape = UE::NNE::FTensorShape::MakeFromSymbolic(InputTensorDescs[0].GetShape());
 						ModelInstanceRDG->SetInputTensorShapes({ InputTensorShape });
 						check(InputTensorDescs[0].GetElementByteSize() == sizeof(float));
@@ -221,7 +221,7 @@ bool UVertexDeltaModelInstance::SetupInputs()
 		return false;
 	}
 
-	TConstArrayView<UE::NNE::FTensorDesc> InputTensorDescs = ModelInstanceRDG->GetInputTensorDescs();
+	const TConstArrayView<UE::NNE::FTensorDesc> InputTensorDescs = ModelInstanceRDG->GetInputTensorDescs();
 	const UE::NNE::FTensorShape InputTensorShape = UE::NNE::FTensorShape::MakeFromSymbolic(InputTensorDescs[0].GetShape());
 	const int64 NumNeuralNetInputs = InputTensorShape.Volume();
 	const int64 NumDeformerAssetInputs = Model->GetInputInfo()->CalcNumNeuralNetInputs(Model->GetNumFloatsPerBone(), Model->GetNumFloatsPerCurve());
diff --git a/Engine/Plugins/Experimental/Animation/MLDeformer/DetailPoseModel/Source/DetailPoseModel/Private/DetailPoseModel.cpp b/Engine/Plugins/Experimental/Animation/MLDeformer/DetailPoseModel/Source/DetailPoseModel/Private/DetailPoseModel.cpp
--- a/Engine/Plugins/Experimental/Animation/MLDeformer/DetailPoseModel/Source/DetailPoseModel/Private/DetailPoseModel.cpp
+++ b/Engine/Plugins/Experimental/Animation/MLDeformer/DetailPoseModel/Source/DetailPoseModel/Private/DetailPoseModel.cpp
@@ -43,12 +43,12 @@ UDetailPoseModel::UDetailPoseModel(const FObjectInitializer& ObjectInitializer)
 
 UMLDeformerModelInstance* UDetailPoseModel::CreateModelInstance(UMLDeformerComponent* Component)
 {
-	return NewObject<UDetailPoseModelInstance>(Component);
+	return NewObject<UDetailPoseModelInstance>(Component, NAME_None, RF_Transactional);
 }
 
 UMLDeformerInputInfo* UDetailPoseModel::CreateInputInfo()
 {
-	return NewObject<UDetailPoseModelInputInfo>(this);
+	return NewObject<UDetailPoseModelInputInfo>(this, NAME_None, RF_Transactional);
 }
 
 #if WITH_EDITORONLY_DATA
