[MLDeformer] Add support for curves on the skeletal meshes.

diff --git a/Engine/Plugins/Animation/MLDeformer/MLDeformerFramework/Source/MLDeformerFrameworkEditor/Private/MLDeformerCurveReferenceCustomization.cpp b/Engine/Plugins/Animation/MLDeformer/MLDeformerFramework/Source/MLDeformerFrameworkEditor/Private/MLDeformerCurveReferenceCustomization.cpp
--- a/Engine/Plugins/Animation/MLDeformer/MLDeformerFramework/Source/MLDeformerFrameworkEditor/Private/MLDeformerCurveReferenceCustomization.cpp
+++ b/Engine/Plugins/Animation/MLDeformer/MLDeformerFramework/Source/MLDeformerFrameworkEditor/Private/MLDeformerCurveReferenceCustomization.cpp
@@ -30,15 +30,13 @@ namespace UE::MLDeformer
 				SNew(SCurveSelectionWidget)
 				.OnCurveSelectionChanged(this, &FMLDeformerCurveReferenceCustomization::OnCurveSelectionChanged)
 				.OnGetSelectedCurve(this, &FMLDeformerCurveReferenceCustomization::OnGetSelectedCurve)
-				.OnGetSkeleton(this, &FMLDeformerCurveReferenceCustomization::OnGetSkeleton)
+				.OnGetSkeletalMesh(this, &FMLDeformerCurveReferenceCustomization::OnGetSkeletalMesh)
 			];
 		}
 	}
 
 	void FMLDeformerCurveReferenceCustomization::SetSkeleton(TSharedRef<IPropertyHandle> StructPropertyHandle) 
 	{
-		Skeleton = nullptr;
-
 		TArray<UObject*> Objects;
 		StructPropertyHandle->GetOuterObjects(Objects);
 		check(Objects.Num() == 1);
@@ -47,7 +45,7 @@ namespace UE::MLDeformer
 		check(DeformerModel);
 
 		bool bInvalidSkeletonIsError = false;
-		Skeleton = DeformerModel->GetSkeleton(bInvalidSkeletonIsError, nullptr);
+		SkeletalMesh = DeformerModel->GetSkeletalMesh();
 	}
 
 	TSharedPtr<IPropertyHandle> FMLDeformerCurveReferenceCustomization::FindStructMemberProperty(TSharedRef<IPropertyHandle> PropertyHandle, const FName& PropertyName)
@@ -84,8 +82,8 @@ namespace UE::MLDeformer
 		return CurveName;
 	}
 
-	USkeleton* FMLDeformerCurveReferenceCustomization::OnGetSkeleton() const
+	USkeletalMesh* FMLDeformerCurveReferenceCustomization::OnGetSkeletalMesh() const
 	{
-		return Skeleton;
+		return SkeletalMesh;
 	}
 }	// namespace UE::MLDeformer
diff --git a/Engine/Plugins/Animation/MLDeformer/MLDeformerFramework/Source/MLDeformerFrameworkEditor/Private/MLDeformerEditorModel.cpp b/Engine/Plugins/Animation/MLDeformer/MLDeformerFramework/Source/MLDeformerFrameworkEditor/Private/MLDeformerEditorModel.cpp
--- a/Engine/Plugins/Animation/MLDeformer/MLDeformerFramework/Source/MLDeformerFrameworkEditor/Private/MLDeformerEditorModel.cpp
+++ b/Engine/Plugins/Animation/MLDeformer/MLDeformerFramework/Source/MLDeformerFrameworkEditor/Private/MLDeformerEditorModel.cpp
@@ -1762,13 +1762,29 @@ namespace UE::MLDeformer
 		}
 
 		check(Model->GetSkeletalMesh());
-		USkeleton* Skeleton = Model->GetSkeletalMesh()->GetSkeleton();
+
+		TArray<FName> CurveList;
+		if (Model->GetSkeletalMesh()->GetSkeleton())
+		{
+			Model->GetSkeletalMesh()->GetSkeleton()->GetCurveMetaDataNames(CurveList);
+		}
+
+		const UAnimCurveMetaData* SkelMeshCurveData = Model->GetSkeletalMesh()->GetAssetUserData<UAnimCurveMetaData>();
+		if (SkelMeshCurveData)
+		{
+			TArray<FName> SkelMeshCurveNames;
+			SkelMeshCurveData->GetCurveMetaDataNames(SkelMeshCurveNames);
+			for (const FName CurveName : SkelMeshCurveNames)
+			{
+				CurveList.AddUnique(CurveName);
+			}
+		}
+		
+		const USkeleton* Skeleton = Model->GetSkeletalMesh()->GetSkeleton();
 		check(Skeleton);
 
 		TArray<FName> AnimatedCurveList;
-		TArray<FName> SkeletonCurveNames;
-		Skeleton->GetCurveMetaDataNames(SkeletonCurveNames);
-		for (const FName SkeletonCurveName : SkeletonCurveNames)
+		for (const FName SkeletonCurveName : CurveList)
 		{
 			const TArray<FFloatCurve>& AnimCurves = DataModel->GetFloatCurves();
 			for (const FFloatCurve& AnimCurve : AnimCurves)
@@ -1795,10 +1811,10 @@ namespace UE::MLDeformer
 			}
 		}
 
-		TArray<FMLDeformerCurveReference>& CurveList = Model->GetCurveIncludeList();
+		TArray<FMLDeformerCurveReference>& CurveIncludeList = Model->GetCurveIncludeList();
 		for (const FName CurveName : AnimatedCurveList)
 		{
-			CurveList.AddUnique(FMLDeformerCurveReference(CurveName));
+			CurveIncludeList.AddUnique(FMLDeformerCurveReference(CurveName));
 		}
 	}
 
diff --git a/Engine/Plugins/Animation/MLDeformer/MLDeformerFramework/Source/MLDeformerFrameworkEditor/Private/SMLDeformerCurvePickerDialog.cpp b/Engine/Plugins/Animation/MLDeformer/MLDeformerFramework/Source/MLDeformerFrameworkEditor/Private/SMLDeformerCurvePickerDialog.cpp
--- a/Engine/Plugins/Animation/MLDeformer/MLDeformerFramework/Source/MLDeformerFrameworkEditor/Private/SMLDeformerCurvePickerDialog.cpp
+++ b/Engine/Plugins/Animation/MLDeformer/MLDeformerFramework/Source/MLDeformerFrameworkEditor/Private/SMLDeformerCurvePickerDialog.cpp
@@ -5,6 +5,8 @@
 #include "Widgets/Views/SListView.h"
 #include "Dialog/SCustomDialog.h"
 #include "Animation/Skeleton.h"
+#include "Animation/AnimCurveMetadata.h"
+#include "Engine/SkeletalMesh.h"
 
 #define LOCTEXT_NAMESPACE "MLDeformerCurvePickerWidget"
 
@@ -98,7 +100,7 @@ namespace UE::MLDeformer
 			.AllowMultiSelect(bAllowMultiSelect)
 			.PickerDialog(SharedThis(this));
 
-		if (Skeleton == nullptr)
+		if (SkeletalMesh == nullptr)
 		{
 			return ListView;
 		}
@@ -140,27 +142,43 @@ namespace UE::MLDeformer
 
 		return FReply::Unhandled();
 	}
-
+	
 	void SMLDeformerCurvePickerDialog::RefreshListElements() const 
 	{
 		CurveListView->Clear();
 
-		TArray<FName> SkeletonCurveNames;
-		Skeleton->GetCurveMetaDataNames(SkeletonCurveNames);
+		TArray<FName> CurveNames;
+		const UAnimCurveMetaData* SkelMeshCurveData = SkeletalMesh->GetAssetUserData<UAnimCurveMetaData>();
+		if (SkelMeshCurveData)
+		{
+			TArray<FName> SkelMeshCurveNames;
+			SkelMeshCurveData->GetCurveMetaDataNames(SkelMeshCurveNames);
+			CurveNames.Append(SkelMeshCurveNames);
+		}
+		
+		if (SkeletalMesh->GetSkeleton())
+		{
+			TArray<FName> SkeletonCurveNames;
+			SkeletalMesh->GetSkeleton()->GetCurveMetaDataNames(SkeletonCurveNames);
+			for (const FName& CurveName : SkeletonCurveNames)
+			{
+				CurveNames.AddUnique(CurveName);
+			}
+		}
 
-		for (const FName SkelCurveName : SkeletonCurveNames)
+		for (const FName CurveName : CurveNames)
 		{
 			// Check if we want to include this bone.
-			bool bIncludeCurve = (FilterText.IsEmpty() || SkelCurveName.ToString().Contains(FilterText));
+			bool bIncludeCurve = (FilterText.IsEmpty() || CurveName.ToString().Contains(FilterText));
 			if (!IncludeList.IsEmpty())
 			{
-				bIncludeCurve &= IncludeList.Contains(SkelCurveName);
+				bIncludeCurve &= IncludeList.Contains(CurveName);
 			}
 
 			if (bIncludeCurve)
 			{
-				TSharedPtr<FMLDeformerCurvePickerElement> Element = MakeShared<FMLDeformerCurvePickerElement>(SkelCurveName, FSlateColor::UseForeground());
-				Element->TextColor = (HighlightCurveNames.Contains(SkelCurveName)) ? HighlightColor : FSlateColor::UseForeground();
+				TSharedPtr<FMLDeformerCurvePickerElement> Element = MakeShared<FMLDeformerCurvePickerElement>(CurveName, FSlateColor::UseForeground());
+				Element->TextColor = (HighlightCurveNames.Contains(CurveName)) ? HighlightColor : FSlateColor::UseForeground();
 				CurveListView->AddElement(Element);
 			}
 		}
@@ -176,7 +194,7 @@ namespace UE::MLDeformer
 	void SMLDeformerCurvePickerDialog::Construct(const FArguments& InArgs)
 	{
 		HighlightCurveNames = InArgs._HighlightCurveNames;
-		Skeleton = InArgs._Skeleton;
+		SkeletalMesh = InArgs._SkeletalMesh;
 		HighlightColor = InArgs._HighlightCurveNamesColor;
 		bAllowMultiSelect = InArgs._AllowMultiSelect;
 		IncludeList = InArgs._IncludeList;
diff --git a/Engine/Plugins/Animation/MLDeformer/MLDeformerFramework/Source/MLDeformerFrameworkEditor/Private/SMLDeformerCurvePickerWidget.cpp b/Engine/Plugins/Animation/MLDeformer/MLDeformerFramework/Source/MLDeformerFrameworkEditor/Private/SMLDeformerCurvePickerWidget.cpp
--- a/Engine/Plugins/Animation/MLDeformer/MLDeformerFramework/Source/MLDeformerFrameworkEditor/Private/SMLDeformerCurvePickerWidget.cpp
+++ b/Engine/Plugins/Animation/MLDeformer/MLDeformerFramework/Source/MLDeformerFrameworkEditor/Private/SMLDeformerCurvePickerWidget.cpp
@@ -9,6 +9,7 @@
 #include "Widgets/Text/STextBlock.h"
 #include "Widgets/Layout/SBox.h"
 #include "Animation/Skeleton.h"
+#include "Engine/SkeletalMesh.h"
 
 #define LOCTEXT_NAMESPACE "SCurvePickerWidget"
 
@@ -17,11 +18,11 @@ namespace UE::MLDeformer
 	void SCurvePickerWidget::Construct(const FArguments& InArgs)
 	{
 		OnCurveNamePicked = InArgs._OnCurveNamePicked;
-		OnGetSkeleton = InArgs._OnGetSkeleton;
+		OnGetSkeletalMesh = InArgs._OnGetSkeletalMesh;
 
-		if (OnGetSkeleton.IsBound())
+		if (OnGetSkeletalMesh.IsBound())
 		{
-			Skeleton = OnGetSkeleton.Execute();
+			SkeletalMesh = OnGetSkeletalMesh.Execute();
 		}
 
 		ChildSlot
@@ -74,13 +75,21 @@ namespace UE::MLDeformer
 		CurveNames.Empty();
 		UniqueCurveNames.Empty();
 
-		// Get the curve names.
-		if (Skeleton)
+		const UAnimCurveMetaData* SkelMeshCurveData = SkeletalMesh->GetAssetUserData<UAnimCurveMetaData>();
+		if (SkelMeshCurveData)
+		{
+			TArray<FName> SkelMeshCurveNames;
+			SkelMeshCurveData->GetCurveMetaDataNames(SkelMeshCurveNames);
+			for (const FName& CurveName : SkelMeshCurveNames)
+			{
+				UniqueCurveNames.Add(CurveName.ToString());
+			}
+		}
+		
+		if (SkeletalMesh->GetSkeleton())
 		{
 			TArray<FName> SkeletonCurveNames;
-			Skeleton->GetCurveMetaDataNames(SkeletonCurveNames);
-
-			UniqueCurveNames.Reserve(SkeletonCurveNames.Num());
+			SkeletalMesh->GetSkeleton()->GetCurveMetaDataNames(SkeletonCurveNames);
 			for (const FName& CurveName : SkeletonCurveNames)
 			{
 				UniqueCurveNames.Add(CurveName.ToString());
@@ -117,7 +126,7 @@ namespace UE::MLDeformer
 	{
 		OnCurveSelectionChanged = InArgs._OnCurveSelectionChanged;
 		OnGetSelectedCurve = InArgs._OnGetSelectedCurve;
-		OnGetSkeleton = InArgs._OnGetSkeleton;
+		OnGetSkeletalMesh = InArgs._OnGetSkeletalMesh;
 
 		SuppliedToolTip = InArgs._ToolTipText.Get();
 
@@ -142,7 +151,7 @@ namespace UE::MLDeformer
 
 		TSharedRef<SCurvePickerWidget> ListWidget = SNew(SCurvePickerWidget)
 			.OnCurveNamePicked(this, &SCurveSelectionWidget::OnSelectionChanged)
-			.OnGetSkeleton(OnGetSkeleton);
+			.OnGetSkeletalMesh(OnGetSkeletalMesh);
 
 		CurvePickerButton->SetMenuContentWidgetToFocus(ListWidget->GetFilterTextWidget());
 
@@ -175,11 +184,11 @@ namespace UE::MLDeformer
 		return FText::Format(LOCTEXT("CurveClickToolTip", "Curve: {0}\n\n{1}"), GetCurrentCurveName(), SuppliedToolTip);
 	}
 
-	USkeleton* SCurveSelectionWidget::GetSkeleton() const
+	USkeletalMesh* SCurveSelectionWidget::GetSkeletalMesh() const
 	{
-		if (OnGetSkeleton.IsBound())
+		if (OnGetSkeletalMesh.IsBound())
 		{
-			return OnGetSkeleton.Execute();
+			return OnGetSkeletalMesh.Execute();
 		}
 
 		return nullptr;
diff --git a/Engine/Plugins/Animation/MLDeformer/MLDeformerFramework/Source/MLDeformerFrameworkEditor/Private/SMLDeformerInputCurvesWidget.cpp b/Engine/Plugins/Animation/MLDeformer/MLDeformerFramework/Source/MLDeformerFrameworkEditor/Private/SMLDeformerInputCurvesWidget.cpp
--- a/Engine/Plugins/Animation/MLDeformer/MLDeformerFramework/Source/MLDeformerFrameworkEditor/Private/SMLDeformerInputCurvesWidget.cpp
+++ b/Engine/Plugins/Animation/MLDeformer/MLDeformerFramework/Source/MLDeformerFrameworkEditor/Private/SMLDeformerInputCurvesWidget.cpp
@@ -13,6 +13,7 @@
 #include "Engine/SkeletalMesh.h"
 #include "Animation/Skeleton.h"
 #include "ScopedTransaction.h"
+#include "Animation/AnimCurveMetadata.h"
 
 #define LOCTEXT_NAMESPACE "MLDeformerInputCurvesWidget"
 
@@ -68,14 +69,14 @@ namespace UE::MLDeformer
 		Elements.Add(Element);
 	}
 
-	void SMLDeformerInputCurveListWidget::RefreshElements(const TArray<FName>& CurveNames, const USkeleton* Skeleton)
+	void SMLDeformerInputCurveListWidget::RefreshElements(const TArray<FName>& CurveNames, USkeletalMesh* SkeletalMesh)
 	{
 		Elements.Reset();
 
 		const FLinearColor ErrorColor = FMLDeformerEditorStyle::Get().GetColor("MLDeformer.InputsWidget.ErrorColor");
 
-		// If we have no reference skeleton, just add everything as flat list as we don't have hierarchy data.
-		if (Skeleton == nullptr)
+		// If we have no skeletal mesh, just add everything as flat list as we don't have hierarchy data.
+		if (SkeletalMesh == nullptr)
 		{
 			for (const FName CurveName : CurveNames)
 			{
@@ -90,13 +91,29 @@ namespace UE::MLDeformer
 		}
 		else
 		{
-			TArray<FName> SkeletonCurveNames;
-			Skeleton->GetCurveMetaDataNames(SkeletonCurveNames);
+			TArray<FName> AssetCurveNames;
+			if (SkeletalMesh)
+			{
+				const UAnimCurveMetaData* SkelMeshCurveData = SkeletalMesh->GetAssetUserData<UAnimCurveMetaData>();
+				if (SkelMeshCurveData)
+				{
+					TArray<FName> SkelMeshCurveNames;
+					SkelMeshCurveData->GetCurveMetaDataNames(SkelMeshCurveNames);
+					AssetCurveNames.Append(SkelMeshCurveNames);
+				}
+			}
+			
+			if (SkeletalMesh->GetSkeleton())
+			{
+				TArray<FName> SkeletonCurveNames;
+				SkeletalMesh->GetSkeleton()->GetCurveMetaDataNames(SkeletonCurveNames);
+				AssetCurveNames.Append(SkeletonCurveNames);
+			}
 
 			for (int32 CurveIndex = 0; CurveIndex < CurveNames.Num(); ++CurveIndex)
 			{
 				const FName CurveName = CurveNames[CurveIndex];
-				const bool bExistsInSkel = SkeletonCurveNames.Contains(CurveName);
+				const bool bExistsInSkel = AssetCurveNames.Contains(CurveName);
 
 				if (InputCurvesWidget->FilterText.IsEmpty() || CurveName.ToString().Contains(InputCurvesWidget->FilterText))
 				{
@@ -231,7 +248,7 @@ namespace UE::MLDeformer
 
 			TSharedPtr<SMLDeformerCurvePickerDialog> Dialog = 
 				SNew(SMLDeformerCurvePickerDialog)
-				.Skeleton(SkelMesh->GetSkeleton())
+				.SkeletalMesh(SkelMesh)
 				.AllowMultiSelect(true)
 				.HighlightCurveNamesColor(FSlateColor(HighlightColor))
 				.HighlightCurveNames(HighlightCurveNames);
@@ -341,8 +358,7 @@ namespace UE::MLDeformer
 			CurveNames.Add(CurveRef.CurveName);		
 		}
 
-		const USkeleton* Skeleton = Model->GetSkeletalMesh() ? Model->GetSkeletalMesh()->GetSkeleton() : nullptr;
-		ListWidget->RefreshElements(CurveNames, Skeleton);
+		ListWidget->RefreshElements(CurveNames, Model->GetSkeletalMesh());
 		ListWidget->RequestListRefresh();
 
 		UMLDeformerInputInfo* InputInfo = EditorModel->GetEditorInputInfo();
diff --git a/Engine/Plugins/Animation/MLDeformer/MLDeformerFramework/Source/MLDeformerFrameworkEditor/Public/MLDeformerCurveReferenceCustomization.h b/Engine/Plugins/Animation/MLDeformer/MLDeformerFramework/Source/MLDeformerFrameworkEditor/Public/MLDeformerCurveReferenceCustomization.h
--- a/Engine/Plugins/Animation/MLDeformer/MLDeformerFramework/Source/MLDeformerFrameworkEditor/Public/MLDeformerCurveReferenceCustomization.h
+++ b/Engine/Plugins/Animation/MLDeformer/MLDeformerFramework/Source/MLDeformerFrameworkEditor/Public/MLDeformerCurveReferenceCustomization.h
@@ -7,7 +7,7 @@
 #include "PropertyHandle.h"
 #include "DetailWidgetRow.h"
 
-class USkeleton;
+class USkeletalMesh;
 
 namespace UE::MLDeformer
 {
@@ -33,12 +33,12 @@ namespace UE::MLDeformer
 		/** Property to change after curve has been picked. */
 		TSharedPtr<IPropertyHandle> CurveNameProperty = nullptr;
 
-		/** The Skeleton we get the curves from. */
-		TObjectPtr<USkeleton> Skeleton = nullptr;
+		/** The skeletal mesh we get the curves from. */
+		TObjectPtr<USkeletalMesh> SkeletalMesh = nullptr;
 
 	private:
 		virtual void OnCurveSelectionChanged(const FString& Name);
 		virtual FString OnGetSelectedCurve() const;
-		virtual USkeleton* OnGetSkeleton() const;
+		virtual USkeletalMesh* OnGetSkeletalMesh() const;
 	};
 }	// namespace UE::MLDeformer
diff --git a/Engine/Plugins/Animation/MLDeformer/MLDeformerFramework/Source/MLDeformerFrameworkEditor/Public/SMLDeformerCurvePickerDialog.h b/Engine/Plugins/Animation/MLDeformer/MLDeformerFramework/Source/MLDeformerFrameworkEditor/Public/SMLDeformerCurvePickerDialog.h
--- a/Engine/Plugins/Animation/MLDeformer/MLDeformerFramework/Source/MLDeformerFrameworkEditor/Public/SMLDeformerCurvePickerDialog.h
+++ b/Engine/Plugins/Animation/MLDeformer/MLDeformerFramework/Source/MLDeformerFrameworkEditor/Public/SMLDeformerCurvePickerDialog.h
@@ -11,6 +11,7 @@
 #include "MLDeformerEditorStyle.h"
 
 class USkeleton;
+class USkeletalMesh;
 
 namespace UE::MLDeformer
 {
@@ -77,7 +78,7 @@ namespace UE::MLDeformer
 		friend class SMLDeformerCurvePickerListView;
 
 		SLATE_BEGIN_ARGS(SMLDeformerCurvePickerDialog) {}
-		SLATE_ARGUMENT(USkeleton*, Skeleton)					// The skeleton to pick the curves from.
+		SLATE_ARGUMENT(USkeletalMesh*, SkeletalMesh)			// The skeleton to pick the curves from.
 		SLATE_ARGUMENT(TArray<FName>, IncludeList)				// If empty, all curves can be picked. If not empty, only show curves with names inside this list.
 		SLATE_ARGUMENT(TArray<FName>, HighlightCurveNames)		// The curve names to highlight, if any.
 		SLATE_ARGUMENT(FSlateColor, HighlightCurveNamesColor)	// The text color to use for this highlighting.
@@ -107,7 +108,7 @@ namespace UE::MLDeformer
 		TArray<FName> PickedCurveNames;
 		TArray<FName> IncludeList;
 		FString FilterText;
-		TObjectPtr<USkeleton> Skeleton;
+		TWeakObjectPtr<USkeletalMesh> SkeletalMesh;
 		TArray<FName> HighlightCurveNames;
 		FSlateColor HighlightColor = FSlateColor(FLinearColor::White);
 		bool bAllowMultiSelect = true;
diff --git a/Engine/Plugins/Animation/MLDeformer/MLDeformerFramework/Source/MLDeformerFrameworkEditor/Public/SMLDeformerCurvePickerWidget.h b/Engine/Plugins/Animation/MLDeformer/MLDeformerFramework/Source/MLDeformerFrameworkEditor/Public/SMLDeformerCurvePickerWidget.h
--- a/Engine/Plugins/Animation/MLDeformer/MLDeformerFramework/Source/MLDeformerFrameworkEditor/Public/SMLDeformerCurvePickerWidget.h
+++ b/Engine/Plugins/Animation/MLDeformer/MLDeformerFramework/Source/MLDeformerFrameworkEditor/Public/SMLDeformerCurvePickerWidget.h
@@ -11,13 +11,13 @@
 
 class SComboButton;
 class SSearchBox;
-class USkeleton;
+class USkeletalMesh;
 
 DECLARE_DELEGATE_OneParam(FOnCurveSelectionChanged, const FString&);
 DECLARE_DELEGATE_OneParam(FOnCurveNamePicked, const FString&)
 
 DECLARE_DELEGATE_RetVal(FString, FOnGetSelectedCurve);
-DECLARE_DELEGATE_RetVal(USkeleton*, FOnGetSkeleton);
+DECLARE_DELEGATE_RetVal(USkeletalMesh*, FOnGetSkeletalMesh)
 
 
 namespace UE::MLDeformer
@@ -31,7 +31,7 @@ namespace UE::MLDeformer
 	public:
 		SLATE_BEGIN_ARGS(SCurvePickerWidget) {}
 		SLATE_EVENT(FOnCurveNamePicked, OnCurveNamePicked)
-		SLATE_EVENT(FOnGetSkeleton, OnGetSkeleton)
+		SLATE_EVENT(FOnGetSkeletalMesh, OnGetSkeletalMesh)
 		SLATE_END_ARGS()
 
 		void Construct(const FArguments& InArgs);
@@ -49,14 +49,14 @@ namespace UE::MLDeformer
 		/** Delegate fired when a curve name is picked. */
 		FOnCurveNamePicked OnCurveNamePicked;
 
-		/** Provide us with a skeleton. */
-		FOnGetSkeleton OnGetSkeleton;
+		/** Provide us with a skeletal mesh. */
+		FOnGetSkeletalMesh OnGetSkeletalMesh;
 
 		/** The search filter box. */
 		TSharedPtr<SSearchBox> SearchBox;
 
-		/** The skeleton to get the curves from. */
-		TObjectPtr<USkeleton> Skeleton;
+		/** The skeletal mesh to get the curves from. */
+		TWeakObjectPtr<USkeletalMesh> SkeletalMesh;
 
 		/** The names of the curves we are displaying. */
 		TArray<TSharedPtr<FString>> CurveNames;
@@ -81,7 +81,7 @@ namespace UE::MLDeformer
 		SLATE_BEGIN_ARGS(SCurveSelectionWidget) {}
 			SLATE_EVENT(FOnCurveSelectionChanged, OnCurveSelectionChanged);
 			SLATE_EVENT(FOnGetSelectedCurve, OnGetSelectedCurve);
-			SLATE_EVENT(FOnGetSkeleton, OnGetSkeleton)
+			SLATE_EVENT(FOnGetSkeletalMesh, OnGetSkeletalMesh)
 		SLATE_END_ARGS();
 
 		void Construct(const FArguments& InArgs);
@@ -89,7 +89,7 @@ namespace UE::MLDeformer
 	private: 
 		TSharedRef<SWidget> CreateSkeletonWidgetMenu();
 		void OnSelectionChanged(const FString& CurveName);
-		USkeleton* GetSkeleton() const;
+		USkeletalMesh* GetSkeletalMesh() const;
 		FText GetCurrentCurveName() const;
 		FText GetFinalToolTip() const;
 
@@ -97,7 +97,7 @@ namespace UE::MLDeformer
 		TSharedPtr<SComboButton> CurvePickerButton;
 		FOnCurveSelectionChanged OnCurveSelectionChanged;
 		FOnGetSelectedCurve OnGetSelectedCurve;
-		FOnGetSkeleton OnGetSkeleton;
+		FOnGetSkeletalMesh OnGetSkeletalMesh;
 		FText SuppliedToolTip;
 		FText SelectedCurve;
 	};
diff --git a/Engine/Plugins/Animation/MLDeformer/MLDeformerFramework/Source/MLDeformerFrameworkEditor/Public/SMLDeformerInputCurvesWidget.h b/Engine/Plugins/Animation/MLDeformer/MLDeformerFramework/Source/MLDeformerFrameworkEditor/Public/SMLDeformerInputCurvesWidget.h
--- a/Engine/Plugins/Animation/MLDeformer/MLDeformerFramework/Source/MLDeformerFrameworkEditor/Public/SMLDeformerInputCurvesWidget.h
+++ b/Engine/Plugins/Animation/MLDeformer/MLDeformerFramework/Source/MLDeformerFrameworkEditor/Public/SMLDeformerInputCurvesWidget.h
@@ -11,7 +11,7 @@
 #include "UObject/NameTypes.h"
 #include "Styling/SlateColor.h"
 
-class USkeleton;
+class USkeletalMesh;
 
 namespace UE::MLDeformer
 {
@@ -72,7 +72,7 @@ namespace UE::MLDeformer
 		void Construct(const FArguments& InArgs);
 		const TArray<TSharedPtr<FMLDeformerInputCurveListElement>>& GetElements() const	{ return Elements; }
 		void AddElement(TSharedPtr<FMLDeformerInputCurveListElement> Element);
-		void RefreshElements(const TArray<FName>& CurveNames, const USkeleton* Skeleton);
+		void RefreshElements(const TArray<FName>& CurveNames, USkeletalMesh* SkeletalMesh);
 
 	private:
 		TSharedRef<ITableRow> MakeTableRowWidget(TSharedPtr<FMLDeformerInputCurveListElement> InItem, const TSharedRef<STableViewBase>& OwnerTable);
diff --git a/Engine/Plugins/Animation/MLDeformer/NeuralMorphModel/Source/NeuralMorphModelEditor/Private/SNeuralMorphCurveGroupsWidget.cpp b/Engine/Plugins/Animation/MLDeformer/NeuralMorphModel/Source/NeuralMorphModelEditor/Private/SNeuralMorphCurveGroupsWidget.cpp
--- a/Engine/Plugins/Animation/MLDeformer/NeuralMorphModel/Source/NeuralMorphModelEditor/Private/SNeuralMorphCurveGroupsWidget.cpp
+++ b/Engine/Plugins/Animation/MLDeformer/NeuralMorphModel/Source/NeuralMorphModelEditor/Private/SNeuralMorphCurveGroupsWidget.cpp
@@ -111,7 +111,7 @@ namespace UE::NeuralMorphModel
 
 		TSharedPtr<SMLDeformerCurvePickerDialog> Dialog = 
 			SNew(SMLDeformerCurvePickerDialog)
-			.Skeleton(SkelMesh->GetSkeleton())
+			.SkeletalMesh(SkelMesh)
 			.AllowMultiSelect(true)
 			.IncludeList(InputInfo->GetCurveNames());
 
@@ -249,7 +249,7 @@ namespace UE::NeuralMorphModel
 		const FLinearColor HighlightColor = FMLDeformerEditorStyle::Get().GetColor("MLDeformer.InputsWidget.HighlightColor");
 		TSharedPtr<SMLDeformerCurvePickerDialog> Dialog = 
 			SNew(SMLDeformerCurvePickerDialog)
-			.Skeleton(SkelMesh->GetSkeleton())
+			.SkeletalMesh(SkelMesh)
 			.AllowMultiSelect(true)
 			.HighlightCurveNames(HighlightedCurves)
 			.HighlightCurveNamesColor(HighlightColor)
