[MLDeformer] Added a new console command "MLDeformer.DebugDrawInputPoses" which will draw the input pose to the ML Deformer. It can be useful to quickly visually see what meshes have ML Deformers on them, and what bones are included.

diff --git a/Engine/Plugins/Animation/MLDeformer/MLDeformerFramework/Source/MLDeformerFramework/Private/MLDeformerModelInstance.cpp b/Engine/Plugins/Animation/MLDeformer/MLDeformerFramework/Source/MLDeformerFramework/Private/MLDeformerModelInstance.cpp
--- a/Engine/Plugins/Animation/MLDeformer/MLDeformerFramework/Source/MLDeformerFramework/Private/MLDeformerModelInstance.cpp
+++ b/Engine/Plugins/Animation/MLDeformer/MLDeformerFramework/Source/MLDeformerFramework/Private/MLDeformerModelInstance.cpp
@@ -11,6 +11,7 @@
 #include "Animation/AnimInstance.h"
 #include "RenderingThread.h"
 #include "RenderGraphBuilder.h"
+#include "DrawDebugHelpers.h"
 
 #include UE_INLINE_GENERATED_CPP_BY_NAME(MLDeformerModelInstance)
 
@@ -20,6 +21,140 @@ CSV_DEFINE_CATEGORY_MODULE(MLDEFORMERFRAMEWORK_API, MLDeformer, false);
 DEFINE_STAT(STAT_MLDeformerInference);
 #endif
 
+namespace
+{
+	void CalculateBoneSpaceTransforms(const FReferenceSkeleton& RefSkel, const TArray<FTransform>& ComponentSpaceTransforms, TArray<FTransform>& OutBoneSpaceTransforms)
+	{
+		for (int32 Index = 0; Index < RefSkel.GetNum(); ++Index)
+		{
+			const int32 ParentIndex = RefSkel.GetParentIndex(Index);
+			if (ParentIndex != INDEX_NONE)
+			{
+				OutBoneSpaceTransforms[Index] = ComponentSpaceTransforms[Index].GetRelativeTransform(ComponentSpaceTransforms[ParentIndex]);
+				OutBoneSpaceTransforms[Index].NormalizeRotation();
+			}
+			else
+			{
+				OutBoneSpaceTransforms[Index] = ComponentSpaceTransforms[Index];
+			}
+		}
+	}
+}	// Anonymous namespace
+
+
+void UMLDeformerModelInstance::DrawDebugInputPoses()
+{
+	// If we use a leader component, we have to get the transforms from that component.
+	USkeletalMeshComponent* FinalSkelMeshComponent = GetFinalSkeletalMeshComponent();
+	const UWorld* World = GetWorld();
+	const USkinnedMeshComponent* LeaderPoseComponent = FinalSkelMeshComponent->LeaderPoseComponent.Get();
+	if (LeaderPoseComponent)
+	{
+		const FTransform& ComponentTransform = FinalSkelMeshComponent->GetComponentTransform();
+		const TArray<FTransform>& LeaderComponentSpaceTransforms = LeaderPoseComponent->GetComponentSpaceTransforms();
+		if (!LeaderComponentSpaceTransforms.IsEmpty())
+		{
+			const USkinnedAsset* SkinnedAsset = LeaderPoseComponent->GetSkinnedAsset();
+			const FReferenceSkeleton& RefSkel = SkinnedAsset->GetRefSkeleton();
+			const int32 NumTrainingBones = AssetBonesToSkelMeshMappings.Num();
+			for (int32 Index = 0; Index < NumTrainingBones; ++Index)
+			{
+				const int32 ComponentBoneIndex = AssetBonesToSkelMeshMappings[Index];
+				if (LeaderComponentSpaceTransforms.IsValidIndex(ComponentBoneIndex))
+				{
+					const FTransform& ComponentSpaceBoneTransform = LeaderComponentSpaceTransforms[ComponentBoneIndex];
+					const int32 ParentIndex = RefSkel.GetParentIndex(ComponentBoneIndex);
+					const FTransform WorldTransform = ComponentSpaceBoneTransform * ComponentTransform;
+					if (LeaderComponentSpaceTransforms.IsValidIndex(ParentIndex))
+					{
+						const FTransform WorldParentTransform = LeaderComponentSpaceTransforms[ParentIndex] * ComponentTransform;
+						DrawDebugLine(World, WorldTransform.GetTranslation(), WorldParentTransform.GetTranslation(), FColor::Yellow, false, -1.0f, 255, 0.0f);
+						DrawDebugPoint(World, WorldTransform.GetTranslation(), 4.0f, FColor::Yellow, false, -1.0f, 255);
+						DrawDebugCoordinateSystem(World, WorldTransform.GetTranslation(), WorldTransform.GetRotation().Rotator(), 7.5f, false, -1.0f, 255, 0.5f);						
+					}
+					else
+					{
+						DrawDebugPoint(World, WorldTransform.GetTranslation(), 4.0f, FColor::Yellow, false, -1.0f, 255);
+						DrawDebugCoordinateSystem(World, WorldTransform.GetTranslation(), WorldTransform.GetRotation().Rotator(), 7.5f, false, -1.0f, 255, 0.5f);						
+					}
+				}
+			}
+		}
+	}
+	else // No leader pose is used.
+	{
+		if (FinalSkelMeshComponent == SkeletalMeshComponent)
+		{
+			const FTransform& ComponentTransform = FinalSkelMeshComponent->GetComponentTransform();
+			const TArray<FTransform>& ComponentSpaceTransforms = FinalSkelMeshComponent->GetComponentSpaceTransforms();
+			if (!ComponentSpaceTransforms.IsEmpty())
+			{
+				const USkinnedAsset* SkinnedAsset = FinalSkelMeshComponent->GetSkinnedAsset();
+				const FReferenceSkeleton& RefSkel = SkinnedAsset->GetRefSkeleton();
+				const int32 NumTrainingBones = AssetBonesToSkelMeshMappings.Num();
+				for (int32 Index = 0; Index < NumTrainingBones; ++Index)
+				{
+					const int32 ComponentBoneIndex = AssetBonesToSkelMeshMappings[Index];
+					if (ComponentSpaceTransforms.IsValidIndex(ComponentBoneIndex))
+					{
+						const FTransform& ComponentSpaceBoneTransform = ComponentSpaceTransforms[ComponentBoneIndex];
+						const FTransform WorldTransform = ComponentSpaceBoneTransform * ComponentTransform;
+						const int32 ParentIndex = RefSkel.GetParentIndex(ComponentBoneIndex);
+						if (ComponentSpaceTransforms.IsValidIndex(ParentIndex))
+						{
+							const FTransform WorldParentTransform = ComponentSpaceTransforms[ParentIndex] * ComponentTransform;
+							DrawDebugLine(World, WorldTransform.GetTranslation(), WorldParentTransform.GetTranslation(), FColor::Yellow, false, -1.0f, 255, 0.0f);
+							DrawDebugPoint(World, WorldTransform.GetTranslation(), 4.0f, FColor::Yellow, false, -1.0f, 255);
+							DrawDebugCoordinateSystem(World, WorldTransform.GetTranslation(), WorldTransform.GetRotation().Rotator(), 7.5f, false, -1.0f, 255, 0.5f);
+						}
+						else
+						{
+							DrawDebugPoint(World, WorldTransform.GetTranslation(), 4.0f, FColor::Yellow, false, -1.0f, 255);
+							DrawDebugCoordinateSystem(World, WorldTransform.GetTranslation(), WorldTransform.GetRotation().Rotator(), 7.5f, false, -1.0f, 255, 1.0f);						
+						}
+					}
+				}
+			}
+		}
+		else // We are debugging and grabbing transforms from another component.
+		{
+			USkeletalMesh* SkelMeshAsset = FinalSkelMeshComponent->GetSkeletalMeshAsset();
+			check(SkelMeshAsset);
+
+			const FTransform& ComponentTransform = SkeletalMeshComponent->GetComponentTransform();
+			
+			const FReferenceSkeleton& RefSkel = SkelMeshAsset->GetRefSkeleton();
+			const TArray<FTransform>& ComponentSpaceTransforms = FinalSkelMeshComponent->GetComponentSpaceTransforms();
+			if (!ComponentSpaceTransforms.IsEmpty())
+			{
+				const int32 NumTrainingBones = AssetBonesToSkelMeshMappings.Num();
+				for (int32 Index = 0; Index < NumTrainingBones; ++Index)
+				{
+					const int32 ComponentBoneIndex = AssetBonesToSkelMeshMappings[Index];
+					if (ComponentSpaceTransforms.IsValidIndex(ComponentBoneIndex))
+					{
+						const FTransform& ComponentSpaceBoneTransform = ComponentSpaceTransforms[ComponentBoneIndex];
+						const FTransform WorldTransform = ComponentSpaceBoneTransform * ComponentTransform;
+						const int32 ParentIndex = RefSkel.GetParentIndex(ComponentBoneIndex);
+						if (ComponentSpaceTransforms.IsValidIndex(ParentIndex))
+						{
+							const FTransform WorldParentTransform = ComponentSpaceTransforms[ParentIndex] * ComponentTransform;
+							DrawDebugLine(World, WorldTransform.GetTranslation(), WorldParentTransform.GetTranslation(), FColor::Yellow, false, -1.0f, 255, 0.0f);
+							DrawDebugPoint(World, WorldTransform.GetTranslation(), 4.0f, FColor::Yellow, false, -1.0f, 255);
+							DrawDebugCoordinateSystem(World, WorldTransform.GetTranslation(), WorldTransform.GetRotation().Rotator(), 7.5f, false, -1.0f, 255, 0.5f);
+						}
+						else
+						{
+							DrawDebugPoint(World, WorldTransform.GetTranslation(), 4.0f, FColor::Yellow, false, -1.0f, 255);
+							DrawDebugCoordinateSystem(World, WorldTransform.GetTranslation(), WorldTransform.GetRotation().Rotator(), 7.5f, false, -1.0f, 255, 1.0f);						
+						}
+					}
+				}
+			}
+		}
+	}
+}
+
 void UMLDeformerModelInstance::BeginDestroy()
 {
 	RenderCommandFence.BeginFence();
@@ -180,28 +315,6 @@ USkeletalMeshComponent* UMLDeformerModelInstance::GetFinalSkeletalMeshComponent(
 	return SkeletalMeshComponent.Get();
 }
 
-
-namespace
-{
-	void CalculateBoneSpaceTransforms(const FReferenceSkeleton& RefSkel, const TArray<FTransform>& ComponentSpaceTransforms, TArray<FTransform>& OutBoneSpaceTransforms)
-	{
-		for (int32 Index = 0; Index < RefSkel.GetNum(); ++Index)
-		{
-			const int32 ParentIndex = RefSkel.GetParentIndex(Index);
-			if (ParentIndex != INDEX_NONE)
-			{
-				OutBoneSpaceTransforms[Index] = ComponentSpaceTransforms[Index].GetRelativeTransform(ComponentSpaceTransforms[ParentIndex]);
-				OutBoneSpaceTransforms[Index].NormalizeRotation();
-			}
-			else
-			{
-				OutBoneSpaceTransforms[Index] = ComponentSpaceTransforms[Index];
-			}
-		}
-	}
-}	// Anonymous namespace
-
-
 void UMLDeformerModelInstance::UpdateBoneTransforms()
 {
 	// If we use a leader component, we have to get the transforms from that component.
@@ -418,6 +531,12 @@ void UMLDeformerModelInstance::Tick(float DeltaTime, float ModelWeight)
 		// For models using neural networks this will perform the inference, 
 		// calculate the network outputs and possibly use them, depending on how the model works.
 		Execute(ModelWeight);
+		const bool bDrawDebugInputPoses = UE::MLDeformer::CVarMLDeformerDebugDrawInputPoses->GetBool();
+		if (bDrawDebugInputPoses)
+		{
+			DrawDebugInputPoses();
+		}
+
 		bExecuteCalled = true;
 	}
 	else
diff --git a/Engine/Plugins/Animation/MLDeformer/MLDeformerFramework/Source/MLDeformerFramework/Public/MLDeformerModelInstance.h b/Engine/Plugins/Animation/MLDeformer/MLDeformerFramework/Source/MLDeformerFramework/Public/MLDeformerModelInstance.h
--- a/Engine/Plugins/Animation/MLDeformer/MLDeformerFramework/Source/MLDeformerFramework/Public/MLDeformerModelInstance.h
+++ b/Engine/Plugins/Animation/MLDeformer/MLDeformerFramework/Source/MLDeformerFramework/Public/MLDeformerModelInstance.h
@@ -181,6 +181,12 @@ public:
 	UE_API USkeletalMeshComponent* GetFinalSkeletalMeshComponent() const;
 
 protected:
+	/**
+	 * Draw the input poses that are currently input to this instance.
+	 * This is linked to the "MLDeformer.DebugDrawInputPoses" console command. 
+	 */
+	UE_API void DrawDebugInputPoses();
+	
 	/**
 	 * Update the neural network input values directly inside its input tensor.
 	 * This will copy over the bone transforms (and possibly morph and curve weights) into this flat array.
diff --git a/Engine/Plugins/Animation/MLDeformer/MLDeformerFramework/Source/MLDeformerFramework/Public/MLDeformerModule.h b/Engine/Plugins/Animation/MLDeformer/MLDeformerFramework/Source/MLDeformerFramework/Public/MLDeformerModule.h
--- a/Engine/Plugins/Animation/MLDeformer/MLDeformerFramework/Source/MLDeformerFramework/Public/MLDeformerModule.h
+++ b/Engine/Plugins/Animation/MLDeformer/MLDeformerFramework/Source/MLDeformerFramework/Public/MLDeformerModule.h
@@ -13,6 +13,12 @@ MLDEFORMERFRAMEWORK_API DECLARE_LOG_CATEGORY_EXTERN(LogMLDeformer, Log, All);
 
 namespace UE::MLDeformer
 {
+	static TAutoConsoleVariable<bool> CVarMLDeformerDebugDrawInputPoses(
+		TEXT("MLDeformer.DebugDrawInputPoses"),
+		false,
+		TEXT("Draw the input pose for each ML Deformer that is active.")
+	);
+	
 	/**
 	 * The runtime module for the ML Deformer.
 	 */
