[MLDeformer] Improve performance of the asset editor. Improved responsiveness no property changes, faster loading, faster sampling of anims.

diff --git a/Engine/Plugins/Animation/MLDeformer/MLDeformerFramework/Source/MLDeformerFramework/Public/MLDeformerGeomCacheModel.h b/Engine/Plugins/Animation/MLDeformer/MLDeformerFramework/Source/MLDeformerFramework/Public/MLDeformerGeomCacheModel.h
--- a/Engine/Plugins/Animation/MLDeformer/MLDeformerFramework/Source/MLDeformerFramework/Public/MLDeformerGeomCacheModel.h
+++ b/Engine/Plugins/Animation/MLDeformer/MLDeformerFramework/Source/MLDeformerFramework/Public/MLDeformerGeomCacheModel.h
@@ -19,7 +19,7 @@ class UMLDeformerAsset;
 class UMLDeformerGeomCacheVizSettings;
 
 /**
- * A ML Deformer model that has a geometry cache as target mesh.
+ * An ML Deformer model that has a geometry cache as target mesh.
  * Use this in combination with UMLDeformerGeomCacheVizSettings, FMLDeformerGeomCacheEditorModel, FMLDeformerGeomCacheVizSettingsDetails and FMLDeformerGeomCacheModelDetails.
  */
 UCLASS(MinimalAPI)
diff --git a/Engine/Plugins/Animation/MLDeformer/MLDeformerFramework/Source/MLDeformerFrameworkEditor/Private/MLDeformerAssetDefinitions.cpp b/Engine/Plugins/Animation/MLDeformer/MLDeformerFramework/Source/MLDeformerFrameworkEditor/Private/MLDeformerAssetDefinitions.cpp
--- a/Engine/Plugins/Animation/MLDeformer/MLDeformerFramework/Source/MLDeformerFrameworkEditor/Private/MLDeformerAssetDefinitions.cpp
+++ b/Engine/Plugins/Animation/MLDeformer/MLDeformerFramework/Source/MLDeformerFrameworkEditor/Private/MLDeformerAssetDefinitions.cpp
@@ -3,12 +3,13 @@
 #include "MLDeformerAssetDefinitions.h"
 #include "MLDeformerAsset.h"
 #include "MLDeformerEditorToolkit.h"
+#include "MLDeformerModule.h"
 
 #define LOCTEXT_NAMESPACE "MLDeformer_AssetTypeActions"
 
 FText UAssetDefinition_MLDeformer::GetAssetDisplayName() const
 { 
-	return LOCTEXT("AssetTypeActions_MLPoser", "ML Deformer");
+	return LOCTEXT("AssetTypeActions_MLDeformer", "ML Deformer");
 }
 
 FLinearColor UAssetDefinition_MLDeformer::GetAssetColor() const
@@ -32,8 +33,10 @@ EAssetCommandResult UAssetDefinition_MLDeformer::OpenAssets(const FAssetOpenArgs
 	using namespace UE::MLDeformer;
 	for (UMLDeformerAsset* Asset : OpenArgs.LoadObjects<UMLDeformerAsset>())
 	{
+		const double StartTime = FPlatformTime::Seconds();
 		TSharedRef<FMLDeformerEditorToolkit> NewEditor(new FMLDeformerEditorToolkit());
 		NewEditor->InitAssetEditor(OpenArgs.GetToolkitMode(), OpenArgs.ToolkitHost, Asset);
+		UE_LOG(LogMLDeformer, Display, TEXT("ML Deformer asset editor took %.1f seconds to open."), FPlatformTime::Seconds() - StartTime); 
 	}
 
 	return EAssetCommandResult::Handled;
diff --git a/Engine/Plugins/Animation/MLDeformer/MLDeformerFramework/Source/MLDeformerFrameworkEditor/Private/MLDeformerEditorModel.cpp b/Engine/Plugins/Animation/MLDeformer/MLDeformerFramework/Source/MLDeformerFrameworkEditor/Private/MLDeformerEditorModel.cpp
--- a/Engine/Plugins/Animation/MLDeformer/MLDeformerFramework/Source/MLDeformerFrameworkEditor/Private/MLDeformerEditorModel.cpp
+++ b/Engine/Plugins/Animation/MLDeformer/MLDeformerFramework/Source/MLDeformerFrameworkEditor/Private/MLDeformerEditorModel.cpp
@@ -145,6 +145,7 @@ namespace UE::MLDeformer
 
 		UpdateLODMappings();
 		UpdateTrainingDeviceList();
+		UpdateEditorInputInfo();
 	}
 
 	void FMLDeformerEditorModel::UpdateVertexAttributeNames()
@@ -988,13 +989,17 @@ namespace UE::MLDeformer
 	{
 		CurrentTrainingFrame = -1;
 		UpdateTimelineTrainingAnimList();
+
 		Editor->UpdateTimeSliderRange();
 		Model->UpdateCachedNumVertices();
+
 		UpdateActorVisibility();
 		UpdateNumTrainingFrames();
 		UpdateDeformerGraph();
 		RefreshMLDeformerComponents();
+		
 		CreateSamplers();
+		UpdateEditorInputInfo();
 		UpdateIsReadyForTrainingState();
 
 		const int32 TrainingFrame = Model->GetVizSettings()->GetTrainingFrameNumber();
@@ -1003,7 +1008,6 @@ namespace UE::MLDeformer
 		const int32 TestFrame = Model->GetVizSettings()->GetTestingFrameNumber();
 		SetTestFrame(TestFrame, false);
 
-		UpdateEditorInputInfo();
 		CheckTrainingDataFrameChanged();
 	}
 
@@ -1191,7 +1195,7 @@ namespace UE::MLDeformer
 			SetResamplingInputOutputsNeeded(true);
 			Model->InitVertexMap();
 			Model->InitGPUData();
-			UpdateDeformerGraph();
+			SampleDeltas();
 			if (GetEditor())
 			{
 				GetEditor()->ZoomOnActors();
@@ -1265,7 +1269,6 @@ namespace UE::MLDeformer
 		}
 		else if (Property->GetFName() == UMLDeformerModel::GetMaxTrainingFramesPropertyName())
 		{
-			TriggerInputAssetChanged();
 			SetResamplingInputOutputsNeeded(true);
 		}
 		else
@@ -2374,12 +2377,10 @@ namespace UE::MLDeformer
 			Model->GetSkeletalMesh() == nullptr ||
 			GetNumTrainingFrames() == 0)
 		{
-			UpdateEditorInputInfo();
 			return false;
 		}
 
 		// Make sure we have inputs.
-		UpdateEditorInputInfo();
 		if (GetEditorInputInfo()->IsEmpty())
 		{
 			return false;
@@ -2412,7 +2413,7 @@ namespace UE::MLDeformer
 		OnInputAssetsChanged();
 		OnPostInputAssetChanged();
 		UpdateMemoryUsage();
-
+		
 		if (Editor->GetModelDetailsView())
 		{
 			Editor->GetModelDetailsView()->ForceRefresh();
@@ -2542,7 +2543,7 @@ namespace UE::MLDeformer
 		check(!Model->GetVertexMap().IsEmpty());
 
 		FScopedSlowTask Task(NumMorphTargets, LOCTEXT("CreateEngineMorphTargetProgress", "Creating morph targets"));
-		Task.MakeDialog(false);	
+		Task.MakeDialogDelayed(2.0f, false);	
 
 		USkeletalMesh* SkelMesh = Model->GetSkeletalMesh();
 		FSkeletalMeshRenderData* RenderData = SkelMesh->GetResourceForRendering();
@@ -2656,7 +2657,7 @@ namespace UE::MLDeformer
 	void FMLDeformerEditorModel::CompressEngineMorphTargets(FMorphTargetVertexInfoBuffers& OutMorphBuffers, const TArray<UMorphTarget*>& MorphTargets, int32 LOD, float MorphErrorTolerance)
 	{
 		FScopedSlowTask Task(2, LOCTEXT("CompressEngineMorphTargetProgress", "Compressing morph targets"));
-		Task.MakeDialog(false);	
+		Task.MakeDialogDelayed(2.0f, false);	
 
 		USkeletalMesh* SkelMesh = Model->GetSkeletalMesh();
 		FSkeletalMeshRenderData* RenderData = SkelMesh->GetResourceForRendering();
diff --git a/Engine/Plugins/Animation/MLDeformer/MLDeformerFramework/Source/MLDeformerFrameworkEditor/Private/MLDeformerEditorToolkit.cpp b/Engine/Plugins/Animation/MLDeformer/MLDeformerFramework/Source/MLDeformerFrameworkEditor/Private/MLDeformerEditorToolkit.cpp
--- a/Engine/Plugins/Animation/MLDeformer/MLDeformerFramework/Source/MLDeformerFrameworkEditor/Private/MLDeformerEditorToolkit.cpp
+++ b/Engine/Plugins/Animation/MLDeformer/MLDeformerFramework/Source/MLDeformerFrameworkEditor/Private/MLDeformerEditorToolkit.cpp
@@ -886,7 +886,6 @@ namespace UE::MLDeformer
 			{
 				GetVizSettingsDetailsView()->SetObject(nullptr, true);
 			}
-			ActiveModel->TriggerInputAssetChanged(false);
 			ActiveModel->CreateHeatMapAssets();
 			ActiveModel->SetHeatMapMaterialEnabled(ActiveModel->GetModel()->GetVizSettings()->GetShowHeatMap());
 			ActiveModel->SetResamplingInputOutputsNeeded(true);
@@ -1234,7 +1233,7 @@ namespace UE::MLDeformer
 		if (ActiveModel.IsValid())
 		{
 			ActiveModel->UpdateLabels();
-			ActiveModel->TriggerInputAssetChanged(false);
+			ActiveModel->RequestTriggerInputAssetsChangedNextTick();
 		}
 	}
 
diff --git a/Engine/Plugins/Animation/MLDeformer/MLDeformerFramework/Source/MLDeformerFrameworkEditor/Private/MLDeformerGeomCacheSampler.cpp b/Engine/Plugins/Animation/MLDeformer/MLDeformerFramework/Source/MLDeformerFrameworkEditor/Private/MLDeformerGeomCacheSampler.cpp
--- a/Engine/Plugins/Animation/MLDeformer/MLDeformerFramework/Source/MLDeformerFrameworkEditor/Private/MLDeformerGeomCacheSampler.cpp
+++ b/Engine/Plugins/Animation/MLDeformer/MLDeformerFramework/Source/MLDeformerFrameworkEditor/Private/MLDeformerGeomCacheSampler.cpp
@@ -4,7 +4,6 @@
 #include "MLDeformerModel.h"
 #include "MLDeformerEditorModel.h"
 #include "MLDeformerInputInfo.h"
-#include "MLDeformerEditorModule.h"
 #include "MLDeformerEditorToolkit.h"
 #include "MLDeformerGeomCacheHelpers.h"
 #include "MLDeformerGeomCacheModel.h"
@@ -70,7 +69,8 @@ namespace UE::MLDeformer
 		// Call this first to update bone and curve values.
 		// This will also calculate the skinned positions if the delta space is set to PostSkinning.
 		FMLDeformerSampler::Sample(InAnimFrameIndex);
-		USkeletalMesh* SkeletalMesh = SkeletalMeshComponent.Get() ? SkeletalMeshComponent->GetSkeletalMeshAsset() : nullptr;
+		
+		const USkeletalMesh* SkeletalMesh = SkeletalMeshComponent.Get() ? SkeletalMeshComponent->GetSkeletalMeshAsset() : nullptr;
 		UGeometryCache* GeometryCache = GeometryCacheComponent->GetGeometryCache();
 		if (SkeletalMeshComponent && SkeletalMesh && GeometryCacheComponent && GeometryCache)
 		{
@@ -84,7 +84,7 @@ namespace UE::MLDeformer
 			
 			for (int32 MeshMappingIndex = 0; MeshMappingIndex < MeshMappings.Num(); ++MeshMappingIndex)
 			{
-				const UE::MLDeformer::FMLDeformerGeomCacheMeshMapping& MeshMapping = MeshMappings[MeshMappingIndex];
+				const FMLDeformerGeomCacheMeshMapping& MeshMapping = MeshMappings[MeshMappingIndex];
 				TArrayView<const int32> GeoPartInfo = GeoPartOffsetAndCounts.Get(MeshMapping.MeshIndex);
 				const int32 StartImportedVertex = GeoPartInfo[0];
 				const int32 NumVertices = GeoPartInfo[1];
@@ -102,66 +102,54 @@ namespace UE::MLDeformer
 				const FSkeletalMeshLODRenderData& SkelMeshLODData = SkeletalMesh->GetResourceForRendering()->LODRenderData[LODIndex];
 				const FSkinWeightVertexBuffer& SkinWeightBuffer = *SkeletalMeshComponent->GetSkinWeightBuffer(LODIndex);
 
-				const int32 BatchSize = 500;
-				const int32 NumBatches = (NumVertices / BatchSize) + 1;
-				ParallelFor(NumBatches, [&](int32 BatchIndex)
+				ParallelFor(TEXT("GeomCacheSampler_Sample"), NumVertices, 250, [&](int32 VertexIndex)
 				{
-					const int32 StartVertex = BatchIndex * BatchSize;
-					if (StartVertex >= NumVertices || VertexDeltas.IsEmpty())
+					const int32 SkinnedVertexIndex = StartImportedVertex + VertexIndex;
+					const int32 GeomCacheVertexIndex = MeshMapping.SkelMeshToTrackVertexMap[VertexIndex];
+					if (GeomCacheVertexIndex != INDEX_NONE && GeomCacheMeshData.Positions.IsValidIndex(GeomCacheVertexIndex))
 					{
-						return;
-					}
+						const int32 RenderVertexIndex = MeshMapping.ImportedVertexToRenderVertexMap[VertexIndex];
+						if (RenderVertexIndex == INDEX_NONE)
+						{
+							return;
+						}
 
-					const int32 NumVertsInBatch = (StartVertex + BatchSize) < NumVertices ? BatchSize : FMath::Max(NumVertices - StartVertex, 0);
-					for (int32 VertexIndex = StartVertex; VertexIndex < StartVertex + NumVertsInBatch; ++VertexIndex)
-					{
-						const int32 SkinnedVertexIndex = StartImportedVertex + VertexIndex;
-						const int32 GeomCacheVertexIndex = MeshMapping.SkelMeshToTrackVertexMap[VertexIndex];
-						if (GeomCacheVertexIndex != INDEX_NONE && GeomCacheMeshData.Positions.IsValidIndex(GeomCacheVertexIndex))
+						// Calculate the pre-skinning delta.
+						FVector3f Delta = FVector3f::ZeroVector;
+						if (VertexDeltaSpace == EVertexDeltaSpace::PreSkinning)
 						{
-							const int32 RenderVertexIndex = MeshMapping.ImportedVertexToRenderVertexMap[VertexIndex];
-							if (RenderVertexIndex == INDEX_NONE)
+							if (SkinningMode == EMLDeformerSkinningMode::Linear)
 							{
-								continue;
-							}
+								// Calculate the inverse skinning transform for this vertex.
+								const FMatrix44f InvSkinningTransform = CalcInverseSkinningTransform(RenderVertexIndex, SkelMeshLODData, SkinWeightBuffer);
 
-							// Calculate the pre-skinning delta.
-							FVector3f Delta = FVector3f::ZeroVector;
-							if (VertexDeltaSpace == EVertexDeltaSpace::PreSkinning)
-							{
-								if (SkinningMode == EMLDeformerSkinningMode::Linear)
-								{
-									// Calculate the inverse skinning transform for this vertex.
-									const FMatrix44f InvSkinningTransform = CalcInverseSkinningTransform(RenderVertexIndex, SkelMeshLODData, SkinWeightBuffer);
-
-									// Calculate the pre-skinning data.
-									const FSkeletalMeshLODRenderData& LODData = SkeletalMesh->GetResourceForRendering()->LODRenderData[0];
-									const FVector3f UnskinnedPosition = LODData.StaticVertexBuffers.PositionVertexBuffer.VertexPosition(RenderVertexIndex);
-									const FVector3f GeomCacheVertexPos = (FVector3f)AlignmentTransform.TransformPosition((FVector)GeomCacheMeshData.Positions[GeomCacheVertexIndex]);
-									const FVector3f PreSkinningTargetPos = InvSkinningTransform.TransformPosition(GeomCacheVertexPos);
-									Delta = PreSkinningTargetPos - UnskinnedPosition;
-								}
-								else // Dual Quaternion Skinning.
-								{
-									const FVector3f SkinnedVertexPos = SkinnedVertexPositions[SkinnedVertexIndex];
-									const FVector3f GeomCacheVertexPos = (FVector3f)AlignmentTransform.TransformPosition((FVector)GeomCacheMeshData.Positions[GeomCacheVertexIndex]);
-									const FVector3f WorldDelta = GeomCacheVertexPos - SkinnedVertexPos;
-									const FSkeletalMeshLODRenderData& LODData = SkeletalMesh->GetResourceForRendering()->LODRenderData[0];
-									Delta = CalcDualQuaternionDelta(RenderVertexIndex, WorldDelta, LODData, SkinWeightBuffer);
-								}
+								// Calculate the pre-skinning data.
+								const FSkeletalMeshLODRenderData& LODData = SkeletalMesh->GetResourceForRendering()->LODRenderData[0];
+								const FVector3f UnskinnedPosition = LODData.StaticVertexBuffers.PositionVertexBuffer.VertexPosition(RenderVertexIndex);
+								const FVector3f GeomCacheVertexPos = (FVector3f)AlignmentTransform.TransformPosition((FVector)GeomCacheMeshData.Positions[GeomCacheVertexIndex]);
+								const FVector3f PreSkinningTargetPos = InvSkinningTransform.TransformPosition(GeomCacheVertexPos);
+								Delta = PreSkinningTargetPos - UnskinnedPosition;
 							}
-							else if (VertexDeltaSpace == EVertexDeltaSpace::PostSkinning)
+							else // Dual Quaternion Skinning.
 							{
 								const FVector3f SkinnedVertexPos = SkinnedVertexPositions[SkinnedVertexIndex];
 								const FVector3f GeomCacheVertexPos = (FVector3f)AlignmentTransform.TransformPosition((FVector)GeomCacheMeshData.Positions[GeomCacheVertexIndex]);
-								Delta = GeomCacheVertexPos - SkinnedVertexPos;
+								const FVector3f WorldDelta = GeomCacheVertexPos - SkinnedVertexPos;
+								const FSkeletalMeshLODRenderData& LODData = SkeletalMesh->GetResourceForRendering()->LODRenderData[0];
+								Delta = CalcDualQuaternionDelta(RenderVertexIndex, WorldDelta, LODData, SkinWeightBuffer);
 							}
-
-							const int32 ArrayIndex = 3 * SkinnedVertexIndex;
-							VertexDeltas[ArrayIndex + 0] = Delta.X;
-							VertexDeltas[ArrayIndex + 1] = Delta.Y;
-							VertexDeltas[ArrayIndex + 2] = Delta.Z;
 						}
+						else if (VertexDeltaSpace == EVertexDeltaSpace::PostSkinning)
+						{
+							const FVector3f SkinnedVertexPos = SkinnedVertexPositions[SkinnedVertexIndex];
+							const FVector3f GeomCacheVertexPos = (FVector3f)AlignmentTransform.TransformPosition((FVector)GeomCacheMeshData.Positions[GeomCacheVertexIndex]);
+							Delta = GeomCacheVertexPos - SkinnedVertexPos;
+						}
+
+						const int32 ArrayIndex = 3 * SkinnedVertexIndex;
+						VertexDeltas[ArrayIndex + 0] = Delta.X;
+						VertexDeltas[ArrayIndex + 1] = Delta.Y;
+						VertexDeltas[ArrayIndex + 2] = Delta.Z;
 					}
 				});	// ParallelFor
 			}
diff --git a/Engine/Plugins/Animation/MLDeformer/MLDeformerFramework/Source/MLDeformerFrameworkEditor/Private/MLDeformerGeomCacheTrainingModel.cpp b/Engine/Plugins/Animation/MLDeformer/MLDeformerFramework/Source/MLDeformerFrameworkEditor/Private/MLDeformerGeomCacheTrainingModel.cpp
--- a/Engine/Plugins/Animation/MLDeformer/MLDeformerFramework/Source/MLDeformerFrameworkEditor/Private/MLDeformerGeomCacheTrainingModel.cpp
+++ b/Engine/Plugins/Animation/MLDeformer/MLDeformerFramework/Source/MLDeformerFrameworkEditor/Private/MLDeformerGeomCacheTrainingModel.cpp
@@ -4,6 +4,10 @@
 #include "MLDeformerGeomCacheTrainingInputAnim.h"
 #include "MLDeformerGeomCacheEditorModel.h"
 #include "MLDeformerGeomCacheModel.h"
+#include "Misc/FileHelper.h"
+#include "Misc/ScopedSlowTask.h"
+
+#define LOCTEXT_NAMESPACE "MLDeformerGeomCacheTrainingModel"
 
 void UMLDeformerGeomCacheTrainingModel::Init(UE::MLDeformer::FMLDeformerEditorModel* InEditorModel)
 {
@@ -18,12 +22,11 @@ bool UMLDeformerGeomCacheTrainingModel::FindNextAnimToSample(int32& OutNextAnimI
 {
 	using namespace UE::MLDeformer;
 
-	UMLDeformerGeomCacheModel* GeomCacheModel = Cast<UMLDeformerGeomCacheModel>(EditorModel->GetModel());
 	int32 NumTries = 0;
 	int32 AnimIndex = SampleAnimIndex;
 
 	const int32 NumInputAnims = EditorModel->GetNumTrainingInputAnims();
-	while (NumTries < NumInputAnims)	// Try all input animations at worst case.
+	while (NumTries < NumInputAnims) // Try all input animations at worst case.
 	{
 		const FMLDeformerTrainingInputAnim* InputAnim = EditorModel->GetTrainingInputAnim(AnimIndex);
 		if (InputAnim && InputAnim->IsEnabled())
@@ -43,7 +46,7 @@ bool UMLDeformerGeomCacheTrainingModel::FindNextAnimToSample(int32& OutNextAnimI
 	}
 
 	OutNextAnimIndex = INDEX_NONE;
-    return false;
+	return false;
 }
 
 bool UMLDeformerGeomCacheTrainingModel::SampleNextFrame()
@@ -56,7 +59,6 @@ bool UMLDeformerGeomCacheTrainingModel::SampleNextFrame()
 		return false;
 	}
 
-	FMLDeformerGeomCacheEditorModel* GeomEditorModel = static_cast<FMLDeformerGeomCacheEditorModel*>(EditorModel);
 	UMLDeformerGeomCacheModel* GeomCacheModel = Cast<UMLDeformerGeomCacheModel>(EditorModel->GetModel());
 
 	// Get the animation to sample from and validate some things.
@@ -67,7 +69,7 @@ bool UMLDeformerGeomCacheTrainingModel::SampleNextFrame()
 
 	const int32 StartFrame = InputAnim.GetUseCustomRange() ? FMath::Min<int32>(InputAnim.GetStartFrame(), InputAnim.GetEndFrame()) : 0;
 	const int32 CurFrameToSample = StartFrame + NumTimesSampled[SampleAnimIndex];
-	check(CurFrameToSample < StartFrame + InputAnim.GetNumFramesToSample());	// We should never sample more frames than the animation has.
+	check(CurFrameToSample < StartFrame + InputAnim.GetNumFramesToSample()); // We should never sample more frames than the animation has.
 	NumTimesSampled[SampleAnimIndex]++;
 
 	// Perform the actual sampling.
@@ -95,3 +97,111 @@ bool UMLDeformerGeomCacheTrainingModel::SampleNextFrame()
 	// Let the caller know we can still sample more.
 	return true;
 }
+
+namespace
+{
+	FString ConvertSecondsToTimeString(double TotalSeconds)
+	{
+		// Convert total seconds to hours, minutes, and seconds
+		const int32 Hours = FMath::FloorToInt(TotalSeconds / 3600);
+		TotalSeconds -= Hours * 3600;
+		const int32 Minutes = FMath::FloorToInt(TotalSeconds / 60);
+		const int32 Seconds = FMath::FloorToInt(TotalSeconds - (Minutes * 60));
+
+		// Format the string with leading zeros
+		return FString::Printf(TEXT("%02d:%02d:%02d"), Hours, Minutes, Seconds);
+	}
+}
+
+bool UMLDeformerGeomCacheTrainingModel::GenerateBasicInputsAndOutputBuffers(const FString& InputsFilePath, const FString& OutputsFilePath)
+{
+	// Create the task window and show it.
+	const int32 NumFrames = NumSamples();
+	FScopedSlowTask Task(NumFrames + 1, LOCTEXT("SamplingTaskTitle", "Sampling frames"));
+	Task.MakeDialog(true);
+
+	// Create our buffer files that store the sampled inputs and outputs of all frames.
+	IFileManager& FileManager = IFileManager::Get();
+	const TUniquePtr<FArchive> InputsArchive = TUniquePtr<FArchive>(FileManager.CreateFileWriter(*InputsFilePath));
+	const TUniquePtr<FArchive> OutputArchive = TUniquePtr<FArchive>(FileManager.CreateFileWriter(*OutputsFilePath));
+
+	// Make sure we free the data whenever we leave this function.
+	ON_SCOPE_EXIT
+	{
+		InputsArchive->Close();
+		OutputArchive->Close();
+	};
+
+	// Make sure the archives are valid.
+	if (!InputsArchive)
+	{
+		UE_LOG(LogMLDeformer, Error,
+		       TEXT("UMLDeformerGeomCacheTrainingModel::GenerateBasicInputsAndOutputBuffers() - Failed to create inputs buffer file %s"),
+		       *InputsFilePath);
+		return false;
+	}
+	
+	if (!OutputArchive)
+	{
+		UE_LOG(LogMLDeformer, Error,
+		       TEXT("UMLDeformerGeomCacheTrainingModel::GenerateBasicInputsAndOutputBuffers() - Failed to create output buffer file %s"),
+		       *OutputsFilePath);
+		return false;
+	}
+
+	// Sample all the frames.
+	const double StartTime = FPlatformTime::Seconds();
+	for (int32 SampleIndex = 0; SampleIndex < NumFrames; SampleIndex++)
+	{
+		if (Task.ShouldCancel())
+		{
+			return false;
+		}
+
+		// Sample the next frame, which fills the SampleBoneRotations, SampleCurveValues and SampleDeltas arrays.
+		NextSample();
+
+		// Write the sampled inputs.
+		InputsArchive->Serialize(SampleBoneRotations.GetData(), SampleBoneRotations.Num() * sizeof(float));
+		if (InputsArchive->IsError() || InputsArchive->IsCriticalError())
+		{
+			UE_LOG(LogMLDeformer, Error, TEXT("Failed to write sampled input data!"));
+			return false;
+		}
+		InputsArchive->Serialize(SampleCurveValues.GetData(), SampleCurveValues.Num() * sizeof(float));
+		if (InputsArchive->IsError() || InputsArchive->IsCriticalError())
+		{
+			UE_LOG(LogMLDeformer, Error, TEXT("Failed to write sampled input data!"));
+			return false;
+		}
+
+		// Write the sampled outputs.
+		OutputArchive->Serialize(SampleDeltas.GetData(), SampleDeltas.Num() * sizeof(float));
+		if (OutputArchive->IsError() || OutputArchive->IsCriticalError())
+		{
+			UE_LOG(LogMLDeformer, Error, TEXT("Failed to write sampled output data!"));
+			return false;
+		}
+
+		// Calculate remaining time.
+		const double ElapsedTime = FPlatformTime::Seconds() - StartTime;
+		const int32 SamplesRemaining = NumFrames - SampleIndex;
+		const double AverageTimePerSample = ElapsedTime / (SampleIndex + 1);
+		const double RemainingSeconds = SamplesRemaining * AverageTimePerSample;
+		const FString RemainingTimeString = ConvertSecondsToTimeString(RemainingSeconds);
+		const FString ElapsedTimeString = ConvertSecondsToTimeString(ElapsedTime);
+
+		// Update the progress in the task window.
+		const FText ProgressText = FText::Format(LOCTEXT("SamplingProgressText", "Sampling frame {0} of {1} - Elapsed: {2} - Remaining: {3}"),
+		                                         SampleIndex + 1,
+		                                         NumFrames,
+		                                         FText::FromString(ElapsedTimeString),
+		                                         FText::FromString(RemainingTimeString));
+		Task.EnterProgressFrame(1.0f, ProgressText);
+	}
+
+	UE_LOG(LogMLDeformer, Display, TEXT("Sampling finished in %.0f seconds"), FPlatformTime::Seconds() - StartTime);
+	return true;
+}
+
+#undef LOCTEXT_NAMESPACE
diff --git a/Engine/Plugins/Animation/MLDeformer/MLDeformerFramework/Source/MLDeformerFrameworkEditor/Private/MLDeformerMorphModelEditorModel.cpp b/Engine/Plugins/Animation/MLDeformer/MLDeformerFramework/Source/MLDeformerFrameworkEditor/Private/MLDeformerMorphModelEditorModel.cpp
--- a/Engine/Plugins/Animation/MLDeformer/MLDeformerFramework/Source/MLDeformerFrameworkEditor/Private/MLDeformerMorphModelEditorModel.cpp
+++ b/Engine/Plugins/Animation/MLDeformer/MLDeformerFramework/Source/MLDeformerFrameworkEditor/Private/MLDeformerMorphModelEditorModel.cpp
@@ -2,18 +2,14 @@
 
 #include "MLDeformerMorphModelEditorModel.h"
 #include "MLDeformerMorphModel.h"
-#include "MLDeformerMorphModelInputInfo.h"
 #include "MLDeformerMorphModelVizSettings.h"
 #include "MLDeformerMorphModelInstance.h"
-#include "MLDeformerGeomCacheActor.h"
 #include "MLDeformerEditorToolkit.h"
 #include "MLDeformerEditorStyle.h"
 #include "MLDeformerInputInfo.h"
 #include "MLDeformerComponent.h"
 #include "MLDeformerGeomCacheHelpers.h"
 #include "Components/ExternalMorphSet.h"
-#include "GeometryCache.h"
-#include "GeometryCacheComponent.h"
 #include "Animation/DebugSkelMeshComponent.h"
 #include "Animation/MorphTarget.h"
 #include "Rendering/SkeletalMeshModel.h"
@@ -23,6 +19,7 @@
 #include "Misc/ScopedSlowTask.h"
 #include "BoneWeights.h"
 #include "MeshAttributeArray.h"
+#include "MeshAttributes.h"
 #include "Async/ParallelFor.h"
 
 #define LOCTEXT_NAMESPACE "MLDeformerMorphModelEditorModel"
@@ -172,7 +169,7 @@ namespace UE::MLDeformer
 
 			// First calculate the mean.
 			float Mean = 0.0f;
-			for (float Value : Values)
+			for (const float Value : Values)
 			{
 				Mean += Value;
 			}
@@ -180,7 +177,7 @@ namespace UE::MLDeformer
 
 			// Now calculate the standard deviation.
 			float Sum = 0.0f;
-			for (float Value : Values)
+			for (const float Value : Values)
 			{
 				Sum += FMath::Square(Value - Mean);
 			}
@@ -323,11 +320,9 @@ namespace UE::MLDeformer
 		FSkeletalMeshRenderData* RenderData = SkelMesh->GetResourceForRendering();
 		check(RenderData);
 		check(!RenderData->LODRenderData.IsEmpty());
-		const int32 NumRenderVertices = RenderData->LODRenderData[LOD].GetNumVertices();
 		const FColorVertexBuffer& ColorBuffer = RenderData->LODRenderData[LOD].StaticVertexBuffers.ColorVertexBuffer;
-
 		const TArray<int32>& VertexMap = Model->GetVertexMap();
-		TArray<float> GlobalMaskWeights = CalcGlobalMaskWeights(VertexMap, ColorBuffer, EMLDeformerMaskChannel::Disabled, false);
+		const TArray<float> GlobalMaskWeights = CalcGlobalMaskWeights(VertexMap, ColorBuffer, EMLDeformerMaskChannel::Disabled, false);
 
 		CalcMorphTargetNormals(
 			LOD,
@@ -552,8 +547,8 @@ namespace UE::MLDeformer
 		EMLDeformerMaskChannel MaskChannel,
 		bool bInvertMaskChannel)
 	{
-		FMLDeformerSampler* Sampler = GetSamplerForActiveAnim();
 		// In case when ActiveTrainingInputAnimIndex == INDEX_NONE but we have a sampler (e.g. testing), use that.
+		FMLDeformerSampler* Sampler = GetSamplerForActiveAnim();
 		if (Sampler == nullptr)
 		{
 			Sampler = GetSamplerForTrainingAnim(0);
@@ -580,8 +575,8 @@ namespace UE::MLDeformer
 		check((Deltas.Num() / NumMorphTargets) == NumBaseMeshVerts);
 		check(!Model->GetVertexMap().IsEmpty());
 
-		FScopedSlowTask Task(NumMorphTargets, LOCTEXT("CreateMorphTargetProgress", "Creating morph targets"));
-		Task.MakeDialog(false);
+		FScopedSlowTask Task(1.0f, LOCTEXT("CreateMorphTargetProgress", "Creating morph targets"));
+		Task.MakeDialogDelayed(2.0f, false);
 
 		USkeletalMesh* SkelMesh = Model->GetSkeletalMesh();
 		FSkeletalMeshRenderData* RenderData = SkelMesh->GetResourceForRendering();
@@ -599,21 +594,9 @@ namespace UE::MLDeformer
 		// Precalc an array that maps imported vertices to a render vertex.
 		TArray<int32> ImportedVertexToRenderVertexMapping;
 		ImportedVertexToRenderVertexMapping.SetNumUninitialized(NumBaseMeshVerts);
-		const int32 BatchSize = 500;
-		const int32 NumBatches = (NumBaseMeshVerts / BatchSize) + 1;
-		ParallelFor(NumBatches, [&](int32 BatchIndex)
+		ParallelFor(TEXT("CreateMorphTargets_BuildMapping"), NumBaseMeshVerts, 5, [&](int32 VertexIndex)
 		{
-			const int32 StartVertex = BatchIndex * BatchSize;
-			if (StartVertex >= NumBaseMeshVerts)
-			{
-				return;
-			}
-
-			const int32 NumVertsInBatch = (StartVertex + BatchSize) < NumBaseMeshVerts ? BatchSize : FMath::Max(NumBaseMeshVerts - StartVertex, 0);
-			for (int32 VertexIndex = StartVertex; VertexIndex < StartVertex + NumVertsInBatch; ++VertexIndex)
-			{
-				ImportedVertexToRenderVertexMapping[VertexIndex] = VertexMap.Find(VertexIndex);
-			}
+			ImportedVertexToRenderVertexMapping[VertexIndex] = VertexMap.Find(VertexIndex);
 		});
 
 		TArray<FVector3f> BaseNormals;
@@ -626,21 +609,26 @@ namespace UE::MLDeformer
 		const TArray<float> GlobalMaskWeights = CalcGlobalMaskWeights(VertexMap, ColorBuffer, MaskChannel, bInvertMaskChannel);
 
 		// Initialize an engine morph target for each model morph target.
-		bool bHasOnlyEmptyMorphs = true;
 		UE_LOG(LogMLDeformer, Display, TEXT("Initializing %d engine morph targets of %d vertices each"), NumMorphTargets, Deltas.Num() / NumMorphTargets);
-		TArray<FVector3f> DeltaNormals;
+
 		for (int32 MorphTargetIndex = 0; MorphTargetIndex < NumMorphTargets; ++MorphTargetIndex)
 		{
+			const FName MorphName = *FString::Printf(TEXT("%s%.3d"), *NamePrefix, MorphTargetIndex);
+			UMorphTarget* MorphTarget = NewObject<UMorphTarget>(SkelMesh, MorphName);
+			MorphTarget->BaseSkelMesh = SkelMesh;
+			OutMorphTargets.Add(MorphTarget);
+		}
+		
+		ParallelFor(NumMorphTargets, [&](int32 MorphTargetIndex)
+		{
+			TArray<FVector3f> DeltaNormals;
 			if (bIncludeNormals)
 			{
 				CalcMorphTargetNormals(LOD, SkelMesh, MorphTargetIndex, Deltas, BaseVertexPositions, BaseNormals, ImportedVertexToRenderVertexMapping, 
 					ColorBuffer, MaskChannel, bInvertMaskChannel, GlobalMaskWeights, DeltaNormals);
 			}
 
-			const FName MorphName = *FString::Printf(TEXT("%s%.3d"), *NamePrefix, MorphTargetIndex);
-			UMorphTarget* MorphTarget = NewObject<UMorphTarget>(SkelMesh, MorphName);
-			MorphTarget->BaseSkelMesh = SkelMesh;
-			OutMorphTargets.Add(MorphTarget);
+			UMorphTarget* MorphTarget = OutMorphTargets[MorphTargetIndex];
 
 			// Create a new LOD model for this morph.
 			TArray<FMorphTargetLODModel>& MorphLODs = MorphTarget->GetMorphLODModels();
@@ -658,6 +646,7 @@ namespace UE::MLDeformer
 
 			// Init deltas for this morph target.
 			MorphLODModel.Vertices.Reserve(NumRenderVertices);
+			TSet<int32> SectionIndices;
 			for (int32 VertexIndex = 0; VertexIndex < NumRenderVertices; ++VertexIndex)
 			{
 				const int32 ImportedVertexNumber = VertexMap[VertexIndex];
@@ -683,24 +672,39 @@ namespace UE::MLDeformer
 						RenderData->LODRenderData[0].GetSectionFromVertexIndex(VertexIndex, RenderSection, TempVertexIndex);
 						if (RenderSection != INDEX_NONE)
 						{
-							MorphLODModel.SectionIndices.AddUnique(RenderSection);
+							SectionIndices.Add(RenderSection);
 						}
 					}
 				}
-			}	// for all morph targets
+			}	// for all render vertices
+			
+			// Add all unique section indices.
+			for (int32 SectionIndex : SectionIndices)
+			{
+				MorphLODModel.SectionIndices.Add(SectionIndex);
+			}
 
 			MorphLODModel.Vertices.Shrink();
-			bHasOnlyEmptyMorphs &= MorphLODModel.Vertices.IsEmpty();
-			Task.EnterProgressFrame();
-		} // For all morph targets.
+		}); // For all morph targets.
+
+		bool bHasOnlyEmptyMorphs = true;
+		for (int32 MorphTargetIndex = 0; MorphTargetIndex < NumMorphTargets; ++MorphTargetIndex)
+		{
+			if (!OutMorphTargets[MorphTargetIndex]->GetMorphLODModels()[0].Vertices.IsEmpty())
+			{
+				bHasOnlyEmptyMorphs = false;
+				break;
+			}
+		}
 
 		GetMorphModel()->SetHasOnlyEmptyMorphs(bHasOnlyEmptyMorphs);
+		Task.TickProgress();
 	}
 
 	void FMLDeformerMorphModelEditorModel::CompressMorphTargets(FMorphTargetVertexInfoBuffers& OutMorphBuffers, const TArray<UMorphTarget*>& MorphTargets, int32 LOD, float MorphErrorTolerance)
 	{
 		FScopedSlowTask Task(1, LOCTEXT("CompressMorphTargetsProgress", "Compressing morph targets"));
-		Task.MakeDialog(false);	
+		Task.MakeDialogDelayed(2.0f, false);	
 
 		USkeletalMesh* SkelMesh = Model->GetSkeletalMesh();
 		FSkeletalMeshRenderData* RenderData = SkelMesh->GetResourceForRendering();
@@ -714,7 +718,7 @@ namespace UE::MLDeformer
 			ReleaseResourceAndFlush(&OutMorphBuffers);
 		}
 
-		// Don't empty the array of morph target data when we init the RHI buffers, as we need them to serialize later on.
+		// Don't empty the array of morph target data when we initialize the RHI buffers, as we need them to serialize later on.
 		OutMorphBuffers = FMorphTargetVertexInfoBuffers();
 		OutMorphBuffers.SetEmptyMorphCPUDataOnInitRHI(false);
 
@@ -767,10 +771,9 @@ namespace UE::MLDeformer
 				const int32 DeltaArrayOffset = NumVerts * FinalMorphTargetIndex + VertexIndex;
 				const FVector3f RawDelta(MorphDeltas[DeltaArrayOffset]);
 				const float MorphMaskValue = MorphMask.IsEmpty() ? 1.0f : MorphMask[VertexIndex];
-				const float VertexWeight = 1.0f;
 				FVector3f ScaledDelta;
 				FVector3f ScaledDeltaNormal;
-				if (ProcessVertexDelta(ScaledDelta, ScaledDeltaNormal, RawDelta, FVector3f::ZeroVector, DeltaThreshold, MorphMaskValue, VertexWeight))
+				if (ProcessVertexDelta(ScaledDelta, ScaledDeltaNormal, RawDelta, FVector3f::ZeroVector, DeltaThreshold, MorphMaskValue, 1.0f))
 				{
 					PDI->DrawPoint(StartPoint, IncludedColor, 1.0f, 0);
 					PDI->DrawLine(StartPoint, StartPoint + FVector(ScaledDelta), IncludedColor, 0);
@@ -791,16 +794,15 @@ namespace UE::MLDeformer
 			return;
 		}
 
-		// Zero out small deltas.
 		TArray<FVector3f> MorphTargetDeltas = Deltas;
 		ZeroDeltasByLengthThreshold(MorphTargetDeltas, MorphModel->GetMorphDeltaZeroThreshold());
 
 		// Turn the delta buffer into a set of engine morph targets.
-		const int32 LODZero = 0;
+		constexpr int32 LODZero = 0;
 		const bool bIncludeNormals = MorphModel->GetIncludeMorphTargetNormals();
 		const EMLDeformerMaskChannel MaskChannel = MorphModel->GetMaskChannel();
 		const bool bInvertMaskChannel = MorphModel->GetInvertMaskChannel();
-
+	
 		// Create the engine morph targets.
 		TArray<UMorphTarget*> MorphTargets;
 		CreateMorphTargets(
@@ -818,7 +820,7 @@ namespace UE::MLDeformer
 
 		// Transfer morphs to the LOD levels.
 		TransferMorphTargets(MorphTargets);
-
+	
 		// Resize to the new desired size.
 		MorphModel->ClearMorphTargetSets();
 		const int32 NumLODs = !MorphTargets.IsEmpty() ? MorphTargets[0]->GetMorphLODModels().Num() : 0;
@@ -871,7 +873,7 @@ namespace UE::MLDeformer
 		const int32 NumLODs = FMath::Min(SkelMesh->GetLODNum(), MaxLodWithMorphs);
 		const int32 NumTaskItems = 1 + (NumLODs - 1) * MorphTargetsLODZero.Num();	// +1 because we see preparing the lookup tables as one task item as well.
 		FScopedSlowTask Task(NumTaskItems, LOCTEXT("TransferMorphTargetProgress", "Generating morph target LODs"));
-		Task.MakeDialog(false);	
+		Task.MakeDialogDelayed(2.0f, false);	
 
 		const double StartTime = FPlatformTime::Seconds();
 
@@ -880,8 +882,6 @@ namespace UE::MLDeformer
 
 		// Build a mapping table to eliminate linear searches.
 		// For every morph target build a map that maps a render vertex to a morph target vertex.
-		const FSkeletalMeshLODModel& LODModelZero = SkelMesh->GetImportedModel()->LODModels[0];
-		const int32 NumRenderVerticesLODZero = LODModelZero.NumVertices;
 		TArray<TMap<int32, int32>> RenderVertexToMorphVertexLodZero;
 		RenderVertexToMorphVertexLodZero.SetNum(MorphTargetsLODZero.Num());
 		for (int32 MorphIndex = 0; MorphIndex < MorphTargetsLODZero.Num(); ++MorphIndex)
@@ -1095,12 +1095,12 @@ namespace UE::MLDeformer
 		FSkinWeightVertexBuffer& SkinWeightBuffer = *SkeletalMeshComponent->GetSkinWeightBuffer(LOD);
 		const FSkeletalMeshLODRenderData& LODData = SkeletalMesh->GetResourceForRendering()->LODRenderData[LOD];
 		const int32 NumRenderVerts = LODData.GetNumVertices();
-		for (int32 VertexIndex = 0; VertexIndex < NumRenderVerts; ++VertexIndex)
+		ParallelFor(TEXT("ApplyBoneToMask"), NumRenderVerts, 5, [&](int32 VertexIndex)
 		{
 			const int32 ImportedVertexNumber = ImportedVertexNumbers[VertexIndex];
 			if (ImportedVertexNumber == INDEX_NONE)
 			{
-				continue;
+				return;
 			}
 
 			// Find the render section, which we need to find the right bone index.
@@ -1124,7 +1124,7 @@ namespace UE::MLDeformer
 					}
 				}
 			}
-		}		
+		});		
 	}
 
 	void FMLDeformerMorphModelEditorModel::AddRequiredBones(const FReferenceSkeleton& RefSkel, int32 SkeletonBoneIndex, const TArray<int32>& VirtualParentTable, TArray<int32>& OutBonesAdded)
diff --git a/Engine/Plugins/Animation/MLDeformer/MLDeformerFramework/Source/MLDeformerFrameworkEditor/Private/MLDeformerSampler.cpp b/Engine/Plugins/Animation/MLDeformer/MLDeformerFramework/Source/MLDeformerFrameworkEditor/Private/MLDeformerSampler.cpp
--- a/Engine/Plugins/Animation/MLDeformer/MLDeformerFramework/Source/MLDeformerFrameworkEditor/Private/MLDeformerSampler.cpp
+++ b/Engine/Plugins/Animation/MLDeformer/MLDeformerFramework/Source/MLDeformerFrameworkEditor/Private/MLDeformerSampler.cpp
@@ -88,7 +88,7 @@ namespace UE::MLDeformer
 
 	void FMLDeformerSampler::UpdateSkinnedPositions()
 	{
-		const int32 LODIndex = 0;
+		constexpr int32 LODIndex = 0;
 		ExtractSkinnedPositions(LODIndex, BoneMatrices, TempVertexPositions, SkinnedVertexPositions);
 	}
 
@@ -142,9 +142,9 @@ namespace UE::MLDeformer
 	// Calculate the inverse skinning transform. This is basically inv(sum(BoneTransform_i * inv(BoneRestTransform_i) * Weight_i)), where i is for each skinning influence for the given vertex.
 	FMatrix44f FMLDeformerSampler::CalcInverseSkinningTransform(int32 VertexIndex, const FSkeletalMeshLODRenderData& SkelMeshLODData, const FSkinWeightVertexBuffer& SkinWeightBuffer) const
 	{
-		check(SkeletalMeshComponent);
+		checkSlow(SkeletalMeshComponent);
 		const USkeletalMesh* Mesh = SkeletalMeshComponent->GetSkeletalMeshAsset();
-		check(Mesh);
+		checkSlow(Mesh);
 
 		// Find the render section, which we need to find the right bone index.
 		int32 SectionIndex = INDEX_NONE;
diff --git a/Engine/Plugins/Animation/MLDeformer/MLDeformerFramework/Source/MLDeformerFrameworkEditor/Private/MLDeformerTrainingModel.cpp b/Engine/Plugins/Animation/MLDeformer/MLDeformerFramework/Source/MLDeformerFrameworkEditor/Private/MLDeformerTrainingModel.cpp
--- a/Engine/Plugins/Animation/MLDeformer/MLDeformerFramework/Source/MLDeformerFrameworkEditor/Private/MLDeformerTrainingModel.cpp
+++ b/Engine/Plugins/Animation/MLDeformer/MLDeformerFramework/Source/MLDeformerFrameworkEditor/Private/MLDeformerTrainingModel.cpp
@@ -8,6 +8,7 @@
 #include "MLDeformerTrainingInputAnim.h"
 #include "SkeletalMeshAttributes.h"
 #include "Engine/SkeletalMesh.h"
+#include "Misc/FileHelper.h"
 
 FName UMLDeformerTrainingModel::DefaultMaskName = FName("MLD_DefaultMask");
 
@@ -106,10 +107,10 @@ TArray<float> UMLDeformerTrainingModel::GetTrainingInputAnimMaskData(FName MaskN
 	TArray<float> PerVertexValues;
 
 	// Get the imported vertex numbers.
-	const int32 LodIndex = 0;
-	USkeletalMesh* SkeletalMesh = EditorModel->GetModel()->GetSkeletalMesh();
+	constexpr int32 LodIndex = 0;
+	const USkeletalMesh* SkeletalMesh = EditorModel->GetModel()->GetSkeletalMesh();
 	FMeshDescription* MeshDescription = SkeletalMesh->GetMeshDescription(LodIndex);
-	TVertexAttributesConstRef<int32> ImportPointIndex = MeshDescription ? MeshDescription->VertexAttributes().GetAttributesRef<int32>(MeshAttribute::Vertex::ImportPointIndex) : TVertexAttributesConstRef<int32>(); 
+	const TVertexAttributesConstRef<int32> ImportPointIndex = MeshDescription ? MeshDescription->VertexAttributes().GetAttributesRef<int32>(MeshAttribute::Vertex::ImportPointIndex) : TVertexAttributesConstRef<int32>(); 
 
 	TVertexAttributesConstRef<float> MaskData = EditorModel->FindVertexAttributes(MaskName);
 	if (ImportPointIndex.IsValid() && MaskData.IsValid())
diff --git a/Engine/Plugins/Animation/MLDeformer/MLDeformerFramework/Source/MLDeformerFrameworkEditor/Private/SMLDeformerTimeline.cpp b/Engine/Plugins/Animation/MLDeformer/MLDeformerFramework/Source/MLDeformerFrameworkEditor/Private/SMLDeformerTimeline.cpp
--- a/Engine/Plugins/Animation/MLDeformer/MLDeformerFramework/Source/MLDeformerFrameworkEditor/Private/SMLDeformerTimeline.cpp
+++ b/Engine/Plugins/Animation/MLDeformer/MLDeformerFramework/Source/MLDeformerFrameworkEditor/Private/SMLDeformerTimeline.cpp
@@ -2005,7 +2005,7 @@ namespace UE::MLDeformer
 				{
 					EditorModel->GetEditor()->UpdateTimeSliderRange();					
 				}
-				EditorModel->TriggerInputAssetChanged();
+				EditorModel->RequestTriggerInputAssetsChangedNextTick();
 			}
 		}
 	}
diff --git a/Engine/Plugins/Animation/MLDeformer/MLDeformerFramework/Source/MLDeformerFrameworkEditor/Public/MLDeformerEditorModel.h b/Engine/Plugins/Animation/MLDeformer/MLDeformerFramework/Source/MLDeformerFrameworkEditor/Public/MLDeformerEditorModel.h
--- a/Engine/Plugins/Animation/MLDeformer/MLDeformerFramework/Source/MLDeformerFrameworkEditor/Public/MLDeformerEditorModel.h
+++ b/Engine/Plugins/Animation/MLDeformer/MLDeformerFramework/Source/MLDeformerFrameworkEditor/Public/MLDeformerEditorModel.h
@@ -1067,6 +1067,11 @@ namespace UE::MLDeformer
 		/** Update the vertex attribute names array as returned by GetVertexAttributeNames(). This is some array that can be used for UI elements like combo boxes. */
 		void UpdateVertexAttributeNames();
 
+		/**
+		 * Request a call to TriggerInputAssetsChanged on the next tick, so not directly.
+		 * This can be useful to prevent multiple nested calls to TriggerInputAssetsChanged.
+		 */
+		void RequestTriggerInputAssetsChangedNextTick()						{ bNeedsAssetReinit = true; }
 
 	protected:
 		virtual void CreateSamplers();
diff --git a/Engine/Plugins/Animation/MLDeformer/MLDeformerFramework/Source/MLDeformerFrameworkEditor/Public/MLDeformerGeomCacheTrainingModel.h b/Engine/Plugins/Animation/MLDeformer/MLDeformerFramework/Source/MLDeformerFrameworkEditor/Public/MLDeformerGeomCacheTrainingModel.h
--- a/Engine/Plugins/Animation/MLDeformer/MLDeformerFramework/Source/MLDeformerFrameworkEditor/Public/MLDeformerGeomCacheTrainingModel.h
+++ b/Engine/Plugins/Animation/MLDeformer/MLDeformerFramework/Source/MLDeformerFrameworkEditor/Public/MLDeformerGeomCacheTrainingModel.h
@@ -23,6 +23,12 @@ public:
 	virtual bool SampleNextFrame() override;
 	// ~END UMLDeformerTrainingModel overrides;
 
+	/**
+	 * Sample all input frames.
+	 */
+	UFUNCTION(BlueprintCallable, Category = "ML Deformer Training")
+	bool GenerateBasicInputsAndOutputBuffers(const FString& InputsFilePath, const FString& OutputsFilePath);
+
 protected:
 	/**
 	 * Find the next input animation to sample from.
@@ -32,4 +38,4 @@ protected:
 	 * @return Returns true when we found our next animation to sample. Returns false when we already sampled everything.
 	 */
 	virtual bool FindNextAnimToSample(int32& OutNextAnimIndex) const override;
-};
\ No newline at end of file
+};
diff --git a/Engine/Plugins/Animation/MLDeformer/MLDeformerFramework/Source/MLDeformerFrameworkEditor/Public/MLDeformerMorphModelEditorModel.h b/Engine/Plugins/Animation/MLDeformer/MLDeformerFramework/Source/MLDeformerFrameworkEditor/Public/MLDeformerMorphModelEditorModel.h
--- a/Engine/Plugins/Animation/MLDeformer/MLDeformerFramework/Source/MLDeformerFrameworkEditor/Public/MLDeformerMorphModelEditorModel.h
+++ b/Engine/Plugins/Animation/MLDeformer/MLDeformerFramework/Source/MLDeformerFrameworkEditor/Public/MLDeformerMorphModelEditorModel.h
@@ -334,6 +334,11 @@ namespace UE::MLDeformer
 		 */
 		TArray<float> CalcGlobalMaskWeights(const TArray<int32>& VertexMap, const FColorVertexBuffer& ColorBuffer, EMLDeformerMaskChannel MaskChannel, bool bInvertMaskChannel) const;
 
+		/**
+		 * Tell the editor model that we should update the mask buffer on the next Tick call.
+		 */
+		void SetNeedsRebuildMaskBuffer() { bNeedsRebuildMaskInfo = true; }
+
 	protected:
 		/**
 		 * The entire set of morph target deltas, 3 per vertex, for each morph target, as one flattened buffer.
@@ -346,5 +351,8 @@ namespace UE::MLDeformer
 
 		/** The backup of the input item mask buffer. This contains all the bone and bone group masks. */
 		TArray<float> InputItemMaskBufferBackup;
+
+		/** Indicates whether we need to rebuild the mask buffer next Tick or not. */
+		bool bNeedsRebuildMaskInfo = false;
 	};
 }	// namespace UE::MLDeformer
diff --git a/Engine/Plugins/Animation/MLDeformer/NeuralMorphModel/Source/NeuralMorphModelEditor/Private/NeuralMorphEditorModel.cpp b/Engine/Plugins/Animation/MLDeformer/NeuralMorphModel/Source/NeuralMorphModelEditor/Private/NeuralMorphEditorModel.cpp
--- a/Engine/Plugins/Animation/MLDeformer/NeuralMorphModel/Source/NeuralMorphModelEditor/Private/NeuralMorphEditorModel.cpp
+++ b/Engine/Plugins/Animation/MLDeformer/NeuralMorphModel/Source/NeuralMorphModelEditor/Private/NeuralMorphEditorModel.cpp
@@ -22,7 +22,6 @@
 #include "Rendering/SkeletalMeshModel.h"
 #include "Rendering/SkeletalMeshLODModel.h"
 #include "ReferenceSkeleton.h"
-#include "Widgets/SBoxPanel.h"
 #include "Misc/TransactionObjectEvent.h"
 #include "SNeuralMorphBoneGroupsWidget.h"
 #include "SNeuralMorphCurveGroupsWidget.h"
@@ -44,6 +43,8 @@ namespace UE::NeuralMorphModel
 		UNeuralMorphInputInfo* NeuralInputInfo = Cast<UNeuralMorphInputInfo>(GetEditorInputInfo());
 		if (NeuralInputInfo)
 		{
+			RemoveNonExistingMaskInfos();
+			UpdateEditorInputInfo();
 			BuildMaskBuffer(NeuralInputInfo->GetInputItemMaskBuffer());
 		}
 	}
@@ -132,16 +133,7 @@ namespace UE::NeuralMorphModel
 	void FNeuralMorphEditorModel::Init(const InitSettings& Settings)
 	{
 		FMLDeformerMorphModelEditorModel::Init(Settings);
-
-		UNeuralMorphModel* NeuralModel = GetNeuralMorphModel();
-
-		check(NeuralModel->SupportsGlobalModeOnly() ? NeuralModel->Mode == ENeuralMorphMode::Global : true);
-
 		RemoveNonExistingMaskInfos();
-		if (NeuralModel->GetModelMode() == ENeuralMorphMode::Local)
-		{
-			RebuildEditorMaskInfo();
-		}
 	}
 
 	ETrainingResult FNeuralMorphEditorModel::Train()
@@ -197,7 +189,6 @@ namespace UE::NeuralMorphModel
 			return;
 		}
 
-		const USkeleton* Skeleton = Model->GetSkeletalMesh() ? Model->GetSkeletalMesh()->GetSkeleton() : nullptr;
 		const USkeletalMesh* SkeletalMesh = Model->GetSkeletalMesh();
 
 		// Handle bones.
@@ -245,7 +236,6 @@ namespace UE::NeuralMorphModel
 		if (SkeletalMesh)
 		{
 			const TArray<FNeuralMorphCurveGroup>& ModelCurveGroups = NeuralMorphModel->GetCurveGroups();
-			const FReferenceSkeleton& RefSkeleton = SkeletalMesh->GetRefSkeleton();
 			for (int32 CurveGroupIndex = 0; CurveGroupIndex < ModelCurveGroups.Num(); ++CurveGroupIndex)
 			{				
 				const FNeuralMorphCurveGroup& CurveGroup = ModelCurveGroups[CurveGroupIndex];
@@ -281,7 +271,7 @@ namespace UE::NeuralMorphModel
 		}
 
 		RemoveNonExistingMaskInfos();
-		BuildMaskBuffer(NeuralInputInfo->GetInputItemMaskBuffer());
+		//BuildMaskBuffer(NeuralInputInfo->GetInputItemMaskBuffer());
 	}
 
 	void FNeuralMorphEditorModel::UpdateIsReadyForTrainingState()
@@ -342,12 +332,11 @@ namespace UE::NeuralMorphModel
 		{
 			return;
 		}
-
+		
 		const int32 NumBaseMeshVerts = Model->GetNumBaseMeshVerts();
 
 		USkeletalMesh* SkeletalMesh = Model->GetSkeletalMesh();	
 		UNeuralMorphInputInfo* NeuralMorphInputInfo = Cast<UNeuralMorphInputInfo>(GetEditorInputInfo());
-		const FReferenceSkeleton& RefSkel = SkeletalMesh->GetRefSkeleton();
 		const TArray<FName>& BoneNames = NeuralMorphInputInfo->GetBoneNames();
 
 		// Calculate the number of flo =ats we need to store all the masks.
@@ -356,7 +345,7 @@ namespace UE::NeuralMorphModel
 		const int32 NumBoneGroups = NeuralMorphInputInfo->GetBoneGroups().Num();
 		const int32 NumCurveGroups = NeuralMorphInputInfo->GetCurveGroups().Num();
 		const int32 NumFloats = (NumBones + NumCurves + NumBoneGroups + NumCurveGroups) * NumBaseMeshVerts;
-		const int32 HierarchyDepth = 1;	// Default hierarchy depth in case we have no mask for specific bones yet.
+		constexpr int32 HierarchyDepth = 1;	// Default hierarchy depth in case we have no mask for specific bones yet.
 
 		// Init the mask to all zeros.
 		OutMaskBuffer.SetNumZeroed(NumFloats);
@@ -484,7 +473,7 @@ namespace UE::NeuralMorphModel
 	{
 		check(UNeuralMorphEditorProjectSettings::Get());
 
-		// Get the twist substring from the per project configuration.
+		// Get the twist substring from the per-project configuration.
 		// If we have an empty string, just skip the whole twist bone handling.
 		const FString& TwistSubString = UNeuralMorphEditorProjectSettings::Get()->TwistBoneFilter;
 		if (TwistSubString.IsEmpty())
@@ -518,7 +507,6 @@ namespace UE::NeuralMorphModel
 		UNeuralMorphInputInfo* NeuralMorphInputInfo = Cast<UNeuralMorphInputInfo>(GetEditorInputInfo());
 		const FReferenceSkeleton& RefSkel = SkeletalMesh->GetRefSkeleton();
 		const TArray<FName>& BoneNames = NeuralMorphInputInfo->GetBoneNames();
-		const int32 NumBones = NeuralMorphInputInfo->GetNumBones();
 
 		// Make sure we have a valid bone name.
 		const FName BoneName = BoneNames[InputInfoBoneIndex];
@@ -559,9 +547,7 @@ namespace UE::NeuralMorphModel
 	{
 		UNeuralMorphModel* NeuralMorphModel = GetNeuralMorphModel();
 		USkeletalMesh* SkeletalMesh = Model->GetSkeletalMesh();	
-		UNeuralMorphInputInfo* NeuralMorphInputInfo = Cast<UNeuralMorphInputInfo>(GetEditorInputInfo());
 		const FReferenceSkeleton& RefSkel = SkeletalMesh->GetRefSkeleton();
-		const TArray<FName>& BoneNames = NeuralMorphInputInfo->GetBoneNames();
 
 		// For all bones.
 		TArray<int32> BonesAdded;
@@ -610,7 +596,6 @@ namespace UE::NeuralMorphModel
 		ResetBoneMaskInfos();
 
 		// Generate a bone mask info for each bone.
-		UNeuralMorphModel* NeuralMorphModel = GetNeuralMorphModel();
 		UNeuralMorphInputInfo* NeuralMorphInputInfo = Cast<UNeuralMorphInputInfo>(GetEditorInputInfo());
 		const int32 NumBones = NeuralMorphInputInfo->GetNumBones();
 		for (int32 Index = 0; Index < NumBones; ++Index)
@@ -626,7 +611,6 @@ namespace UE::NeuralMorphModel
 		ResetBoneGroupMaskInfos();
 
 		// Generate a mask info for each group.
-		UNeuralMorphModel* NeuralMorphModel = GetNeuralMorphModel();
 		UNeuralMorphInputInfo* NeuralMorphInputInfo = Cast<UNeuralMorphInputInfo>(GetEditorInputInfo());
 		const int32 NumBoneGroups = NeuralMorphInputInfo->GetBoneGroups().Num();	
 		for (int32 Index = 0; Index < NumBoneGroups; ++Index)
@@ -703,7 +687,6 @@ namespace UE::NeuralMorphModel
 			return;
 		}
 
-		UMLDeformerMorphModelVizSettings* VizSettings = Cast<UMLDeformerMorphModelVizSettings>(Model->GetVizSettings());
 		const int32 NumVerts = Model->GetNumBaseMeshVerts();
 		const TArray<FVector3f>& UnskinnedPositions = Sampler->GetUnskinnedVertexPositions();
 		if (MaskVizItemIndex == INDEX_NONE ||
@@ -713,7 +696,6 @@ namespace UE::NeuralMorphModel
 			return;
 		}
 
-		UMLDeformerMorphModel* MorphModel = GetMorphModel();
 		UNeuralMorphInputInfo* InputInfo = Cast<UNeuralMorphInputInfo>(GetEditorInputInfo());
 		const int32 NumMaskItems = InputInfo->GetBoneNames().Num() + InputInfo->GetCurveNames().Num() + InputInfo->GetBoneGroups().Num() + InputInfo->GetCurveGroups().Num();
 		const int32 FinalMaskItemIndex = FMath::Clamp<int32>(MaskItemIndex, 0, NumMaskItems - 1);
diff --git a/Engine/Plugins/Animation/MLDeformer/NeuralMorphModel/Source/NeuralMorphModelEditor/Public/NeuralMorphEditorModel.h b/Engine/Plugins/Animation/MLDeformer/NeuralMorphModel/Source/NeuralMorphModelEditor/Public/NeuralMorphEditorModel.h
--- a/Engine/Plugins/Animation/MLDeformer/NeuralMorphModel/Source/NeuralMorphModelEditor/Public/NeuralMorphEditorModel.h
+++ b/Engine/Plugins/Animation/MLDeformer/NeuralMorphModel/Source/NeuralMorphModelEditor/Public/NeuralMorphEditorModel.h
@@ -21,7 +21,7 @@ namespace UE::NeuralMorphModel
 	 * The editor model related to the neural morph model's runtime class (UNeuralMorphModel).
 	 */
 	class NEURALMORPHMODELEDITOR_API FNeuralMorphEditorModel
-		: public UE::MLDeformer::FMLDeformerMorphModelEditorModel
+		: public FMLDeformerMorphModelEditorModel
 	{
 	public:
 		// We need to implement this static MakeInstance method.
