[MLDeformer] Fixed an error that caused the console command to be registered multiple times.

#rnx

diff --git a/Engine/Plugins/Animation/MLDeformer/MLDeformerFramework/Source/MLDeformerFramework/Private/MLDeformerModelInstance.cpp b/Engine/Plugins/Animation/MLDeformer/MLDeformerFramework/Source/MLDeformerFramework/Private/MLDeformerModelInstance.cpp
--- a/Engine/Plugins/Animation/MLDeformer/MLDeformerFramework/Source/MLDeformerFramework/Private/MLDeformerModelInstance.cpp
+++ b/Engine/Plugins/Animation/MLDeformer/MLDeformerFramework/Source/MLDeformerFramework/Private/MLDeformerModelInstance.cpp
@@ -21,6 +21,15 @@ CSV_DEFINE_CATEGORY_MODULE(MLDEFORMERFRAMEWORK_API, MLDeformer, false);
 DEFINE_STAT(STAT_MLDeformerInference);
 #endif
 
+namespace UE::MLDeformer
+{
+	static TAutoConsoleVariable<bool> CVarMLDeformerDebugDrawInputPoses(
+		TEXT("MLDeformer.DebugDrawInputPoses"),
+		false,
+		TEXT("Draw the input pose for each active ML Deformer.")
+	);
+}
+
 namespace
 {
 	void CalculateBoneSpaceTransforms(const FReferenceSkeleton& RefSkel, const TArray<FTransform>& ComponentSpaceTransforms, TArray<FTransform>& OutBoneSpaceTransforms)
diff --git a/Engine/Plugins/Animation/MLDeformer/MLDeformerFramework/Source/MLDeformerFramework/Public/MLDeformerModule.h b/Engine/Plugins/Animation/MLDeformer/MLDeformerFramework/Source/MLDeformerFramework/Public/MLDeformerModule.h
--- a/Engine/Plugins/Animation/MLDeformer/MLDeformerFramework/Source/MLDeformerFramework/Public/MLDeformerModule.h
+++ b/Engine/Plugins/Animation/MLDeformer/MLDeformerFramework/Source/MLDeformerFramework/Public/MLDeformerModule.h
@@ -13,12 +13,6 @@ MLDEFORMERFRAMEWORK_API DECLARE_LOG_CATEGORY_EXTERN(LogMLDeformer, Log, All);
 
 namespace UE::MLDeformer
 {
-	static TAutoConsoleVariable<bool> CVarMLDeformerDebugDrawInputPoses(
-		TEXT("MLDeformer.DebugDrawInputPoses"),
-		false,
-		TEXT("Draw the input pose for each ML Deformer that is active.")
-	);
-	
 	/**
 	 * The runtime module for the ML Deformer.
 	 */
