[MLDeformer] Fixed a bug when training, related to masks of input animations. Training would fail with some error right after sampling.

diff --git a/Engine/Plugins/Animation/MLDeformer/MLDeformerFramework/Source/MLDeformerFrameworkEditor/Private/MLDeformerGeomCacheTrainingInputAnimCustomize.cpp b/Engine/Plugins/Animation/MLDeformer/MLDeformerFramework/Source/MLDeformerFrameworkEditor/Private/MLDeformerGeomCacheTrainingInputAnimCustomize.cpp
--- a/Engine/Plugins/Animation/MLDeformer/MLDeformerFramework/Source/MLDeformerFrameworkEditor/Private/MLDeformerGeomCacheTrainingInputAnimCustomize.cpp
+++ b/Engine/Plugins/Animation/MLDeformer/MLDeformerFramework/Source/MLDeformerFrameworkEditor/Private/MLDeformerGeomCacheTrainingInputAnimCustomize.cpp
@@ -35,14 +35,13 @@ namespace UE::MLDeformer
 		UMLDeformerGeomCacheModel* Model = FindMLDeformerModel(StructPropertyHandle);
 		check(Model);
 		int32 NumAnimFrames = 0;
-		int32 NumFramesToSample = 0;
 		FMLDeformerGeomCacheTrainingInputAnim& InputAnim = Model->GetTrainingInputAnims()[ArrayIndex];
 		const UGeometryCache* GeomCache = InputAnim.GetGeometryCache();
 		if (GeomCache)
 		{
 			NumAnimFrames = GeomCache->GetEndFrame() - GeomCache->GetStartFrame() + 1;
 		}
-		NumFramesToSample = InputAnim.GetNumFramesToSample();
+		const int32 NumFramesToSample = InputAnim.GetNumFramesToSample();
 
 		HeaderRow
 		.NameContent()
@@ -181,6 +180,7 @@ namespace UE::MLDeformer
 					[&InputAnim, EditorModel](TSharedPtr<FName> Item, ESelectInfo::Type)
 					{
 						const auto& VertexAttribNames = EditorModel->GetVertexAttributeNames();
+						EditorModel->SetResamplingInputOutputsNeeded(true);
 						if (!VertexAttribNames.IsEmpty() && Item == VertexAttribNames[0])
 						{
 							InputAnim.SetVertexMask(FName());
diff --git a/Engine/Plugins/Animation/MLDeformer/MLDeformerFramework/Source/MLDeformerFrameworkEditor/Private/MLDeformerTrainingModel.cpp b/Engine/Plugins/Animation/MLDeformer/MLDeformerFramework/Source/MLDeformerFrameworkEditor/Private/MLDeformerTrainingModel.cpp
--- a/Engine/Plugins/Animation/MLDeformer/MLDeformerFramework/Source/MLDeformerFrameworkEditor/Private/MLDeformerTrainingModel.cpp
+++ b/Engine/Plugins/Animation/MLDeformer/MLDeformerFramework/Source/MLDeformerFrameworkEditor/Private/MLDeformerTrainingModel.cpp
@@ -112,28 +112,41 @@ TArray<float> UMLDeformerTrainingModel::GetTrainingInputAnimMaskData(FName MaskN
 	FMeshDescription* MeshDescription = SkeletalMesh->GetMeshDescription(LodIndex);
 	const TVertexAttributesConstRef<int32> ImportPointIndex = MeshDescription ? MeshDescription->VertexAttributes().GetAttributesRef<int32>(MeshAttribute::Vertex::ImportPointIndex) : TVertexAttributesConstRef<int32>(); 
 
-	TVertexAttributesConstRef<float> MaskData = EditorModel->FindVertexAttributes(MaskName);
-	if (ImportPointIndex.IsValid() && MaskData.IsValid())
+	const TVertexAttributesConstRef<float> MaskData = EditorModel->FindVertexAttributes(MaskName);
+	if (MaskData.IsValid())
 	{
 		const int32 NumModelVerts = EditorModel->GetModel()->GetNumBaseMeshVerts();
 		PerVertexValues.SetNumZeroed(NumModelVerts);
 
-		check(ImportPointIndex.GetNumElements() == MaskData.GetNumElements());
-		for (int32 Index = 0; Index < MaskData.GetNumElements(); ++Index)
+		if (ImportPointIndex.IsValid())
 		{
-			PerVertexValues[ImportPointIndex[Index]] = MaskData.Get(Index);
+			check(ImportPointIndex.GetNumElements() == MaskData.GetNumElements());
+			for (int32 Index = 0; Index < MaskData.GetNumElements(); ++Index)
+			{
+				PerVertexValues[ImportPointIndex[Index]] = MaskData.Get(Index);
+			}
+		}
+		else
+		{
+			check(PerVertexValues.Num() == MaskData.GetNumElements());
+			for (int32 Index = 0; Index < MaskData.GetNumElements(); ++Index)
+			{
+				PerVertexValues[Index] = MaskData.Get(Index);
+			}
 		}
 	}
 	else
 	{
-		if (MaskName == DefaultMaskName)
+		if (MaskName != DefaultMaskName)
 		{
-			const int32 NumModelVerts = EditorModel->GetModel()->GetNumBaseMeshVerts();
-			PerVertexValues.SetNumUninitialized(NumModelVerts);		
-			for (int32 Index = 0; Index < PerVertexValues.Num(); ++Index)
-			{
-				PerVertexValues[Index] = 1.0f;
-			}
+			UE_LOG(LogMLDeformer, Warning, TEXT("Failed to get vertex mask data for mask '%s'. A fully painted mask will be used."), *MaskName.ToString());
+		}
+		
+		const int32 NumModelVerts = EditorModel->GetModel()->GetNumBaseMeshVerts();
+		PerVertexValues.SetNumUninitialized(NumModelVerts);		
+		for (int32 Index = 0; Index < PerVertexValues.Num(); ++Index)
+		{
+			PerVertexValues[Index] = 1.0f;
 		}
 	}
 
@@ -146,7 +159,7 @@ int32 UMLDeformerTrainingModel::GetMaskIndexForAnimIndex(int32 AnimIndex) const
 	const FName AnimMaskName = Anim->GetVertexMask();
 	if (AnimMaskName.IsNone())
 	{
-		return 0;	
+		return 0;
 	}
 
 	const int32 MaskIndex = MaskNames.Find(AnimMaskName);
