[MLDeformer] Fixed automated tests and some UI issue when changing the max num training frames.

#rnx

diff --git a/Engine/Plugins/Animation/MLDeformer/MLDeformerFramework/Source/MLDeformerFrameworkEditor/Private/MLDeformerEditorModel.cpp b/Engine/Plugins/Animation/MLDeformer/MLDeformerFramework/Source/MLDeformerFrameworkEditor/Private/MLDeformerEditorModel.cpp
--- a/Engine/Plugins/Animation/MLDeformer/MLDeformerFramework/Source/MLDeformerFrameworkEditor/Private/MLDeformerEditorModel.cpp
+++ b/Engine/Plugins/Animation/MLDeformer/MLDeformerFramework/Source/MLDeformerFrameworkEditor/Private/MLDeformerEditorModel.cpp
@@ -1269,7 +1269,9 @@ namespace UE::MLDeformer
 		}
 		else if (Property->GetFName() == UMLDeformerModel::GetMaxTrainingFramesPropertyName())
 		{
+			UpdateNumTrainingFrames();
 			SetResamplingInputOutputsNeeded(true);
+			GetEditor()->GetModelDetailsView()->ForceRefresh();
 		}
 		else
 		if (Property->GetFName() == UMLDeformerModel::GetBoneIncludeListPropertyName() ||
diff --git a/Engine/Plugins/Animation/MLDeformer/MLDeformerFramework/Source/MLDeformerFrameworkEditor/Public/MLDeformerEditorModel.h b/Engine/Plugins/Animation/MLDeformer/MLDeformerFramework/Source/MLDeformerFrameworkEditor/Public/MLDeformerEditorModel.h
--- a/Engine/Plugins/Animation/MLDeformer/MLDeformerFramework/Source/MLDeformerFrameworkEditor/Public/MLDeformerEditorModel.h
+++ b/Engine/Plugins/Animation/MLDeformer/MLDeformerFramework/Source/MLDeformerFrameworkEditor/Public/MLDeformerEditorModel.h
@@ -1348,6 +1348,7 @@ namespace UE::MLDeformer
 
 		// Perform the training.
 		// This will trigger the Python class train function to be called.
+		EditorModel->TriggerInputAssetChanged(false);
 		TrainingModelClass* TrainingModel = Cast<TrainingModelClass>(TrainingModels.Last()->GetDefaultObject());
 		TrainingModel->Init(EditorModel);
 		const int32 ReturnCode = TrainingModel->Train();
