[MLDeformer] Fixed a crash on undo on morph based models. Also changed the icon of the Train button.

diff --git a/Engine/Plugins/Animation/MLDeformer/MLDeformerFramework/Source/MLDeformerFramework/Private/MLDeformerMorphModel.cpp b/Engine/Plugins/Animation/MLDeformer/MLDeformerFramework/Source/MLDeformerFramework/Private/MLDeformerMorphModel.cpp
--- a/Engine/Plugins/Animation/MLDeformer/MLDeformerFramework/Source/MLDeformerFramework/Private/MLDeformerMorphModel.cpp
+++ b/Engine/Plugins/Animation/MLDeformer/MLDeformerFramework/Source/MLDeformerFramework/Private/MLDeformerMorphModel.cpp
@@ -104,6 +104,7 @@ void UMLDeformerMorphModel::Serialize(FArchive& Archive)
 		// If we only support 1 LOD, in older files.
 		if (Archive.CustomVer(UE::MLDeformer::FMLDeformerObjectVersion::GUID) < UE::MLDeformer::FMLDeformerObjectVersion::LODSupportAdded)
 		{
+			ClearMorphTargetSets();
 			AddMorphSets(1);
 			Archive << bHasMorphData;
 
@@ -118,6 +119,7 @@ void UMLDeformerMorphModel::Serialize(FArchive& Archive)
 			int32 NumLODs = 1;
 			Archive << NumLODs;
 
+			ClearMorphTargetSets();
 			AddMorphSets(NumLODs);
 
 			for (int32 LOD = 0; LOD < NumLODs; ++LOD)
diff --git a/Engine/Plugins/Animation/MLDeformer/MLDeformerFramework/Source/MLDeformerFramework/Public/MLDeformerMorphModel.h b/Engine/Plugins/Animation/MLDeformer/MLDeformerFramework/Source/MLDeformerFramework/Public/MLDeformerMorphModel.h
--- a/Engine/Plugins/Animation/MLDeformer/MLDeformerFramework/Source/MLDeformerFramework/Public/MLDeformerMorphModel.h
+++ b/Engine/Plugins/Animation/MLDeformer/MLDeformerFramework/Source/MLDeformerFramework/Public/MLDeformerMorphModel.h
@@ -29,7 +29,6 @@ public:
 
 	// UMLDeformerModel overrides.
 	virtual bool IsNeuralNetworkOnGPU() const override				{ return false; }	// CPU based neural network.
-	virtual bool DoesSupportQualityLevels() const override			{ return true; }	// We can disable morph targets based on the Deformer LOD float value.
 	virtual bool DoesSupportLOD() const override					{ return true; }	// We support deformations in different LOD levels.
 	UE_API virtual void Serialize(FArchive& Archive) override;
 	UE_API virtual UMLDeformerInputInfo* CreateInputInfo() override;
diff --git a/Engine/Plugins/Animation/MLDeformer/MLDeformerFramework/Source/MLDeformerFrameworkEditor/Private/MLDeformerEditorToolkit.cpp b/Engine/Plugins/Animation/MLDeformer/MLDeformerFramework/Source/MLDeformerFrameworkEditor/Private/MLDeformerEditorToolkit.cpp
--- a/Engine/Plugins/Animation/MLDeformer/MLDeformerFramework/Source/MLDeformerFrameworkEditor/Private/MLDeformerEditorToolkit.cpp
+++ b/Engine/Plugins/Animation/MLDeformer/MLDeformerFramework/Source/MLDeformerFrameworkEditor/Private/MLDeformerEditorToolkit.cpp
@@ -497,7 +497,26 @@ namespace UE::MLDeformer
 	{
 		return ActiveModel.Get() ? ActiveModel->IsReadyForTraining() : false;
 	}
+	
+	FSlateIcon FMLDeformerEditorToolkit::GetTrainingStatusIcon() const
+	{
+		static const FName CompileStatusBackground("AssetEditor.CompileStatus.Background");
+		static const FName CompileStatusUnknown("AssetEditor.CompileStatus.Overlay.Unknown");
+		static const FName CompileStatusWarning("AssetEditor.CompileStatus.Overlay.Warning");
+		
+		if (!ActiveModel.IsValid())
+		{
+			return FSlateIcon(FAppStyle::Get().GetStyleSetName(), CompileStatusBackground, NAME_None, CompileStatusUnknown); 
+		}
 
+		if (ActiveModel->IsReadyForTraining())
+		{
+			return FSlateIcon(FAppStyle::Get().GetStyleSetName(), CompileStatusBackground, NAME_None);
+		}
+
+		return FSlateIcon(FAppStyle::Get().GetStyleSetName(), CompileStatusBackground, NAME_None, CompileStatusWarning);
+	}
+	
 	void FMLDeformerEditorToolkit::FillToolbar(FToolBarBuilder& ToolbarBuilder)
 	{
 		// Training button and model selection.
@@ -506,13 +525,13 @@ namespace UE::MLDeformer
 			ToolbarBuilder.AddToolBarButton(
 				FUIAction
 				(
-					FExecuteAction::CreateLambda([this]() { Train(/*bSuppressDialogs*/false); }),
+					FExecuteAction::CreateLambda([this]() { Train(false); }),
 					FCanExecuteAction::CreateLambda([this]() { return IsTrainButtonEnabled() && !IsPaintModeActive(); })
 				),
 				NAME_None,
 				LOCTEXT("TrainModel", "Train Model"),
 				LOCTEXT("TrainModelTooltip", "Train the machine learning model.\nIf this button is disabled you first need to configure your inputs in the details panel."),
-				FSlateIcon(),
+				TAttribute<FSlateIcon>(this, &FMLDeformerEditorToolkit::GetTrainingStatusIcon),
 				EUserInterfaceActionType::ToggleButton
 			);
 
diff --git a/Engine/Plugins/Animation/MLDeformer/MLDeformerFramework/Source/MLDeformerFrameworkEditor/Public/MLDeformerEditorToolkit.h b/Engine/Plugins/Animation/MLDeformer/MLDeformerFramework/Source/MLDeformerFrameworkEditor/Public/MLDeformerEditorToolkit.h
--- a/Engine/Plugins/Animation/MLDeformer/MLDeformerFramework/Source/MLDeformerFrameworkEditor/Public/MLDeformerEditorToolkit.h
+++ b/Engine/Plugins/Animation/MLDeformer/MLDeformerFramework/Source/MLDeformerFrameworkEditor/Public/MLDeformerEditorToolkit.h
@@ -186,6 +186,7 @@ namespace UE::MLDeformer
 		FText GetVizModeName(EMLDeformerVizMode Mode) const;
 		void ShowNoModelsWarningIfNeeded();
 		EVisibility GetDebuggingVisibility() const;
+		FSlateIcon GetTrainingStatusIcon() const;
 
 		TSharedRef<SWidget> GenerateModelButtonContents(TSharedRef<FUICommandList> InCommandList);
 		TSharedRef<SWidget> GenerateVizModeButtonContents(TSharedRef<FUICommandList> InCommandList);
diff --git a/Engine/Plugins/Animation/MLDeformer/NearestNeighborModel/Source/NearestNeighborModel/Private/NearestNeighborModel.cpp b/Engine/Plugins/Animation/MLDeformer/NearestNeighborModel/Source/NearestNeighborModel/Private/NearestNeighborModel.cpp
--- a/Engine/Plugins/Animation/MLDeformer/NearestNeighborModel/Source/NearestNeighborModel/Private/NearestNeighborModel.cpp
+++ b/Engine/Plugins/Animation/MLDeformer/NearestNeighborModel/Source/NearestNeighborModel/Private/NearestNeighborModel.cpp
@@ -1315,7 +1315,7 @@ FString UNearestNeighborModel::GetDefaultDeformerGraphAssetPath() const
 
 UMLDeformerModelInstance* UNearestNeighborModel::CreateModelInstance(UMLDeformerComponent* Component)
 {
-	return NewObject<UNearestNeighborModelInstance>(Component, NAME_None, RF_Transactional);
+	return NewObject<UNearestNeighborModelInstance>(Component);
 }
 
 void UNearestNeighborModel::Serialize(FArchive& Archive)
diff --git a/Engine/Plugins/Animation/MLDeformer/NeuralMorphModel/Source/NeuralMorphModel/Private/NeuralMorphModel.cpp b/Engine/Plugins/Animation/MLDeformer/NeuralMorphModel/Source/NeuralMorphModel/Private/NeuralMorphModel.cpp
--- a/Engine/Plugins/Animation/MLDeformer/NeuralMorphModel/Source/NeuralMorphModel/Private/NeuralMorphModel.cpp
+++ b/Engine/Plugins/Animation/MLDeformer/NeuralMorphModel/Source/NeuralMorphModel/Private/NeuralMorphModel.cpp
@@ -43,7 +43,7 @@ UNeuralMorphModel::UNeuralMorphModel(const FObjectInitializer& ObjectInitializer
 
 UMLDeformerModelInstance* UNeuralMorphModel::CreateModelInstance(UMLDeformerComponent* Component)
 {
-	return NewObject<UNeuralMorphModelInstance>(Component, NAME_None, RF_Transactional);
+	return NewObject<UNeuralMorphModelInstance>(Component);
 }
 
 UMLDeformerInputInfo* UNeuralMorphModel::CreateInputInfo()
diff --git a/Engine/Plugins/Animation/MLDeformer/VertexDeltaModel/Source/VertexDeltaModel/Private/VertexDeltaModel.cpp b/Engine/Plugins/Animation/MLDeformer/VertexDeltaModel/Source/VertexDeltaModel/Private/VertexDeltaModel.cpp
--- a/Engine/Plugins/Animation/MLDeformer/VertexDeltaModel/Source/VertexDeltaModel/Private/VertexDeltaModel.cpp
+++ b/Engine/Plugins/Animation/MLDeformer/VertexDeltaModel/Source/VertexDeltaModel/Private/VertexDeltaModel.cpp
@@ -51,7 +51,7 @@ void UVertexDeltaModel::SetNNEModelData(TObjectPtr<UNNEModelData> ModelData)
 
 UMLDeformerModelInstance* UVertexDeltaModel::CreateModelInstance(UMLDeformerComponent* Component)
 {
-	return NewObject<UVertexDeltaModelInstance>(Component, NAME_None, RF_Transactional);
+	return NewObject<UVertexDeltaModelInstance>(Component);
 }
 
 FString UVertexDeltaModel::GetDefaultDeformerGraphAssetPath() const 
diff --git a/Engine/Plugins/Experimental/Animation/MLDeformer/DetailPoseModel/Source/DetailPoseModel/Private/DetailPoseModel.cpp b/Engine/Plugins/Experimental/Animation/MLDeformer/DetailPoseModel/Source/DetailPoseModel/Private/DetailPoseModel.cpp
--- a/Engine/Plugins/Experimental/Animation/MLDeformer/DetailPoseModel/Source/DetailPoseModel/Private/DetailPoseModel.cpp
+++ b/Engine/Plugins/Experimental/Animation/MLDeformer/DetailPoseModel/Source/DetailPoseModel/Private/DetailPoseModel.cpp
@@ -43,7 +43,7 @@ UDetailPoseModel::UDetailPoseModel(const FObjectInitializer& ObjectInitializer)
 
 UMLDeformerModelInstance* UDetailPoseModel::CreateModelInstance(UMLDeformerComponent* Component)
 {
-	return NewObject<UDetailPoseModelInstance>(Component, NAME_None, RF_Transactional);
+	return NewObject<UDetailPoseModelInstance>(Component);
 }
 
 UMLDeformerInputInfo* UDetailPoseModel::CreateInputInfo()
