Updated headers using UnrealCodeFixup to make sure dllstorage is on methods/staticvars instead of types. Part 1/n

diff --git a/Engine/Plugins/Animation/MLDeformer/MLDeformerFramework/Source/MLDeformerFrameworkEditor/Public/MLDeformerEditorModel.h b/Engine/Plugins/Animation/MLDeformer/MLDeformerFramework/Source/MLDeformerFrameworkEditor/Public/MLDeformerEditorModel.h
--- a/Engine/Plugins/Animation/MLDeformer/MLDeformerFramework/Source/MLDeformerFrameworkEditor/Public/MLDeformerEditorModel.h
+++ b/Engine/Plugins/Animation/MLDeformer/MLDeformerFramework/Source/MLDeformerFrameworkEditor/Public/MLDeformerEditorModel.h
@@ -15,6 +15,8 @@
 #include "MeshDescription.h"
 #include "MLDeformerEditorModel.generated.h"
 
+#define UE_API MLDEFORMERFRAMEWORKEDITOR_API
+
 class UMLDeformerModel;
 class UMLDeformerInputInfo;
 class AAnimationEditorPreviewActor;
@@ -71,7 +73,7 @@ namespace UE::MLDeformer
 	 * In most cases you will want to inherit your own editor model from either FMLDeformerGeomCacheEditorModel or FMLDeformerMorphModelEditorModel though.
 	 * These classes will already implement a lot of base functionality in case you work with geometry caches as ground truth data etc.
 	 */
-	class MLDEFORMERFRAMEWORKEDITOR_API FMLDeformerEditorModel
+	class FMLDeformerEditorModel
 		: public TSharedFromThis<FMLDeformerEditorModel>
 		, public FGCObject
 	{
@@ -80,7 +82,7 @@ namespace UE::MLDeformer
 		 * The editor model initialization settings.
 		 * This is used in the Init call.
 		 */
-		struct MLDEFORMERFRAMEWORKEDITOR_API InitSettings
+		struct InitSettings
 		{
 			FMLDeformerEditorToolkit* Editor = nullptr;
 			UMLDeformerModel* Model = nullptr;
@@ -90,10 +92,10 @@ namespace UE::MLDeformer
 		 * The destructor, which unregisters the instance of this editor model from the model registry.
 		 * It also deletes the editor actors and unregisters from the OnPostEditChangeProperty delegate.
 		 */
-		virtual ~FMLDeformerEditorModel() override;
+		UE_API virtual ~FMLDeformerEditorModel() override;
 
 		// FGCObject overrides.
-		virtual void AddReferencedObjects(FReferenceCollector& Collector) override;
+		UE_API virtual void AddReferencedObjects(FReferenceCollector& Collector) override;
 		virtual FString GetReferencerName() const override { return TEXT("FMLDeformerEditorModel"); }
 		// ~END FGCObject overrides.
 
@@ -102,7 +104,7 @@ namespace UE::MLDeformer
 		 * This includes the skeletal mesh, input animations, and shared visualization settings.
 		 * @param SourceEditorModel The source model to copy the shared settings from.
 		 */
-		virtual void CopyBaseSettingsFromModel(const FMLDeformerEditorModel* SourceEditorModel);
+		UE_API virtual void CopyBaseSettingsFromModel(const FMLDeformerEditorModel* SourceEditorModel);
 
 		/**
 		 * Change the skeletal mesh on a specific skeletal mesh component.
@@ -115,7 +117,7 @@ namespace UE::MLDeformer
 		 * @param SkelMeshComponent The skeletal mesh component to change the mesh for.
 		 * @param Mesh The new mesh to use.
 		 */
-		static void ChangeSkeletalMeshOnComponent(USkeletalMeshComponent* SkelMeshComponent, USkeletalMesh* Mesh);
+		static UE_API void ChangeSkeletalMeshOnComponent(USkeletalMeshComponent* SkelMeshComponent, USkeletalMesh* Mesh);
 
 		/**
 		 * Get the number of training frames.
@@ -123,7 +125,7 @@ namespace UE::MLDeformer
 		 * For example this could return the number of frames inside your geometry cache.
 		 * @return The number of frames in the training data.
 		 */
-		virtual int32 GetNumTrainingFrames() const;
+		UE_API virtual int32 GetNumTrainingFrames() const;
 
 		/**
 		 * Updates the number of available training frames as returned by GetNumTrainingFrames().
@@ -131,7 +133,7 @@ namespace UE::MLDeformer
 		 * So this must be the sum of all frames of all enabled training input animations.
 		 * This should not take into account the maximum number of frames we want to train with.
 		 */
-		virtual void UpdateNumTrainingFrames();
+		UE_API virtual void UpdateNumTrainingFrames();
 
 		/**
 		 * Launch the training. This gets executed when the Train button is pressed.
@@ -164,7 +166,7 @@ namespace UE::MLDeformer
 		 * 
 		 * @return The training result.
 		 */
-		virtual ETrainingResult Train();
+		UE_API virtual ETrainingResult Train();
 
 		// Optional overrides.
 		/**
@@ -172,7 +174,7 @@ namespace UE::MLDeformer
 		 * This will mainly create the delta sampler, editor input info and some other things.
 		 * @param Settings The initialization settings.
 		 */
-		virtual void Init(const InitSettings& Settings);
+		UE_API virtual void Init(const InitSettings& Settings);
 
 		/**
 		 * Create the editor actors. These are the actors that will appear in the viewport of the asset editor.
@@ -182,23 +184,23 @@ namespace UE::MLDeformer
 		 * Nothing prevents you from adding more or less actors though.
 		 * @param InPersonaPreviewScene The persona scene.
 		 */
-		virtual void CreateActors(const TSharedRef<IPersonaPreviewScene>& InPersonaPreviewScene);
+		UE_API virtual void CreateActors(const TSharedRef<IPersonaPreviewScene>& InPersonaPreviewScene);
 
 		/**
 		 * Clears the world, which basically means it removes all the editor actors and engine actors from the editor world.
 		 */
-		virtual void ClearWorld();
+		UE_API virtual void ClearWorld();
 
 		/**
 		 * Clears the Persona preview scene. Resetting the preview mesh, preview mesh component etc.
 		 */
-		virtual void ClearPersonaPreviewScene();
+		UE_API virtual void ClearPersonaPreviewScene();
 
 		/**
 		 * Clear both the world and preview scene.
 		 * This calls ClearWorld and ClearPersonaPreviewScene.
 		 */
-		virtual void ClearWorldAndPersonaPreviewScene();
+		UE_API virtual void ClearWorldAndPersonaPreviewScene();
 
 		/**
 		 * Create an actor in the editor viewport.
@@ -208,7 +210,7 @@ namespace UE::MLDeformer
 		 * @param Settings The creation settings for the editor actor you are creating.
 		 * @return A pointer to the editor actor object that was created.
 		 */
-		virtual FMLDeformerEditorActor* CreateEditorActor(const FMLDeformerEditorActor::FConstructSettings& Settings) const;
+		UE_API virtual FMLDeformerEditorActor* CreateEditorActor(const FMLDeformerEditorActor::FConstructSettings& Settings) const;
 
 		UE_DEPRECATED(5.4, "This method will be removed. Please use the CreateSamplerObject that returns a shared pointer instead.")
 		virtual FMLDeformerSampler* CreateSampler() const { return nullptr; }
@@ -219,7 +221,7 @@ namespace UE::MLDeformer
 		 * A geometry cache based editor model would create a new FMLDeformerGeomCacheSampler for example.
 		 * @return A pointer to the newly created sampler.
 		 */
-		virtual TSharedPtr<FMLDeformerSampler> CreateSamplerObject() const;
+		UE_API virtual TSharedPtr<FMLDeformerSampler> CreateSamplerObject() const;
 
 		/**
 		 * Get the number of training input animations.
@@ -227,7 +229,7 @@ namespace UE::MLDeformer
 		 * This data is fed into the training process.
 		 * @return The number of training input anims.
 		 */
-		virtual int32 GetNumTrainingInputAnims() const;
+		UE_API virtual int32 GetNumTrainingInputAnims() const;
 
 		/**
 		 * Get a given training input animation.
@@ -236,55 +238,55 @@ namespace UE::MLDeformer
 		 * @param Index The training input animation to get.
 		 * @see GetNumTrainingInputAnims
 		 */
-		virtual FMLDeformerTrainingInputAnim* GetTrainingInputAnim(int32 Index) const;
+		UE_API virtual FMLDeformerTrainingInputAnim* GetTrainingInputAnim(int32 Index) const;
 
 		/**
 		 * Update the timeline's training input animation list.
 		 * This is the list of animations that show inside the timeline when in training mode.
 		 * Each training input animation pair needs some name that is shown inside the timeline.
 		 */
-		virtual void UpdateTimelineTrainingAnimList();
+		UE_API virtual void UpdateTimelineTrainingAnimList();
 
 		/**
 		 * Get the time in seconds, for a given frame in the training data.
 		 * @param FrameNumber The frame number to get the time in seconds for.
 		 * @return The time offset of this frame, in seconds.
 		 */
-		virtual double GetTrainingTimeAtFrame(int32 FrameNumber) const;
+		UE_API virtual double GetTrainingTimeAtFrame(int32 FrameNumber) const;
 
 		/**
 		 * Get the training data frame number for a specific time offset in seconds.
 		 * @param TimeInSeconds The time offset, in seconds, to get the frame number for.
 		 * @return The frame number for this time offset.
 		 */
-		virtual int32 GetTrainingFrameAtTime(double TimeInSeconds) const;
+		UE_API virtual int32 GetTrainingFrameAtTime(double TimeInSeconds) const;
 
 		/**
 		 * Get the time in seconds, for a given frame in the test data.
 		 * @param FrameNumber The frame number to get the time in seconds for.
 		 * @return The time offset of this frame, in seconds.
 		 */
-		virtual double GetTestTimeAtFrame(int32 FrameNumber) const;
+		UE_API virtual double GetTestTimeAtFrame(int32 FrameNumber) const;
 
 		/**
 		 * Get the test data frame number for a specific time offset in seconds.
 		 * @param TimeInSeconds The time offset, in seconds, to get the frame number for.
 		 * @return The frame number for this time offset.
 		 */
-		virtual int32 GetTestFrameAtTime(double TimeInSeconds) const;
+		UE_API virtual int32 GetTestFrameAtTime(double TimeInSeconds) const;
 
 		/**
 		 * Get the total number of frames in the test anim sequence.
 		 * @return The number of frames.
 		 */
-		virtual int32 GetNumTestFrames() const;
+		UE_API virtual int32 GetNumTestFrames() const;
 
 		/**
 		 * Tick the editor viewport.
 		 * @param ViewportClient The editor viewport client object.
 		 * @param DeltaTime The delta time of this tick, in seconds.
 		 */
-		virtual void Tick(FEditorViewportClient* ViewportClient, float DeltaTime);
+		UE_API virtual void Tick(FEditorViewportClient* ViewportClient, float DeltaTime);
 
 		/**
 		 * Create the linear skinned actor that is used as training base.
@@ -293,7 +295,7 @@ namespace UE::MLDeformer
 		 * The default implementation will directly set the persona actor and preview mesh and component to this actor.
 		 * @param InPersonaPreviewScene The personal preview scene.
 		 */
-		virtual void CreateTrainingLinearSkinnedActor(const TSharedRef<IPersonaPreviewScene>& InPersonaPreviewScene);
+		UE_API virtual void CreateTrainingLinearSkinnedActor(const TSharedRef<IPersonaPreviewScene>& InPersonaPreviewScene);
 
 		/**
 		 * Create the linear skinned test actor.
@@ -301,7 +303,7 @@ namespace UE::MLDeformer
 		 * An editor actor needs to be created by this method.
 		 * @param World The world to create the actor in.
 		 */
-		virtual void CreateTestLinearSkinnedActor(UWorld* World);
+		UE_API virtual void CreateTestLinearSkinnedActor(UWorld* World);
 
 		/**
 		 * Create the test ML Deformed actor.
@@ -310,13 +312,13 @@ namespace UE::MLDeformer
 		 * An editor actor needs to be created by this method.
 		 * @param World The world to create the actor in.
 		 */
-		virtual void CreateTestMLDeformedActor(UWorld* World);
+		UE_API virtual void CreateTestMLDeformedActor(UWorld* World);
 
 		/**
 		 * Create the compare actors for the testing mode.
 		 * @param World The world to create the actors in.
 		 */
-		virtual void CreateTestCompareActors(UWorld* World);
+		UE_API virtual void CreateTestCompareActors(UWorld* World);
 
 		/**
 		 * Check if this model is compatible with a specific deformer asset.
@@ -325,13 +327,13 @@ namespace UE::MLDeformer
 		 * @param Deformer The deformer asset to check against.
 		 * @return Returns true if the specified asset is compatible with this one, otherwise false is returned.
 		 */
-		virtual bool IsCompatibleDeformer(UMLDeformerAsset* Deformer) const;
+		UE_API virtual bool IsCompatibleDeformer(UMLDeformerAsset* Deformer) const;
 
 		/**
 		 * Update the mesh offset factors, which is basically controlling where to place the actors in the scene.
 		 * The spacing between the characters is the mesh offset factor multiplied by the mesh spacing.
 		 */
-		virtual void UpdateMeshOffsetFactors();
+		UE_API virtual void UpdateMeshOffsetFactors();
 
 		/**
 		 * Create the training ground truth actor.
@@ -356,38 +358,38 @@ namespace UE::MLDeformer
 		 * For example when they click the timeline in the editor, or enter a specific frame number.
 		 * This happens only in training mode.
 		 */
-		virtual void OnTrainingDataFrameChanged();
+		UE_API virtual void OnTrainingDataFrameChanged();
 
 		/**
 		 * Update the transforms of all the editor actors.
 		 * This should apply the mesh offsets to space the actors in the world/viewport.
 		 */
-		virtual void UpdateActorTransforms();
+		UE_API virtual void UpdateActorTransforms();
 
 		/**
 		 * Update the visibility status of the editor actors.
 		 * This should hide test actors when the editor is in training mode for example, and the other way around.
 		 */
-		virtual void UpdateActorVisibility();
+		UE_API virtual void UpdateActorVisibility();
 
 		/**
 		 * Update the visibility of the labels, as well as their positions.
 		 * For example labels of invisible actors should be made invisible as well.
 		 */
-		virtual void UpdateLabels();
+		UE_API virtual void UpdateLabels();
 
 		/**
 		 * This is called when training data inputs change.
 		 * This includes changing things such as the Skeletal Mesh, training anim sequence and target mesh.
 		 * It will re-initialize the components. For example a skeletal mesh change will update the skeletal mesh component to the new mesh.
 		 */
-		virtual void OnInputAssetsChanged();
+		UE_API virtual void OnInputAssetsChanged();
 
 		/**
 		 * This is executed after the input assets have changed.
 		 * On default this will refresh the deformer graphs, update the timeline ranges, select the right frames, modify the training button ready state, etc.
 		 */
-		virtual void OnPostInputAssetChanged();
+		UE_API virtual void OnPostInputAssetChanged();
 
 		/**
 		 * This is called whenever a property value changes in the UI.
@@ -395,7 +397,7 @@ namespace UE::MLDeformer
 		 * You can overload this and call the base class method inside it, and then handle your own model specific property changes.
 		 * @param PropertyChangedEvent The event information object that contains info about the property change.
 		 */
-		virtual void OnPropertyChanged(FPropertyChangedEvent& PropertyChangedEvent);
+		UE_API virtual void OnPropertyChanged(FPropertyChangedEvent& PropertyChangedEvent);
 
 		/**
 		 * This is called before undo is called on the model.
@@ -411,13 +413,13 @@ namespace UE::MLDeformer
 		 * This is called whenever a transaction happened on the model.
 		 * @param Event Information about the event.
 		 */
-		virtual void OnPostTransacted(const FTransactionObjectEvent& Event);
+		UE_API virtual void OnPostTransacted(const FTransactionObjectEvent& Event);
 
 		/**
 		 * Executed when the user presses the play button in the timeline.
 		 * On default this will start playing the anim sequence related to whether we are in training or testing mode.
 		 */
-		virtual void OnPlayPressed();
+		UE_API virtual void OnPlayPressed();
 
 		/**
 		 * This is called just before we start the training process.
@@ -432,7 +434,7 @@ namespace UE::MLDeformer
 		 * @param bUsePartiallyTrainedWhenAborted When the user aborts training, they can get offered to use the partially trained neural network. This boolean 
 		 *                                        specifies whether they picked "Yes" (true) or "No" (false) to that question.
 		 */
-		virtual void OnPostTraining(ETrainingResult TrainingResult, bool bUsePartiallyTrainedWhenAborted);
+		UE_API virtual void OnPostTraining(ETrainingResult TrainingResult, bool bUsePartiallyTrainedWhenAborted);
 
 		/**
 		 * This is executed when the training process got aborted by the user.
@@ -450,14 +452,14 @@ namespace UE::MLDeformer
 		 * The same goes in testing mode, if we are in testing mode it will only return true when the test anim sequence is playing.
 		 * @return Returns true whether the mode specific animation sequence is playing.
 		 */
-		virtual bool IsPlayingAnim() const;
+		UE_API virtual bool IsPlayingAnim() const;
 
 		/**
 		 * Check whether our anim sequence is playing forward or backward.
 		 * Just like the IsPlayingAnim method, it is specific to the mode we are in.
 		 * @return Returns true when we are playing the animation forwards, and not backward.
 		 */
-		virtual bool IsPlayingForward() const;
+		UE_API virtual bool IsPlayingForward() const;
 
 		/**
 		 * Get the current playback position of our training mode.
@@ -465,7 +467,7 @@ namespace UE::MLDeformer
 		 * So if you have a geometry cache or so specified as target mesh, it would extract the current time value from that.
 		 * @return The time offset, in seconds, of our training data.
 		 */
-		virtual double CalcTrainingTimelinePosition() const;
+		UE_API virtual double CalcTrainingTimelinePosition() const;
 
 		/**
 		 * Calculate the test mode timeline position.
@@ -473,7 +475,7 @@ namespace UE::MLDeformer
 		 * So if you have a geometry cache or so specified as test ground truth, it would extract the current time value from that.
 		 * @return The time offset, in seconds, of our test data.
 		 */
-		virtual double CalcTestTimelinePosition() const;
+		UE_API virtual double CalcTestTimelinePosition() const;
 
 		/**
 		 * This is executed when the timeline has been scrubbed.
@@ -481,7 +483,7 @@ namespace UE::MLDeformer
 		 * @param NewScrubTime This is the new time value, in seconds.
 		 * @param bIsScrubbing Specifies whether the user is scrubbing or just clicked.
 		 */
-		virtual void OnTimeSliderScrubPositionChanged(double NewScrubTime, bool bIsScrubbing);
+		UE_API virtual void OnTimeSliderScrubPositionChanged(double NewScrubTime, bool bIsScrubbing);
 
 		/**
 		 * Update the playback speed of the test data.
@@ -489,7 +491,7 @@ namespace UE::MLDeformer
 		 * it will update the play speed. For example a linear skinned mesh will update its skeletal mesh component's play speed, while
 		 * a geometry cache ground truth actor would update the play speed set on the geometry cache component.
 		 */
-		virtual void UpdateTestAnimPlaySpeed();
+		UE_API virtual void UpdateTestAnimPlaySpeed();
 
 		/**
 		 * Clamp the current training frame index to be in a valid range.
@@ -497,7 +499,7 @@ namespace UE::MLDeformer
 		 * So if the training data contains 100 frames, so index 0..99, and we set the TrainingFrameNumber member to 500, this method
 		 * will turn that into 99.
 		 */
-		virtual void ClampCurrentTrainingFrameIndex();
+		UE_API virtual void ClampCurrentTrainingFrameIndex();
 
 		/**
 		 * Clamp the current test frame index to be in a valid range.
@@ -505,7 +507,7 @@ namespace UE::MLDeformer
 		 * So if the test data contains 100 frames, so index 0..99, and we set the TestingFrameNumber member to 500, this method
 		 * will turn that into 99.
 		 */
-		virtual void ClampCurrentTestFrameIndex();
+		UE_API virtual void ClampCurrentTestFrameIndex();
 
 		/**
 		 * Get the total number of frames used for training, but limited by the maximum number of training frames we have setup.
@@ -513,7 +515,7 @@ namespace UE::MLDeformer
 		 * If the number of training frames is 10000 but we set the max number of training frames to 1 million, then it will still return only 10000.
 		 * @return The number of frames we should be training on.
 		 */
-		virtual int32 GetNumFramesForTraining() const;
+		UE_API virtual int32 GetNumFramesForTraining() const;
 
 		/**
 		 * Go to a given frame in the training data.
@@ -521,7 +523,7 @@ namespace UE::MLDeformer
 		 * @param FrameNumber The training frame to go to. This will automatically be clamped to a valid range in case it goes out of bounds.
 		 * @param bIsScrubbing This is set to true when we are scrubbing the timeline. It is set to false when the mouse button is released.
 		 */
-		virtual void SetTrainingFrame(int32 FrameNumber, bool bIsScrubbing);
+		UE_API virtual void SetTrainingFrame(int32 FrameNumber, bool bIsScrubbing);
 
 		UE_DEPRECATED(5.5, "Please use the SetTrainingFrame that takes a bIsScrubbing parameter.")
 		virtual void SetTrainingFrame(int32 FrameNumber) { SetTrainingFrame(FrameNumber, false); }
@@ -532,7 +534,7 @@ namespace UE::MLDeformer
 		 * @param FrameNumber The test frame to go to. This will automatically be clamped to a valid range in case it goes out of bounds.
 		 * @param bIsScrubbing This is set to true when we are scrubbing the timeline. It is set to false when the mouse button is released.
 		 */
-		virtual void SetTestFrame(int32 FrameNumber, bool bIsScrubbing);
+		UE_API virtual void SetTestFrame(int32 FrameNumber, bool bIsScrubbing);
 
 		UE_DEPRECATED(5.5, "Please use the SetTestFrame that takes a bIsScrubbing parameter.")
 		virtual void SetTestFrame(int32 FrameNumber) { SetTestFrame(FrameNumber, false); }
@@ -544,7 +546,7 @@ namespace UE::MLDeformer
 		 * @param Viewport The viewport object.
 		 * @param PDI The draw primitive interface used for drawing things.
 		 */
-		virtual void Render(const FSceneView* View, FViewport* Viewport, FPrimitiveDrawInterface* PDI);
+		UE_API virtual void Render(const FSceneView* View, FViewport* Viewport, FPrimitiveDrawInterface* PDI);
 
 		/**
 		 * Update the state that is used to decide whether the Train button in the UI is enabled or not.
@@ -557,7 +559,7 @@ namespace UE::MLDeformer
 		 * Get the text that we should display as overlay on top of the asset editor viewport.
 		 * @return The text, or an empty text object if there is nothing to show.
 		 */
-		virtual FText GetOverlayText() const;
+		UE_API virtual FText GetOverlayText() const;
 
 		/**
 		 * Initialize the input information, which basically describes what the inputs to the neural network are going to be.
@@ -567,31 +569,31 @@ namespace UE::MLDeformer
 		 * compatibility checks when applying this ML Deformer onto some character.
 		 * @param InputInfo The input info object to initialize.
 		 */
-		virtual void InitInputInfo(UMLDeformerInputInfo* InputInfo);
+		UE_API virtual void InitInputInfo(UMLDeformerInputInfo* InputInfo);
 
 		/**
 		 * Call the SetupComponent on every ML Deformer component on our editor actors and update the compatibility status of the deformer components.
 		 * This gets triggered when some input assets change.
 		 */
-		virtual void RefreshMLDeformerComponents();
+		UE_API virtual void RefreshMLDeformerComponents();
 
 		/**
 		 * Create the heat map material for this model.
 		 * This updates the HeatMapMaterial member.
 		 */
-		virtual void CreateHeatMapMaterial();
+		UE_API virtual void CreateHeatMapMaterial();
 
 		/**
 		 * Create the heat map deformer graph for this model.
 		 * This updates the HeatMapDeformerGraph member.
 		 */
-		virtual void CreateHeatMapDeformerGraph();
+		UE_API virtual void CreateHeatMapDeformerGraph();
 
 		/**
 		 * Create the heat map assets for this model.
 		 * This calls CreateHeatMapMaterial and CreateHeatMapDeformerGraph on default.
 		 */
-		virtual void CreateHeatMapAssets();
+		UE_API virtual void CreateHeatMapAssets();
 
 		/**
 		 * Get the heat map deformer that should currently be used (when heatmaps are enabled).
@@ -600,33 +602,33 @@ namespace UE::MLDeformer
 		 * make it return the appropriate deformer graph for skinning method that is currently used.
 		 * @return Returns a pointer to the mesh deformer to use when heatmaps are enabled.
 		 */
-		virtual UMeshDeformer* GetActiveHeatMapDeformer() const;
+		UE_API virtual UMeshDeformer* GetActiveHeatMapDeformer() const;
 
 		/**
 		 * Specify whether the heat map material should be enabled or not.
 		 * This will activate the heat map visualization, or deactivate it.
 		 * @param bEnabled Set to true to enable heat map mode.
 		 */
-		virtual void SetHeatMapMaterialEnabled(bool bEnabled);
+		UE_API virtual void SetHeatMapMaterialEnabled(bool bEnabled);
 
 		/**
 		 * Load the default Optimus deformer graph to use on the skeletal mesh when using this model.
 		 * @return The mesh deformer graph.
 		 */
-		virtual UMeshDeformer* LoadDefaultDeformerGraph() const;
+		UE_API virtual UMeshDeformer* LoadDefaultDeformerGraph() const;
 
 		/**
 		 * This makes sure that the deformer graph selected in the visualization settings is not a nullptr.
 		 * If the current graph hasn't been set, it will assign the default deformer graph to it.
 		 */
-		virtual void SetDefaultDeformerGraphIfNeeded();
+		UE_API virtual void SetDefaultDeformerGraphIfNeeded();
 
 		/**
 		 * This applies the right mesh deformer to the editor actors that have an ML Deformer component on them.
 		 * Basically this just ensures that the heat map deformer graph is active when this is enabled. Otherwise the default
 		 * mesh deformer is used.
 		 */
-		virtual void UpdateDeformerGraph();
+		UE_API virtual void UpdateDeformerGraph();
 
 		/**
 		 * Executed when the maximum number of LOD levels changed.
@@ -640,7 +642,7 @@ namespace UE::MLDeformer
 		 * So this will update the internal state of the Sampler member. It will use the current training frame number to sample from.
 		 * If this number is somehow out of range it will automatically clamp it to be inside a valid range.
 		 */
-		virtual void SampleDeltas();
+		UE_API virtual void SampleDeltas();
 
 		/**
 		 * Load the trained network from an onnx file.
@@ -655,7 +657,7 @@ namespace UE::MLDeformer
 		 * Where "intermediatedir" is your project's Intermediate folder. And "YourModel" is the name of your runtime model class.
 		 * So if you have a model class named UFancyModel, on default this method will return "<intermediatedir>/FancyModel/FancyModel.onnx"
 		 */
-		virtual FString GetTrainedNetworkOnnxFile() const;
+		UE_API virtual FString GetTrainedNetworkOnnxFile() const;
 
 		/**
 		 * Check whether our model is trained or not.
@@ -669,31 +671,31 @@ namespace UE::MLDeformer
 		 * In training mode this is the training target actor, while in test mode it is the test ground truth actor.
 		 * @return A pointer to the ground truth actor depending on whether we're in test or training mode.
 		 */
-		virtual FMLDeformerEditorActor* GetTimelineEditorActor() const;
+		UE_API virtual FMLDeformerEditorActor* GetTimelineEditorActor() const;
 
 		/**
 		 * Get the path to the heat map material asset.
 		 * @return The path to the heat map asset.
 		 */
-		virtual FString GetHeatMapMaterialPath() const;
+		UE_API virtual FString GetHeatMapMaterialPath() const;
 
 		/**
 		 * Get the path to the heat map deformer graph asset, when using linear skinning.
 		 * @return The path to the heat map deformer graph asset.
 		 */
-		virtual FString GetHeatMapDeformerGraphPath() const;
+		UE_API virtual FString GetHeatMapDeformerGraphPath() const;
 
 		/**
 		 * Get the path to the heat map deformer graph asset, when using dual quaternion skinning.
 		 * @return The path to the heat map deformer graph asset.
 		 */
-		virtual FString GetHeatMapDeformerGraphDualQuatPath() const;
+		UE_API virtual FString GetHeatMapDeformerGraphDualQuatPath() const;
 
 		/**
 		 * Refresh the contents of the input widget.
 		 * This includes the bone and curve list contents.
 		 */
-		virtual void RefreshInputWidget();
+		UE_API virtual void RefreshInputWidget();
 
 		/**
 		 * Create the input widget.
@@ -702,10 +704,10 @@ namespace UE::MLDeformer
 		 * This can be useful when you want to add new types of inputs to your model.
 		 * @return A shared pointer to the input widget that will appear in the UI.
 		 */
-		virtual TSharedPtr<SMLDeformerInputWidget> CreateInputWidget();
+		UE_API virtual TSharedPtr<SMLDeformerInputWidget> CreateInputWidget();
 
 		/** Update the persona preview scene's actor, skeletal mesh component, mesh and anim asset. These are based on the current visualization mode. */
-		virtual void UpdatePreviewScene();
+		UE_API virtual void UpdatePreviewScene();
 
 		/**
 		 * Generate the normals for a given morph target.
@@ -718,7 +720,7 @@ namespace UE::MLDeformer
 		 * @param OutDeltaNormals The array that we will write the generated normals to. This will automatically be resized by this method.
 		 */
 		UE_DEPRECATED(5.3, "Please call FMLDeformerMorphModelEditorModel::CalcMorphTargetNormals instead.")
-		virtual void GenerateNormalsForMorphTarget(int32 LOD, USkeletalMesh* SkelMesh, int32 MorphTargetIndex, TArrayView<const FVector3f> Deltas, TArrayView<const FVector3f> BaseVertexPositions, TArrayView<FVector3f> BaseNormals, TArray<FVector3f>& OutDeltaNormals);
+		UE_API virtual void GenerateNormalsForMorphTarget(int32 LOD, USkeletalMesh* SkelMesh, int32 MorphTargetIndex, TArrayView<const FVector3f> Deltas, TArrayView<const FVector3f> BaseVertexPositions, TArrayView<FVector3f> BaseNormals, TArray<FVector3f>& OutDeltaNormals);
 
 		/**
 		 * Called whenever the actor to debug has been changed.
@@ -726,13 +728,13 @@ namespace UE::MLDeformer
 		 * When debugging gets disabled, this will contain a nullptr.
 		 * @param DebugActor The new actor we want to debug.
 		 */
-		virtual void OnDebugActorChanged(TObjectPtr<AActor> DebugActor);
+		UE_API virtual void OnDebugActorChanged(TObjectPtr<AActor> DebugActor);
 
 		/** 
 		 * Apply the transforms of the actor we debug to the editor actors in our world.
 		 * @param DebugActorComponentSpaceTransforms The component space transforms from the actor we are debugging.
 		 */
-		virtual void ApplyDebugActorTransforms(const TArray<FTransform>& DebugActorComponentSpaceTransforms);
+		UE_API virtual void ApplyDebugActorTransforms(const TArray<FTransform>& DebugActorComponentSpaceTransforms);
 
 		/**
 		 * Update the character's pose in paint mode.
@@ -743,7 +745,7 @@ namespace UE::MLDeformer
 		 * @param bFullUpdate When set to true, this can internally update the acceleration structures, which is slow. When set to false it can quickly update the visual mesh
 		 *                    It is still important to once call it with bFullUpdate set to true though, before you start painting.
 		 */
-		virtual void UpdatePaintModePose(bool bFullUpdate=true);
+		UE_API virtual void UpdatePaintModePose(bool bFullUpdate=true);
 
 		/**
 		 * Debug draw helpers inside the PIE viewport that highlight debuggable actors.
@@ -751,7 +753,7 @@ namespace UE::MLDeformer
 		 * Also it will render the actor names.
 		 * NOTE: This renders inside the PIE viewport, not our own MLD asset editor viewport.
 		 */
-		virtual void DrawPIEDebugActors();
+		UE_API virtual void DrawPIEDebugActors();
 
 		/**
 		 * Check whether this model supports a per training input animation vertex mask.
@@ -764,7 +766,7 @@ namespace UE::MLDeformer
 		 * Update LOD levels of the actors in editor world.
 		 * This can be used to for example sync the LOD levels of the compare actors with the main actor.
 		 */
-		virtual void UpdateActorLODs();
+		UE_API virtual void UpdateActorLODs();
 
 		/**
 		 * This should update the list of available training devices.
@@ -801,68 +803,68 @@ namespace UE::MLDeformer
 		virtual void UpdateTrainingDeviceList() {}
 
 		/** Apply the transforms of the debug actor to the actors in the asset editor world. This will internally call ApplyDebugActorTransforms(DebugActorComponentSpaceTransforms). */
-		void ApplyDebugActorTransforms();
+		UE_API void ApplyDebugActorTransforms();
 
 		/** Invalidate the memory usage, so it gets updated in the UI again. */
-		void UpdateMemoryUsage();
+		UE_API void UpdateMemoryUsage();
 
 		/** Get the currently active training input animation sequence, which is the one that the viewport is showing the timeline for in training mode. */
-		UAnimSequence* GetActiveTrainingInputAnimSequence() const;
+		UE_API UAnimSequence* GetActiveTrainingInputAnimSequence() const;
 
 		/** Get the current view range. */
-		TRange<double> GetViewRange() const;
+		UE_API TRange<double> GetViewRange() const;
 
 		/** Set the current view range. */
-		void SetViewRange(TRange<double> InRange);
+		UE_API void SetViewRange(TRange<double> InRange);
 
 		/** Get the working range of the model's data. */
-		TRange<double> GetWorkingRange() const;
+		UE_API TRange<double> GetWorkingRange() const;
 
 		/** Get the playback range of the model's data. */
-		TRange<FFrameNumber> GetPlaybackRange() const;
+		UE_API TRange<FFrameNumber> GetPlaybackRange() const;
 
 		/** Get the current scrub position. */
-		FFrameNumber GetTickResScrubPosition() const;
+		UE_API FFrameNumber GetTickResScrubPosition() const;
 
 		/** Get how many ticks there are per frame. */
-		int32 GetTicksPerFrame() const;
+		UE_API int32 GetTicksPerFrame() const;
 
 		/** Get the current scrub time. */
-		float GetScrubTime() const;
+		UE_API float GetScrubTime() const;
 
 		/** Set the current scrub position. */
 		UE_DEPRECATED(5.5, "Please use the SetScrubPosition which takes a bIsScrubbing parameter.")
-		void SetScrubPosition(FFrameTime NewScrubPosition);
-		void SetScrubPosition(FFrameTime NewScrubPosition, bool bIsScrubbing);
+		UE_API void SetScrubPosition(FFrameTime NewScrubPosition);
+		UE_API void SetScrubPosition(FFrameTime NewScrubPosition, bool bIsScrubbing);
 
 		/** Set the current scrub position. */
 		UE_DEPRECATED(5.5, "Please use the SetScrubPosition which takes a bIsScrubbing parameter.")
-		void SetScrubPosition(FFrameNumber NewScrubPosition);
-		void SetScrubPosition(FFrameNumber NewScrubPosition, bool bIsScrubbing);
+		UE_API void SetScrubPosition(FFrameNumber NewScrubPosition);
+		UE_API void SetScrubPosition(FFrameNumber NewScrubPosition, bool bIsScrubbing);
 
 		/** Set if frames are displayed. */
-		void SetDisplayFrames(bool bDisplayFrames);
+		UE_API void SetDisplayFrames(bool bDisplayFrames);
 
 		/** Is the timeline displaying frames rather than seconds? */
-		bool IsDisplayingFrames() const;
+		UE_API bool IsDisplayingFrames() const;
 
 		/** Handle the timeline changes when our model changes. */
-		void HandleModelChanged();
+		UE_API void HandleModelChanged();
 
 		/** Handle the timeline changes required when changing the visualization mode, so when we switch between training and testing mode. */
-		void HandleVizModeChanged(EMLDeformerVizMode Mode);
+		UE_API void HandleVizModeChanged(EMLDeformerVizMode Mode);
 
 		/** Handle the view range being changed. */
-		void HandleViewRangeChanged(TRange<double> InRange);
+		UE_API void HandleViewRangeChanged(TRange<double> InRange);
 
 		/** Handle the working range being changed. */
-		void HandleWorkingRangeChanged(TRange<double> InRange);
+		UE_API void HandleWorkingRangeChanged(TRange<double> InRange);
 
 		/** Get the framerate specified by the anim sequence. */
-		double GetFrameRate() const;
+		UE_API double GetFrameRate() const;
 
 		/** Get the tick resolution we are displaying at. */
-		int32 GetTickResolution() const;
+		UE_API int32 GetTickResolution() const;
 
 		/** Get a pointer to the ML Deformer editor tookit. Cannot be a nullptr. */
 		FMLDeformerEditorToolkit* GetEditor() const { return Editor; }
@@ -871,20 +873,20 @@ namespace UE::MLDeformer
 		UMLDeformerModel* GetModel() const { return Model.Get(); }
 
 		/** Get a pointer to the world that our editor is displaying. Cannot be nullptr. */
-		UWorld* GetWorld() const;
+		UE_API UWorld* GetWorld() const;
 
 		/** Get all the created editor actors. */
 		const TArray<FMLDeformerEditorActor*>& GetEditorActors() const { return EditorActors; }
 
 		/** Find an editor actor with a specific ID, or return nullptr when not found. */
-		FMLDeformerEditorActor* FindEditorActor(int32 ActorTypeID) const;
+		UE_API FMLDeformerEditorActor* FindEditorActor(int32 ActorTypeID) const;
 
 		/** Check whether we are ready to train. The Train button in the UI will be enabled or disabled based on this return value. */
 		bool IsReadyForTraining() const { return bIsReadyForTraining; }
 
 		/** Get the sampler we use to calculate vertex deltas when we are in training mdoe and enable deltas. */
 		UE_DEPRECATED(5.4, "Please use GetSamplerForTrainingAnim instead.")
-		FMLDeformerSampler* GetSampler() const;
+		UE_API FMLDeformerSampler* GetSampler() const;
 
 		/**
 		 * Get the vertex delta sampler for a given training input animation.
@@ -911,27 +913,27 @@ namespace UE::MLDeformer
 		 * Check whether the vertex counts of the skeletal mesh have changed compared to the one we have trained on.
 		 * @return An error string if there is an error, otherwise an empty text object is returned.
 		 */
-		FText GetBaseAssetChangedErrorText() const;
+		UE_API FText GetBaseAssetChangedErrorText() const;
 
 		/**
 		 * Check whether the vertex map has changed or not.
 		 * This basically checks whether the vertex order has changed.
 		 * @return An error string if there is an error, otherwise an empty text object is returned.
 		 */
-		FText GetVertexMapChangedErrorText() const;
+		UE_API FText GetVertexMapChangedErrorText() const;
 
 		/**
 		 * Check for errors in the inputs. This basically checks if there are bones or curves to train on.
 		 * If there aren't any it will give some error.
 		 * @return An error string if there is an error, otherwise an empty text object is returned.
 		 */
-		FText GetInputsErrorText() const;
+		UE_API FText GetInputsErrorText() const;
 
 		/**
 		 * Check for incompatible skeletons between the skeletal mesh and anim sequence skeleton.
 		 * @return An error string if there is an error, otherwise an empty text object is returned.
 		 */
-		FText GetIncompatibleSkeletonErrorText(const USkeletalMesh* InSkelMesh, const UAnimSequence* InAnimSeq) const;
+		UE_API FText GetIncompatibleSkeletonErrorText(const USkeletalMesh* InSkelMesh, const UAnimSequence* InAnimSeq) const;
 
 		/**
 		 * Check to see if the skeletal mesh has to be reimported because it misses some newly added data.
@@ -939,14 +941,14 @@ namespace UE::MLDeformer
 		 * to make the ML Deformer work properly.
 		 * @return An error string if there is an error, otherwise an empty text object is returned.
 		 */
-		FText GetSkeletalMeshNeedsReimportErrorText() const;
+		UE_API FText GetSkeletalMeshNeedsReimportErrorText() const;
 
 		/**
 		 * Check whether the vertex count in the target mesh has changed or not.
 		 * If this is the case you would need to re-train the model.
 		 * @return An error string if there is an error, otherwise an empty text object is returned.
 		 */
-		FText GetTargetAssetChangedErrorText() const;
+		UE_API FText GetTargetAssetChangedErrorText() const;
 
 		/**
 		 * Get the editor's input info object.
@@ -959,45 +961,45 @@ namespace UE::MLDeformer
 		/**
 		 * Update the editor's input info based on the current UI settings.
 		 */
-		void UpdateEditorInputInfo();
+		UE_API void UpdateEditorInputInfo();
 
 		/**
 		 * This will call OnInputAssetsChanged first, followed by OnPostInputAssetChanged.
 		 * It then will update the editor only flags for the assets, and then forces a refresh of the model details panel and optionally also the visualization settings panel.
 		 * @param bRefreshVizSettings Set to true if you want the visualization settings details panel to also refresh.
 		 */
-		void TriggerInputAssetChanged(bool bRefreshVizSettings=false);
+		UE_API void TriggerInputAssetChanged(bool bRefreshVizSettings=false);
 
 		/**
 		 * Initialize the bone include list of the runtime model to the set of bones that are animated.
 		 * A bone is seen as animated if its animation data has keyframes that have a different value than the first keyframe.
 		 * So if the bone's animation data has 1000 keyframes, and they are all the same value, it won't see it as animated.
 		 */
-		void InitBoneIncludeListToAnimatedBonesOnly();
+		UE_API void InitBoneIncludeListToAnimatedBonesOnly();
 
 		/**
 		 * Initialize the curve include list of the runtime model to the set of curves that are animated.
 		 * A curve is seen as animated if its animation data has keyframes that have a different value than the first keyframe.
 		 * So if the curve's animation data has 1000 keyframes, and they are all the same value, it won't see it as animated.
 		 */
-		void InitCurveIncludeListToAnimatedCurvesOnly();
+		UE_API void InitCurveIncludeListToAnimatedCurvesOnly();
 
 		/**
 		 * Add all animated bones to the bone include list.
 		 */
-		void AddAnimatedBonesToBonesIncludeList();
+		UE_API void AddAnimatedBonesToBonesIncludeList();
 
 		/**
 		 * Add all animated curves to the curve include list.
 		 */
-		void AddAnimatedCurvesToCurvesIncludeList();
+		UE_API void AddAnimatedCurvesToCurvesIncludeList();
 
 		/**
 		 * Get the number of curves on a specific skeletal mesh.
 		 * @param SkelMesh The skeletal mesh to get the curves count from.
 		 * @return The number of curves found on the skeletal mesh.
 		 */
-		int32 GetNumCurvesOnSkeletalMesh(USkeletalMesh* SkelMesh) const;
+		UE_API int32 GetNumCurvesOnSkeletalMesh(USkeletalMesh* SkelMesh) const;
 
 		/**
 		 * Get the currently desired training frame number.
@@ -1010,7 +1012,7 @@ namespace UE::MLDeformer
 		 * Check whether our training data frame number changed.
 		 * We have a TrainingFrameNumber and a CurrentTrainingFrame. If they don't match it will call OnTrainingDataFrameChanged.
 		 */
-		void CheckTrainingDataFrameChanged();
+		UE_API void CheckTrainingDataFrameChanged();
 
 		/**
 		 * Debug draw specific morph targets using lines and points.
@@ -1023,14 +1025,14 @@ namespace UE::MLDeformer
 		 * @param DrawOffset An offset to perform the debug draw at.
 		 */
 		UE_DEPRECATED(5.3, "Please call FMLDeformerMorphModelEditorModel::DebugDrawMorphTargets instead.")
-		void DrawMorphTarget(FPrimitiveDrawInterface* PDI, const TArray<FVector3f>& MorphDeltas, float DeltaThreshold, int32 MorphTargetIndex, const FVector& DrawOffset);
+		UE_API void DrawMorphTarget(FPrimitiveDrawInterface* PDI, const TArray<FVector3f>& MorphDeltas, float DeltaThreshold, int32 MorphTargetIndex, const FVector& DrawOffset);
 
 		/**
 		 * Find the ML Deformer component of a given editor actor.
 		 * @param ActorID The actor ID to get the component for.
 		 * @return A pointer to the ML Deformer component, if it exists on the actor, otherwise nullptr is returned.
 		 */
-		UMLDeformerComponent* FindMLDeformerComponent(int32 ActorID = ActorID_Test_MLDeformed) const;
+		UE_API UMLDeformerComponent* FindMLDeformerComponent(int32 ActorID = ActorID_Test_MLDeformed) const;
 
 		/**
 		 * Correct floating point errors that can cause issues when sampling animation using step timing.
@@ -1039,24 +1041,24 @@ namespace UE::MLDeformer
 		 * @param FrameRate The frame rate for conversion.
 		 * @return A corrected frame time, which is equal to FrameTimeForFrameNumber unless a floating point error means the floor is not equal to the FrameNumber.
 		 */
-		static float CorrectedFrameTime(int32 FrameNumber, float FrameTimeForFrameNumber, FFrameRate FrameRate);
+		static UE_API float CorrectedFrameTime(int32 FrameNumber, float FrameTimeForFrameNumber, FFrameRate FrameRate);
 
 		TSharedPtr<SMLDeformerInputWidget> GetInputWidget() const				{ return InputWidget; }
 		void SetInputWidget(const TSharedPtr<SMLDeformerInputWidget>& Widget)	{ InputWidget = Widget; }
 
-		FMLDeformerSampler* GetSamplerForActiveAnim() const;
+		UE_API FMLDeformerSampler* GetSamplerForActiveAnim() const;
 
-		void SetActiveTrainingInputAnimIndex(int32 Index);
-		int32 GetActiveTrainingInputAnimIndex() const;
+		UE_API void SetActiveTrainingInputAnimIndex(int32 Index);
+		UE_API int32 GetActiveTrainingInputAnimIndex() const;
 
 		/** Mark the deltas to be updated on next Tick. */
-		void InvalidateDeltas();
+		UE_API void InvalidateDeltas();
 
 		/** 
 		 * Find the float based vertex attributes on the skeletal map that have a given name.
 		 * You can check whether the attribute data has been found using ReturnedValue.IsValid().
 		 */
-		TVertexAttributesConstRef<float> FindVertexAttributes(FName AttributeName) const;
+		UE_API TVertexAttributesConstRef<float> FindVertexAttributes(FName AttributeName) const;
 
 		/** Are we scrubbing on the timeline? */
 		bool IsScrubbingTimeline() const									{ return bIsScrubbingTimeline; }
@@ -1065,7 +1067,7 @@ namespace UE::MLDeformer
 		const TArray<TSharedPtr<FName>>& GetVertexAttributeNames() const	{ return VertexAttributeNames; }
 
 		/** Update the vertex attribute names array as returned by GetVertexAttributeNames(). This is some array that can be used for UI elements like combo boxes. */
-		void UpdateVertexAttributeNames();
+		UE_API void UpdateVertexAttributeNames();
 
 		/**
 		 * Request a call to TriggerInputAssetsChanged on the next tick, so not directly.
@@ -1074,22 +1076,22 @@ namespace UE::MLDeformer
 		void RequestTriggerInputAssetsChangedNextTick()						{ bNeedsAssetReinit = true; }
 
 	protected:
-		virtual void CreateSamplers();
+		UE_API virtual void CreateSamplers();
 
-		void AddAnimatedBonesToBonesIncludeList(const UAnimSequence* AnimSequence);
-		void AddAnimatedCurvesToCurvesIncludeList(const UAnimSequence* AnimSequence);
+		UE_API void AddAnimatedBonesToBonesIncludeList(const UAnimSequence* AnimSequence);
+		UE_API void AddAnimatedCurvesToCurvesIncludeList(const UAnimSequence* AnimSequence);
 
 		/**
 		 * Executed when a property changes by a change in the UI.
 		 * This internally calls OnPropertyChanged.
 		 */
-		void OnPostEditChangeProperty(FPropertyChangedEvent& PropertyChangedEvent);
+		UE_API void OnPostEditChangeProperty(FPropertyChangedEvent& PropertyChangedEvent);
 
 		/**
 		 * Trigger a PostEditChangeProperty using a ValueSet event type, on the model, which will also trigger the OnPostEditChangeProperty.
 		 * @param PropertyName The name of the property to trigger it for.
 		 */
-		void BroadcastPropertyChanged(const FName PropertyName);
+		UE_API void BroadcastPropertyChanged(const FName PropertyName);
 
 		/**
 		 * Called whenever an engine object is being modified.
@@ -1098,19 +1100,19 @@ namespace UE::MLDeformer
 		 * The default implementation of this method will check for changes in the skeletal mesh and training and test anim sequence.
 		 * If you want to check more things (your target mesh), you need to check for that.
 		 */
-		virtual void OnObjectModified(UObject* Object);
+		UE_API virtual void OnObjectModified(UObject* Object);
 
 		/**
 		 * Delete all editor actors.
 		 * This does not actually delete the actors in the world.
 		 * You can call ClearWorld to include that.
 		 */
-		void DeleteEditorActors();
+		UE_API void DeleteEditorActors();
 
 		/**
 		 * Set the names of the anim sequences in the timeline. 
 		 */
-		void SetTimelineAnimNames(const TArray<TSharedPtr<FMLDeformerTrainingInputAnimName>>& AnimNames);
+		UE_API void SetTimelineAnimNames(const TArray<TSharedPtr<FMLDeformerTrainingInputAnimName>>& AnimNames);
 
 		/**
 		 * Perform some basic checks to see if the editor can be ready to train the model.
@@ -1118,11 +1120,11 @@ namespace UE::MLDeformer
 		 * It can be called inside the UpdateIsReadyForTrainingState method, which can then do additional checks afterwards.
 		 * @return Returns true when at least the basic required inputs have been selected, false if not.
 		 */
-		bool IsEditorReadyForTrainingBasicChecks();
+		UE_API bool IsEditorReadyForTrainingBasicChecks();
 
 		/** Zero all deltas with a length equal to, or smaller than the threshold value. */
 		UE_DEPRECATED(5.3, "Please call FMLDeformerMorphModelEditorModel::ZeroDeltasByLengthThreshold instead.")
-		void ZeroDeltasByThreshold(TArray<FVector3f>& Deltas, float Threshold);
+		UE_API void ZeroDeltasByThreshold(TArray<FVector3f>& Deltas, float Threshold);
 
 		/**
 		 * Generate engine morph targets from a set of deltas.
@@ -1137,7 +1139,7 @@ namespace UE::MLDeformer
 		 * @param bInvertMaskChannel Specifies whether the weight mask should be inverted or not.
 		 */
 		UE_DEPRECATED(5.3, "Please call FMLDeformerMorphModelEditorModel::CreateMorphTargets instead.")
-		void CreateEngineMorphTargets(
+		UE_API void CreateEngineMorphTargets(
 			TArray<UMorphTarget*>& OutMorphTargets,
 			const TArray<FVector3f>& Deltas,
 			const FString& NamePrefix = TEXT("MorphTarget_"),
@@ -1155,7 +1157,7 @@ namespace UE::MLDeformer
 		 * @param MorphErrorTolerance The error tolerance for the delta compression, in cm. Higher values compress better but can result in artifacts.
 		 */
 		UE_DEPRECATED(5.3, "Please call FMLDeformerMorphModelEditorModel::CompressMorphTargets instead.")
-		void CompressEngineMorphTargets(FMorphTargetVertexInfoBuffers& OutMorphBuffers, const TArray<UMorphTarget*>& MorphTargets, int32 LOD = 0, float MorphErrorTolerance = 0.01f);
+		UE_API void CompressEngineMorphTargets(FMorphTargetVertexInfoBuffers& OutMorphBuffers, const TArray<UMorphTarget*>& MorphTargets, int32 LOD = 0, float MorphErrorTolerance = 0.01f);
 
 		/**
 		 * Calculate the normals for each vertex, given the triangle data and positions.
@@ -1166,50 +1168,50 @@ namespace UE::MLDeformer
 		 * @param OutNormals The array that will contain the normals. This will automatically be resized internally by this method.
 		 */
 		UE_DEPRECATED(5.3, "Please call FMLDeformerMorphModelEditorModel::CalcVertexNormals instead.")
-		void CalcMeshNormals(TArrayView<const FVector3f> VertexPositions, TArrayView<const uint32> IndexArray, TArrayView<const int32> VertexMap, TArray<FVector3f>& OutNormals) const;
+		UE_API void CalcMeshNormals(TArrayView<const FVector3f> VertexPositions, TArrayView<const uint32> IndexArray, TArrayView<const int32> VertexMap, TArray<FVector3f>& OutNormals) const;
 
 		/**
 		 * Get the base actor depending on the currently active mode (train or test mode).
 		 * The base actor is the linear-skinned one.
 		 * @return A pointer to the linear skinned actor in either test or training mode.
 		 */
-		FMLDeformerEditorActor* GetVisualizationModeBaseActor() const;
+		UE_API FMLDeformerEditorActor* GetVisualizationModeBaseActor() const;
 
 		/**
 		 * Get the animation sequence, either the test or training one, depending on in which mode we are in.
 		 * This can return a nullptr if no anim sequence has been set up yet.
 		 * @return A pointer to the test sequence when we are in test mode, or the training anim sequence if we are in training mode.
 		 */
-		const UAnimSequence* GetAnimSequence() const;
+		UE_API const UAnimSequence* GetAnimSequence() const;
 
 		/**
 		 * Calculate the current timeline position offset, of the currently selected frame.
 		 * This depends on whether we are in test or training mode.
 		 * @return The timeline offset, in seconds.
 		 */
-		double CalcTimelinePosition() const;
+		UE_API double CalcTimelinePosition() const;
 
 		/**
 		 * Update the timeline related ranges, based on the length of the training or testing data.
 		 */
-		void UpdateRanges();
+		UE_API void UpdateRanges();
 
-		void AddCompareActor(int32 ArrayIndex);
-		void RemoveCompareActor(int32 ArrayIndex);
-		void RemoveAllCompareActors();
-		void UpdateCompareActorLabels();		
+		UE_API void AddCompareActor(int32 ArrayIndex);
+		UE_API void RemoveCompareActor(int32 ArrayIndex);
+		UE_API void RemoveAllCompareActors();
+		UE_API void UpdateCompareActorLabels();		
 
-		int32 CalcNumValidCompareActorsPriorTo(int32 CompareActorIndex) const;
-		virtual bool IsAnimIndexValid(int32 AnimIndex) const;
+		UE_API int32 CalcNumValidCompareActorsPriorTo(int32 CompareActorIndex) const;
+		UE_API virtual bool IsAnimIndexValid(int32 AnimIndex) const;
 
-		void UpdateLODMappings();
+		UE_API void UpdateLODMappings();
 
 		/**
 		 * Update whether we force to use step interpolation on the ML Deformed model or not.
 		 * Step interpolation is needed when we are enabled heatmaps and have a ground truth mesh and are in ground truth heatmap mode.
 		 * This will modify the AnimInstance's ForceStepInterpolation state on just the ML Deformed model. It can also disable it when needed.
 		 */
-		void UpdateStepInterpolationMode();
+		UE_API void UpdateStepInterpolationMode();
 
 		/**
 		 * Because of some technical reasons we cannot display a correct ground truth heat map at the moment while playing an animation.
@@ -1217,7 +1219,7 @@ namespace UE::MLDeformer
 		 * This method will look whether we are having heatmaps enabled or not, and if we are playing or not and if we're in ground truth mode.
 		 * It then disables or enables the heatmap material based on the current state.
 		 */
-		void UpdateHeatMapMaterialBasedOnMode();
+		UE_API void UpdateHeatMapMaterialBasedOnMode();
 
 	protected:
 		struct FLODInfo
@@ -1355,3 +1357,5 @@ namespace UE::MLDeformer
 		return static_cast<ETrainingResult>(ReturnCode);
 	}
 }	// namespace UE::MLDeformer
+
+#undef UE_API
diff --git a/Engine/Plugins/Animation/MLDeformer/MLDeformerFramework/Source/MLDeformerFrameworkEditor/Public/MLDeformerEditorModule.h b/Engine/Plugins/Animation/MLDeformer/MLDeformerFramework/Source/MLDeformerFrameworkEditor/Public/MLDeformerEditorModule.h
--- a/Engine/Plugins/Animation/MLDeformer/MLDeformerFramework/Source/MLDeformerFrameworkEditor/Public/MLDeformerEditorModule.h
+++ b/Engine/Plugins/Animation/MLDeformer/MLDeformerFramework/Source/MLDeformerFrameworkEditor/Public/MLDeformerEditorModule.h
@@ -6,6 +6,8 @@
 #include "Modules/ModuleInterface.h"
 #include "MLDeformerModelRegistry.h"
 
+#define UE_API MLDEFORMERFRAMEWORKEDITOR_API
+
 namespace UE::MLDeformer
 {
 	namespace TrainingDataProcessor
@@ -18,13 +20,13 @@ namespace UE::MLDeformer
 	 * This registers the editor mode and custom property customizations.
 	 * It also contains the model registry, which you will use to register your custom models to the editor.
 	 */
-	class MLDEFORMERFRAMEWORKEDITOR_API FMLDeformerEditorModule
+	class FMLDeformerEditorModule
 		: public IModuleInterface
 	{
 	public:
 		// IModuleInterface overrides.
-		virtual void StartupModule() override;
-		virtual void ShutdownModule() override;
+		UE_API virtual void StartupModule() override;
+		UE_API virtual void ShutdownModule() override;
 		// ~END IModuleInterface overrides.
 
 		/**
@@ -47,3 +49,5 @@ namespace UE::MLDeformer
 		FMLDeformerEditorModelRegistry ModelRegistry;
 	};
 }	// namespace UE::MLDeformer
+
+#undef UE_API
diff --git a/Engine/Plugins/Animation/MLDeformer/MLDeformerFramework/Source/MLDeformerFrameworkEditor/Public/MLDeformerEditorStyle.h b/Engine/Plugins/Animation/MLDeformer/MLDeformerFramework/Source/MLDeformerFrameworkEditor/Public/MLDeformerEditorStyle.h
--- a/Engine/Plugins/Animation/MLDeformer/MLDeformerFramework/Source/MLDeformerFrameworkEditor/Public/MLDeformerEditorStyle.h
+++ b/Engine/Plugins/Animation/MLDeformer/MLDeformerFramework/Source/MLDeformerFrameworkEditor/Public/MLDeformerEditorStyle.h
@@ -4,18 +4,22 @@
 
 #include "Styling/SlateStyle.h"
 
+#define UE_API MLDEFORMERFRAMEWORKEDITOR_API
+
 namespace UE::MLDeformer
 {
 	/**
 	 * The editor style class that describes specific UI style related settings for the ML Deformer editor.
 	 */
-	class MLDEFORMERFRAMEWORKEDITOR_API FMLDeformerEditorStyle
+	class FMLDeformerEditorStyle
 		: public FSlateStyleSet
 	{
 	public:
-		FMLDeformerEditorStyle();
-		~FMLDeformerEditorStyle();
+		UE_API FMLDeformerEditorStyle();
+		UE_API ~FMLDeformerEditorStyle();
 
-		static FMLDeformerEditorStyle& Get();
+		static UE_API FMLDeformerEditorStyle& Get();
 	};
 }	// namespace UE::MLDeformer
+
+#undef UE_API
diff --git a/Engine/Plugins/Animation/MLDeformer/MLDeformerFramework/Source/MLDeformerFrameworkEditor/Public/MLDeformerEditorToolkit.h b/Engine/Plugins/Animation/MLDeformer/MLDeformerFramework/Source/MLDeformerFrameworkEditor/Public/MLDeformerEditorToolkit.h
--- a/Engine/Plugins/Animation/MLDeformer/MLDeformerFramework/Source/MLDeformerFrameworkEditor/Public/MLDeformerEditorToolkit.h
+++ b/Engine/Plugins/Animation/MLDeformer/MLDeformerFramework/Source/MLDeformerFrameworkEditor/Public/MLDeformerEditorToolkit.h
@@ -15,6 +15,8 @@
 #include "Math/Color.h"
 #include "Widgets/Notifications/SNotificationList.h"
 
+#define UE_API MLDEFORMERFRAMEWORKEDITOR_API
+
 struct FMenuEntryParams;
 class FTabManager;
 class IDetailsView;
@@ -33,7 +35,7 @@ namespace UE::MLDeformer
 	}
 
 	class FMLDeformerEditorToolkit;
-	class MLDEFORMERFRAMEWORKEDITOR_API FToolsMenuExtender
+	class FToolsMenuExtender
 	{
 	public:
 		virtual ~FToolsMenuExtender() {}
@@ -45,7 +47,7 @@ namespace UE::MLDeformer
 	 * The ML Deformer asset editor toolkit.
 	 * This is the editor that opens when you double click an ML Deformer asset.
 	 */
-	class MLDEFORMERFRAMEWORKEDITOR_API FMLDeformerEditorToolkit :
+	class FMLDeformerEditorToolkit :
 		public FPersonaAssetEditorToolkit,
 		public IHasPersonaToolkit,
 		public FGCObject,
@@ -56,54 +58,54 @@ namespace UE::MLDeformer
 		friend class FMLDeformerApplicationMode;
 		friend struct FMLDeformerVizSettingsTabSummoner;
 
-		~FMLDeformerEditorToolkit();
+		UE_API ~FMLDeformerEditorToolkit();
 
 		/** Initialize the asset editor. This will register the application mode, init the preview scene, etc. */
-		void InitAssetEditor(
+		UE_API void InitAssetEditor(
 			const EToolkitMode::Type Mode,
 			const TSharedPtr<IToolkitHost>& InitToolkitHost,
 			UMLDeformerAsset* InDeformerAsset);
 
 		// FAssetEditorToolkit overrides.
-		virtual void RegisterTabSpawners(const TSharedRef<FTabManager>& InTabManager) override;
-		virtual void UnregisterTabSpawners(const TSharedRef<FTabManager>& InTabManager) override;
-		virtual FName GetToolkitFName() const override;
-		virtual FText GetBaseToolkitName() const override;
-		virtual FText GetToolkitName() const override;
-		virtual FLinearColor GetWorldCentricTabColorScale() const override;
-		virtual FString GetWorldCentricTabPrefix() const override;
-		virtual void AddViewportOverlayWidget(TSharedRef<SWidget> InViewportOverlayWidget, int32 ZOrder = INDEX_NONE) override;
-		virtual void RemoveViewportOverlayWidget(TSharedRef<SWidget> InViewportOverlayWidget) override;
+		UE_API virtual void RegisterTabSpawners(const TSharedRef<FTabManager>& InTabManager) override;
+		UE_API virtual void UnregisterTabSpawners(const TSharedRef<FTabManager>& InTabManager) override;
+		UE_API virtual FName GetToolkitFName() const override;
+		UE_API virtual FText GetBaseToolkitName() const override;
+		UE_API virtual FText GetToolkitName() const override;
+		UE_API virtual FLinearColor GetWorldCentricTabColorScale() const override;
+		UE_API virtual FString GetWorldCentricTabPrefix() const override;
+		UE_API virtual void AddViewportOverlayWidget(TSharedRef<SWidget> InViewportOverlayWidget, int32 ZOrder = INDEX_NONE) override;
+		UE_API virtual void RemoveViewportOverlayWidget(TSharedRef<SWidget> InViewportOverlayWidget) override;
 		// ~END FAssetEditorToolkit overrides.
 
 		// FGCObject overrides.
 		virtual FString GetReferencerName() const override							{ return TEXT("FMLDeformerEditorToolkit"); }
-		virtual void AddReferencedObjects(FReferenceCollector& Collector) override;
+		UE_API virtual void AddReferencedObjects(FReferenceCollector& Collector) override;
 		// ~END FGCObject overrides.
 
 		// FTickableEditorObject overrides.
-		virtual void Tick(float DeltaTime) override;
+		UE_API virtual void Tick(float DeltaTime) override;
 		virtual ETickableTickType GetTickableTickType() const override				{ return ETickableTickType::Always; }
-		virtual TStatId GetStatId() const override;
+		UE_API virtual TStatId GetStatId() const override;
 		// ~END FTickableEditorObject overrides.
 
 		// FEditorUndoClient overrides.
-		void PostUndo(bool bSuccess) override;
-		void PostRedo(bool bSuccess) override;
+		UE_API void PostUndo(bool bSuccess) override;
+		UE_API void PostRedo(bool bSuccess) override;
 		// ~END FEditorUndoClient overrides.
 
 		// IHasPersonaToolkit overrides.
-		virtual TSharedRef<IPersonaToolkit> GetPersonaToolkit() const override;
+		UE_API virtual TSharedRef<IPersonaToolkit> GetPersonaToolkit() const override;
 		IPersonaToolkit* GetPersonaToolkitPointer() const							{ return PersonaToolkit.Get(); }
 		// ~END IHasPersonaToolkit overrides.
 
 		void SetVizSettingsDetailsView(TSharedPtr<IDetailsView> InDetailsView)		{ VizSettingsDetailsView = InDetailsView; }
-		IDetailsView* GetVizSettingsDetailsView() const;
+		UE_API IDetailsView* GetVizSettingsDetailsView() const;
 
-		IDetailsView* GetModelDetailsView() const;
+		UE_API IDetailsView* GetModelDetailsView() const;
 
-		void SetTimeSlider(TSharedPtr<SMLDeformerTimeline> InTimeSlider);
-		SMLDeformerTimeline* GetTimeSlider() const;
+		UE_API void SetTimeSlider(TSharedPtr<SMLDeformerTimeline> InTimeSlider);
+		UE_API SMLDeformerTimeline* GetTimeSlider() const;
 
 		UMLDeformerAsset* GetDeformerAsset() const									{ return DeformerAsset.Get(); }
 		FMLDeformerEditorModel* GetActiveModel()									{ return ActiveModel.Get(); }
@@ -114,15 +116,15 @@ namespace UE::MLDeformer
 
 		void SetNeedsPaintModeDisable(bool bNeedsDisable)							{ bNeedsPaintModeDisable = bNeedsDisable; }
 
-		double CalcTimelinePosition() const;
-		void OnTimeSliderScrubPositionChanged(double NewScrubTime, bool bIsScrubbing);
-		void UpdateTimeSliderRange();
-		void SetTimeSliderRange(double StartTime, double EndTime);
+		UE_API double CalcTimelinePosition() const;
+		UE_API void OnTimeSliderScrubPositionChanged(double NewScrubTime, bool bIsScrubbing);
+		UE_API void UpdateTimeSliderRange();
+		UE_API void SetTimeSliderRange(double StartTime, double EndTime);
 
-		void EnablePaintMode();
-		void DisablePaintMode();
-		bool IsDefaultModeActive() const;
-		bool IsPaintModeActive() const;
+		UE_API void EnablePaintMode();
+		UE_API void DisablePaintMode();
+		UE_API bool IsDefaultModeActive() const;
+		UE_API bool IsPaintModeActive() const;
 
 		bool IsInitialized() const { return bIsInitialized; }
 
@@ -132,65 +134,65 @@ namespace UE::MLDeformer
 		 * @param bForceChange Force changing to this model? This will suppress any UI popups.
 		 * @return Returns true in case we successfully switched model types, or otherwise false is returned.
 		 */
-		bool SwitchModelType(UClass* ModelType, bool bForceChange);
+		UE_API bool SwitchModelType(UClass* ModelType, bool bForceChange);
 
 		/**
 		 * Switch the editor's visualization mode.
 		 * This essentially allows you to switch the UI between testing and training modes.
 		 * @param Mode The mode to switch to.
 		 */
-		void SwitchVizMode(EMLDeformerVizMode Mode);
+		UE_API void SwitchVizMode(EMLDeformerVizMode Mode);
 
-		bool Train(bool bSuppressDialogs);
-		bool IsTrainButtonEnabled() const;
-		bool IsTraining() const;
+		UE_API bool Train(bool bSuppressDialogs);
+		UE_API bool IsTrainButtonEnabled() const;
+		UE_API bool IsTraining() const;
 
 		/** Get the actor we want to debug, if any. Returns a nullptr when we don't want to debug anything. */
-		AActor* GetDebugActor() const;
+		UE_API AActor* GetDebugActor() const;
 
 		/** Get the component space transforms of the actor we want to debug. Returns an empty array if GetDebugActor returns a nullptr. */
-		TArray<FTransform> GetDebugActorComponentSpaceTransforms() const;
+		UE_API TArray<FTransform> GetDebugActorComponentSpaceTransforms() const;
 
-		void ZoomOnActors();
+		UE_API void ZoomOnActors();
 
 		TSharedPtr<SMLDeformerDebugSelectionWidget> GetDebugWidget() const { return DebugWidget; }
 
-		static void AddToolsMenuExtender(TUniquePtr<FToolsMenuExtender> Extender);
-		static TConstArrayView<TUniquePtr<FToolsMenuExtender>> GetToolsMenuExtenders();
+		static UE_API void AddToolsMenuExtender(TUniquePtr<FToolsMenuExtender> Extender);
+		static UE_API TConstArrayView<TUniquePtr<FToolsMenuExtender>> GetToolsMenuExtenders();
 
 	private:
 		UE_DEPRECATED(5.3, "Please use the OnModelChanged that takes two parameters instead.")
-		void OnModelChanged(int Index);
-		void OnModelChanged(int Index, bool bForceChange);
+		UE_API void OnModelChanged(int Index);
+		UE_API void OnModelChanged(int Index, bool bForceChange);
 
 		UE_DEPRECATED(5.3, "Please use the SwitchVizMode instead.")
-		void OnVizModeChanged(EMLDeformerVizMode Mode);
+		UE_API void OnVizModeChanged(EMLDeformerVizMode Mode);
 
 		/* Toolbar related. */
-		void ExtendToolbar();
-		void FillToolbar(FToolBarBuilder& ToolbarBuilder);
+		UE_API void ExtendToolbar();
+		UE_API void FillToolbar(FToolBarBuilder& ToolbarBuilder);
 
 		/** Preview scene setup. */
-		void HandlePreviewSceneCreated(const TSharedRef<IPersonaPreviewScene>& InPersonaPreviewScene);
-		void HandleViewportCreated(const TSharedRef<IPersonaViewport>& InPersonaViewport);
-		void HandleDetailsCreated(const TSharedRef<class IDetailsView>& InDetailsView);
-		void OnFinishedChangingDetails(const FPropertyChangedEvent& PropertyChangedEvent);
+		UE_API void HandlePreviewSceneCreated(const TSharedRef<IPersonaPreviewScene>& InPersonaPreviewScene);
+		UE_API void HandleViewportCreated(const TSharedRef<IPersonaViewport>& InPersonaViewport);
+		UE_API void HandleDetailsCreated(const TSharedRef<class IDetailsView>& InDetailsView);
+		UE_API void OnFinishedChangingDetails(const FPropertyChangedEvent& PropertyChangedEvent);
 
 		/** Helpers. */
-		void ShowNotification(const FText& Message, SNotificationItem::ECompletionState State, bool PlaySound) const;
-		FText GetOverlayText() const;
-		void OnSwitchedVisualizationMode();
-		bool HandleTrainingResult(ETrainingResult TrainingResult, double TrainingDuration, bool& bOutUsePartiallyTrained, bool bSuppressDialogs, bool& bOutSuccess);
-		FText GetActiveModelName() const;
-		FText GetCurrentVizModeName() const;
-		FText GetVizModeName(EMLDeformerVizMode Mode) const;
-		void ShowNoModelsWarningIfNeeded();
-		EVisibility GetDebuggingVisibility() const;
-		FSlateIcon GetTrainingStatusIcon() const;
-
-		TSharedRef<SWidget> GenerateModelButtonContents(TSharedRef<FUICommandList> InCommandList);
-		TSharedRef<SWidget> GenerateVizModeButtonContents(TSharedRef<FUICommandList> InCommandList);
-		TSharedRef<SWidget> GenerateToolsMenuContents(TSharedRef<FUICommandList> InCommandList);
+		UE_API void ShowNotification(const FText& Message, SNotificationItem::ECompletionState State, bool PlaySound) const;
+		UE_API FText GetOverlayText() const;
+		UE_API void OnSwitchedVisualizationMode();
+		UE_API bool HandleTrainingResult(ETrainingResult TrainingResult, double TrainingDuration, bool& bOutUsePartiallyTrained, bool bSuppressDialogs, bool& bOutSuccess);
+		UE_API FText GetActiveModelName() const;
+		UE_API FText GetCurrentVizModeName() const;
+		UE_API FText GetVizModeName(EMLDeformerVizMode Mode) const;
+		UE_API void ShowNoModelsWarningIfNeeded();
+		UE_API EVisibility GetDebuggingVisibility() const;
+		UE_API FSlateIcon GetTrainingStatusIcon() const;
+
+		UE_API TSharedRef<SWidget> GenerateModelButtonContents(TSharedRef<FUICommandList> InCommandList);
+		UE_API TSharedRef<SWidget> GenerateVizModeButtonContents(TSharedRef<FUICommandList> InCommandList);
+		UE_API TSharedRef<SWidget> GenerateToolsMenuContents(TSharedRef<FUICommandList> InCommandList);
 
 
 	private:
@@ -235,9 +237,11 @@ namespace UE::MLDeformer
 		bool bNeedsPaintModeDisable = false;
 
 		/** Extenders for Tools menu */
-		static TArray<TUniquePtr<FToolsMenuExtender>> ToolsMenuExtenders;
+		static UE_API TArray<TUniquePtr<FToolsMenuExtender>> ToolsMenuExtenders;
 
 		/** Mutex for adding extenders */
-		static FCriticalSection ExtendersMutex;
+		static UE_API FCriticalSection ExtendersMutex;
 	};
 }	// namespace UE::MLDeformer
+
+#undef UE_API
diff --git a/Engine/Plugins/Animation/MLDeformer/MLDeformerFramework/Source/MLDeformerFrameworkEditor/Public/MLDeformerGeomCacheActor.h b/Engine/Plugins/Animation/MLDeformer/MLDeformerFramework/Source/MLDeformerFrameworkEditor/Public/MLDeformerGeomCacheActor.h
--- a/Engine/Plugins/Animation/MLDeformer/MLDeformerFramework/Source/MLDeformerFrameworkEditor/Public/MLDeformerGeomCacheActor.h
+++ b/Engine/Plugins/Animation/MLDeformer/MLDeformerFramework/Source/MLDeformerFrameworkEditor/Public/MLDeformerGeomCacheActor.h
@@ -4,6 +4,8 @@
 
 #include "MLDeformerEditorActor.h"
 
+#define UE_API MLDEFORMERFRAMEWORKEDITOR_API
+
 class UMLDeformerComponent;
 class UGeometryCacheComponent;
 
@@ -13,25 +15,25 @@ namespace UE::MLDeformer
 	 * An editor actor with a geometry cache component on it.
 	 * This can for example be used as ground truth viewport actors.
 	 */
-	class MLDEFORMERFRAMEWORKEDITOR_API FMLDeformerGeomCacheActor
+	class FMLDeformerGeomCacheActor
 		: public FMLDeformerEditorActor
 	{
 	public:
-		FMLDeformerGeomCacheActor(const FConstructSettings& Settings);
+		UE_API FMLDeformerGeomCacheActor(const FConstructSettings& Settings);
 
 		void SetGeometryCacheComponent(UGeometryCacheComponent* Component)	{ GeomCacheComponent = Component; }
 		UGeometryCacheComponent* GetGeometryCacheComponent() const			{ return GeomCacheComponent; }
 
 		// FMLDeformerEditorActor overrides.
-		virtual void SetVisibility(bool bIsVisible) override;
-		virtual bool IsVisible() const override;
-		virtual bool HasVisualMesh() const override;
-		virtual void SetPlayPosition(float TimeInSeconds, bool bAutoPause = true) override;
-		virtual float GetPlayPosition() const override;
-		virtual void SetPlaySpeed(float PlaySpeed) override;
-		virtual void Pause(bool bPaused) override;
-		virtual bool IsPlaying() const override;
-		virtual FBox GetBoundingBox() const override;
+		UE_API virtual void SetVisibility(bool bIsVisible) override;
+		UE_API virtual bool IsVisible() const override;
+		UE_API virtual bool HasVisualMesh() const override;
+		UE_API virtual void SetPlayPosition(float TimeInSeconds, bool bAutoPause = true) override;
+		UE_API virtual float GetPlayPosition() const override;
+		UE_API virtual void SetPlaySpeed(float PlaySpeed) override;
+		UE_API virtual void Pause(bool bPaused) override;
+		UE_API virtual bool IsPlaying() const override;
+		UE_API virtual FBox GetBoundingBox() const override;
 		// ~END FMLDeformerEditorActor overrides.
 
 	protected:
@@ -39,3 +41,5 @@ namespace UE::MLDeformer
 		TObjectPtr<UGeometryCacheComponent> GeomCacheComponent = nullptr;
 	};
 }	// namespace UE::MLDeformer
+
+#undef UE_API
diff --git a/Engine/Plugins/Animation/MLDeformer/MLDeformerFramework/Source/MLDeformerFrameworkEditor/Public/MLDeformerGeomCacheEditorModel.h b/Engine/Plugins/Animation/MLDeformer/MLDeformerFramework/Source/MLDeformerFrameworkEditor/Public/MLDeformerGeomCacheEditorModel.h
--- a/Engine/Plugins/Animation/MLDeformer/MLDeformerFramework/Source/MLDeformerFrameworkEditor/Public/MLDeformerGeomCacheEditorModel.h
+++ b/Engine/Plugins/Animation/MLDeformer/MLDeformerFramework/Source/MLDeformerFrameworkEditor/Public/MLDeformerGeomCacheEditorModel.h
@@ -5,6 +5,8 @@
 #include "CoreTypes.h"
 #include "MLDeformerEditorModel.h"
 
+#define UE_API MLDEFORMERFRAMEWORKEDITOR_API
+
 class UMLDeformerGeomCacheModel;
 class UMLDeformerGeomCacheVizSettings;
 
@@ -18,7 +20,7 @@ namespace UE::MLDeformer
 	 * This editor model automatically handles sampling of the geometry caches to calculate deltas.
 	 * It also handles the creation of ground truth models with a geometry cache on them.
 	 */
-	class MLDEFORMERFRAMEWORKEDITOR_API FMLDeformerGeomCacheEditorModel
+	class FMLDeformerGeomCacheEditorModel
 		: public FMLDeformerEditorModel
 	{
 	public:
@@ -27,30 +29,30 @@ namespace UE::MLDeformer
 		// ~END FGCObject overrides.
 
 		// FMLDeformerEditorModel overrides.
-		virtual void Init(const InitSettings& Settings) override;
-		virtual void CopyBaseSettingsFromModel(const FMLDeformerEditorModel* SourceEditorModel) override;
-		virtual FMLDeformerEditorActor* CreateEditorActor(const FMLDeformerEditorActor::FConstructSettings& Settings) const override;
-		virtual TSharedPtr<FMLDeformerSampler> CreateSamplerObject() const override;
-		virtual void CreateTrainingGroundTruthActor(UWorld* World) override;
-		virtual void CreateTestGroundTruthActor(UWorld* World) override;
-		virtual int32 GetNumTrainingInputAnims() const override;
-		virtual FMLDeformerTrainingInputAnim* GetTrainingInputAnim(int32 Index) const override;
-		virtual void UpdateNumTrainingFrames() override;
-		virtual void UpdateIsReadyForTrainingState() override;
-		virtual void OnPropertyChanged(FPropertyChangedEvent& PropertyChangedEvent) override;
-		virtual void OnInputAssetsChanged() override;
-		virtual void OnObjectModified(UObject* Object) override;
-		virtual ETrainingResult Train() override;
+		UE_API virtual void Init(const InitSettings& Settings) override;
+		UE_API virtual void CopyBaseSettingsFromModel(const FMLDeformerEditorModel* SourceEditorModel) override;
+		UE_API virtual FMLDeformerEditorActor* CreateEditorActor(const FMLDeformerEditorActor::FConstructSettings& Settings) const override;
+		UE_API virtual TSharedPtr<FMLDeformerSampler> CreateSamplerObject() const override;
+		UE_API virtual void CreateTrainingGroundTruthActor(UWorld* World) override;
+		UE_API virtual void CreateTestGroundTruthActor(UWorld* World) override;
+		UE_API virtual int32 GetNumTrainingInputAnims() const override;
+		UE_API virtual FMLDeformerTrainingInputAnim* GetTrainingInputAnim(int32 Index) const override;
+		UE_API virtual void UpdateNumTrainingFrames() override;
+		UE_API virtual void UpdateIsReadyForTrainingState() override;
+		UE_API virtual void OnPropertyChanged(FPropertyChangedEvent& PropertyChangedEvent) override;
+		UE_API virtual void OnInputAssetsChanged() override;
+		UE_API virtual void OnObjectModified(UObject* Object) override;
+		UE_API virtual ETrainingResult Train() override;
 		// ~END FMLDeformerEditorModel overrides.
 		
 		// Helpers.
 		UE_DEPRECATED(5.4, "Please use GetSamplerForTrainingAnim() or GetSamplerForActiveAnim() instead and cast those manually.")
 		FMLDeformerGeomCacheSampler* GetGeomCacheSampler() const { return nullptr; }
 
-		UMLDeformerGeomCacheModel* GetGeomCacheModel() const;
-		UMLDeformerGeomCacheVizSettings* GetGeomCacheVizSettings() const;
-		FMLDeformerGeomCacheActor* FindGeomCacheEditorActor(int32 ID) const;
-		UGeometryCache* GetActiveGeometryCache() const;
+		UE_API UMLDeformerGeomCacheModel* GetGeomCacheModel() const;
+		UE_API UMLDeformerGeomCacheVizSettings* GetGeomCacheVizSettings() const;
+		UE_API FMLDeformerGeomCacheActor* FindGeomCacheEditorActor(int32 ID) const;
+		UE_API UGeometryCache* GetActiveGeometryCache() const;
 
 
 	protected:
@@ -66,6 +68,8 @@ namespace UE::MLDeformer
 		 * @param LabelText The text of the label.
 		 * @param bIsTrainingActor Set this to true when the actor is an actor to be used in training mode, or set to false when it is to be used in testing mode.
 		 */
-		void CreateGeomCacheActor(UWorld* World, int32 ActorID, const FName& Name, UGeometryCache* GeomCache, FLinearColor LabelColor, FLinearColor WireframeColor, const FText& LabelText, bool bIsTrainingActor);
+		UE_API void CreateGeomCacheActor(UWorld* World, int32 ActorID, const FName& Name, UGeometryCache* GeomCache, FLinearColor LabelColor, FLinearColor WireframeColor, const FText& LabelText, bool bIsTrainingActor);
 	};
 }	// namespace UE::MLDeformer
+
+#undef UE_API
diff --git a/Engine/Plugins/Animation/MLDeformer/MLDeformerFramework/Source/MLDeformerFrameworkEditor/Public/MLDeformerGeomCacheModelDetails.h b/Engine/Plugins/Animation/MLDeformer/MLDeformerFramework/Source/MLDeformerFrameworkEditor/Public/MLDeformerGeomCacheModelDetails.h
--- a/Engine/Plugins/Animation/MLDeformer/MLDeformerFramework/Source/MLDeformerFrameworkEditor/Public/MLDeformerGeomCacheModelDetails.h
+++ b/Engine/Plugins/Animation/MLDeformer/MLDeformerFramework/Source/MLDeformerFrameworkEditor/Public/MLDeformerGeomCacheModelDetails.h
@@ -6,6 +6,8 @@
 #include "MLDeformerModelDetails.h"
 #include "IDetailCustomization.h"
 
+#define UE_API MLDEFORMERFRAMEWORKEDITOR_API
+
 class IDetailLayoutBuilder;
 class UMLDeformerGeomCacheModel;
 class UGeometryCache;
@@ -22,13 +24,13 @@ namespace UE::MLDeformer
 	 * It will automatically insert some error when the geom cache doesn't match the animation sequence etc.
 	 * If your model inherits from a UMLDeformerGeomCacheModel then you can also inherit your model detail customization from this FMLDeformerGeomCacheModelDetails class.
 	 */
-	class MLDEFORMERFRAMEWORKEDITOR_API FMLDeformerGeomCacheModelDetails
+	class FMLDeformerGeomCacheModelDetails
 		: public FMLDeformerModelDetails
 	{
 	public:
 		// FMLDeformerModelDetails overrides.
-		virtual bool UpdateMemberPointers(const TArray<TWeakObjectPtr<UObject>>& Objects) override;
-		virtual void AddTrainingInputAnims() override;
+		UE_API virtual bool UpdateMemberPointers(const TArray<TWeakObjectPtr<UObject>>& Objects) override;
+		UE_API virtual void AddTrainingInputAnims() override;
 		// ~END FMLDeformerModelDetails overrides.
 
 	protected:
@@ -44,3 +46,5 @@ namespace UE::MLDeformer
 		FMLDeformerGeomCacheEditorModel* GeomCacheEditorModel = nullptr;
 	};
 }	// namespace UE::MLDeformer
+
+#undef UE_API
diff --git a/Engine/Plugins/Animation/MLDeformer/MLDeformerFramework/Source/MLDeformerFrameworkEditor/Public/MLDeformerGeomCacheSampler.h b/Engine/Plugins/Animation/MLDeformer/MLDeformerFramework/Source/MLDeformerFrameworkEditor/Public/MLDeformerGeomCacheSampler.h
--- a/Engine/Plugins/Animation/MLDeformer/MLDeformerFramework/Source/MLDeformerFrameworkEditor/Public/MLDeformerGeomCacheSampler.h
+++ b/Engine/Plugins/Animation/MLDeformer/MLDeformerFramework/Source/MLDeformerFrameworkEditor/Public/MLDeformerGeomCacheSampler.h
@@ -9,6 +9,8 @@
 #include "UObject/ObjectPtr.h"
 #include "GeometryCacheMeshData.h"
 
+#define UE_API MLDEFORMERFRAMEWORKEDITOR_API
+
 class UGeometryCacheComponent;
 class UGeometryCache;
 
@@ -20,14 +22,14 @@ namespace UE::MLDeformer
 	 * The input data sampler, which is used to sample vertex positions from geometry caches.
 	 * It can then also calculate deltas between the sampled skeletal mesh data and geometry cache data.
 	 */
-	class MLDEFORMERFRAMEWORKEDITOR_API FMLDeformerGeomCacheSampler
+	class FMLDeformerGeomCacheSampler
 		: public FMLDeformerSampler
 	{
 	public:
 		// FVertexDeltaSampler overrides.
-		virtual void Init(FMLDeformerEditorModel* InModel, int32 InAnimIndex) override;
-		virtual void Sample(int32 AnimFrameIndex) override;
-		virtual float GetTimeAtFrame(int32 InAnimFrameIndex) const override;
+		UE_API virtual void Init(FMLDeformerEditorModel* InModel, int32 InAnimIndex) override;
+		UE_API virtual void Sample(int32 AnimFrameIndex) override;
+		UE_API virtual float GetTimeAtFrame(int32 InAnimFrameIndex) const override;
 		// ~END FVertexDeltaSampler overrides.
 
 		/**
@@ -86,3 +88,5 @@ namespace UE::MLDeformer
 		FMLDeformerGetGeomCacheEvent GetGeometryCacheEvent;
 	};
 }	// namespace UE::MLDeformer
+
+#undef UE_API
diff --git a/Engine/Plugins/Animation/MLDeformer/MLDeformerFramework/Source/MLDeformerFrameworkEditor/Public/MLDeformerGeomCacheTrainingInputAnimCustomize.h b/Engine/Plugins/Animation/MLDeformer/MLDeformerFramework/Source/MLDeformerFrameworkEditor/Public/MLDeformerGeomCacheTrainingInputAnimCustomize.h
--- a/Engine/Plugins/Animation/MLDeformer/MLDeformerFramework/Source/MLDeformerFrameworkEditor/Public/MLDeformerGeomCacheTrainingInputAnimCustomize.h
+++ b/Engine/Plugins/Animation/MLDeformer/MLDeformerFramework/Source/MLDeformerFrameworkEditor/Public/MLDeformerGeomCacheTrainingInputAnimCustomize.h
@@ -8,6 +8,8 @@
 #include "DetailWidgetRow.h"
 #include "AssetRegistry/AssetData.h"
 
+#define UE_API MLDEFORMERFRAMEWORKEDITOR_API
+
 class USkeleton;
 class UMLDeformerGeomCacheModel;
 
@@ -19,19 +21,21 @@ namespace UE::MLDeformer
 	 * That struct is shown in the list of training input animations in the details panel.
 	 * @see FMLDeformerGeomCacheTrainingInputAnim
 	 */
-	class MLDEFORMERFRAMEWORKEDITOR_API FMLDeformerGeomCacheTrainingInputAnimCustomization
+	class FMLDeformerGeomCacheTrainingInputAnimCustomization
 		: public IPropertyTypeCustomization
 	{
 	public:
-		static TSharedRef<IPropertyTypeCustomization> MakeInstance();
+		static UE_API TSharedRef<IPropertyTypeCustomization> MakeInstance();
 
 		// IPropertyTypeCustomization overrides.
-		virtual void CustomizeHeader(TSharedRef<IPropertyHandle> StructPropertyHandle, class FDetailWidgetRow& HeaderRow, IPropertyTypeCustomizationUtils& StructCustomizationUtils) override;
-		virtual void CustomizeChildren(TSharedRef<IPropertyHandle> StructPropertyHandle, class IDetailChildrenBuilder& ChildBuilder, IPropertyTypeCustomizationUtils& StructCustomizationUtils) override;
+		UE_API virtual void CustomizeHeader(TSharedRef<IPropertyHandle> StructPropertyHandle, class FDetailWidgetRow& HeaderRow, IPropertyTypeCustomizationUtils& StructCustomizationUtils) override;
+		UE_API virtual void CustomizeChildren(TSharedRef<IPropertyHandle> StructPropertyHandle, class IDetailChildrenBuilder& ChildBuilder, IPropertyTypeCustomizationUtils& StructCustomizationUtils) override;
 		// ~END IPropertyTypeCustomization overrides.
 
 	private:
-		bool FilterAnimSequences(const FAssetData& AssetData, USkeleton* Skeleton);
-		UMLDeformerGeomCacheModel* FindMLDeformerModel(TSharedRef<IPropertyHandle> StructPropertyHandle) const;
+		UE_API bool FilterAnimSequences(const FAssetData& AssetData, USkeleton* Skeleton);
+		UE_API UMLDeformerGeomCacheModel* FindMLDeformerModel(TSharedRef<IPropertyHandle> StructPropertyHandle) const;
 	};
 }	// namespace UE::MLDeformer
+
+#undef UE_API
diff --git a/Engine/Plugins/Animation/MLDeformer/MLDeformerFramework/Source/MLDeformerFrameworkEditor/Public/MLDeformerGeomCacheTrainingModel.h b/Engine/Plugins/Animation/MLDeformer/MLDeformerFramework/Source/MLDeformerFrameworkEditor/Public/MLDeformerGeomCacheTrainingModel.h
--- a/Engine/Plugins/Animation/MLDeformer/MLDeformerFramework/Source/MLDeformerFrameworkEditor/Public/MLDeformerGeomCacheTrainingModel.h
+++ b/Engine/Plugins/Animation/MLDeformer/MLDeformerFramework/Source/MLDeformerFrameworkEditor/Public/MLDeformerGeomCacheTrainingModel.h
@@ -7,27 +7,29 @@
 #include "MLDeformerTrainingModel.h"
 #include "MLDeformerGeomCacheTrainingModel.generated.h"
 
+#define UE_API MLDEFORMERFRAMEWORKEDITOR_API
+
 /**
  * The training model base class for geometry cache based models.
  * This class is our link to the Python training.
  */
-UCLASS(Blueprintable)
-class MLDEFORMERFRAMEWORKEDITOR_API UMLDeformerGeomCacheTrainingModel
+UCLASS(MinimalAPI, Blueprintable)
+class UMLDeformerGeomCacheTrainingModel
 	: public UMLDeformerTrainingModel
 {
 	GENERATED_BODY()
 
 public:
 	// UMLDeformerTrainingModel overrides.
-	virtual void Init(UE::MLDeformer::FMLDeformerEditorModel* InEditorModel) override;
-	virtual bool SampleNextFrame() override;
+	UE_API virtual void Init(UE::MLDeformer::FMLDeformerEditorModel* InEditorModel) override;
+	UE_API virtual bool SampleNextFrame() override;
 	// ~END UMLDeformerTrainingModel overrides;
 
 	/**
 	 * Sample all input frames.
 	 */
 	UFUNCTION(BlueprintCallable, Category = "ML Deformer Training")
-	bool GenerateBasicInputsAndOutputBuffers(const FString& InputsFilePath, const FString& OutputsFilePath);
+	UE_API bool GenerateBasicInputsAndOutputBuffers(const FString& InputsFilePath, const FString& OutputsFilePath);
 
 protected:
 	/**
@@ -37,5 +39,7 @@ protected:
 	 * @param OutNextAnimIndex The next animation index to sample from when we take our next sample.
 	 * @return Returns true when we found our next animation to sample. Returns false when we already sampled everything.
 	 */
-	virtual bool FindNextAnimToSample(int32& OutNextAnimIndex) const override;
+	UE_API virtual bool FindNextAnimToSample(int32& OutNextAnimIndex) const override;
 };
+
+#undef UE_API
diff --git a/Engine/Plugins/Animation/MLDeformer/MLDeformerFramework/Source/MLDeformerFrameworkEditor/Public/MLDeformerGeomCacheVizSettingsDetails.h b/Engine/Plugins/Animation/MLDeformer/MLDeformerFramework/Source/MLDeformerFrameworkEditor/Public/MLDeformerGeomCacheVizSettingsDetails.h
--- a/Engine/Plugins/Animation/MLDeformer/MLDeformerFramework/Source/MLDeformerFrameworkEditor/Public/MLDeformerGeomCacheVizSettingsDetails.h
+++ b/Engine/Plugins/Animation/MLDeformer/MLDeformerFramework/Source/MLDeformerFrameworkEditor/Public/MLDeformerGeomCacheVizSettingsDetails.h
@@ -5,6 +5,8 @@
 #include "CoreTypes.h"
 #include "MLDeformerVizSettingsDetails.h"
 
+#define UE_API MLDEFORMERFRAMEWORKEDITOR_API
+
 class UMLDeformerGeomCacheModel;
 class UMLDeformerGeomCacheVizSettings;
 
@@ -15,13 +17,13 @@ namespace UE::MLDeformer
 	 * A geometry cache based model is one inherited from the UMLDeformerGeomCacheModel base class.
 	 * You can inherit your own model detail customization from this class if you also inherited your model from the UMLDeformerGeomCacheModel class.
 	 */
-	class MLDEFORMERFRAMEWORKEDITOR_API FMLDeformerGeomCacheVizSettingsDetails
+	class FMLDeformerGeomCacheVizSettingsDetails
 		: public FMLDeformerVizSettingsDetails
 	{
 	public:
 		// FMLDeformerVizSettingsDetails overrides.
-		virtual bool UpdateMemberPointers(const TArray<TWeakObjectPtr<UObject>>& Objects) override;
-		virtual void AddGroundTruth() override;
+		UE_API virtual bool UpdateMemberPointers(const TArray<TWeakObjectPtr<UObject>>& Objects) override;
+		UE_API virtual void AddGroundTruth() override;
 		// ~END FMLDeformerVizSettingsDetails overrides.
 
 	protected:
@@ -32,3 +34,5 @@ namespace UE::MLDeformer
 		TObjectPtr<UMLDeformerGeomCacheVizSettings> GeomCacheVizSettings = nullptr;
 	};
 }	// namespace UE::MLDeformer
+
+#undef UE_API
diff --git a/Engine/Plugins/Animation/MLDeformer/MLDeformerFramework/Source/MLDeformerFrameworkEditor/Public/MLDeformerModelDetails.h b/Engine/Plugins/Animation/MLDeformer/MLDeformerFramework/Source/MLDeformerFrameworkEditor/Public/MLDeformerModelDetails.h
--- a/Engine/Plugins/Animation/MLDeformer/MLDeformerFramework/Source/MLDeformerFrameworkEditor/Public/MLDeformerModelDetails.h
+++ b/Engine/Plugins/Animation/MLDeformer/MLDeformerFramework/Source/MLDeformerFrameworkEditor/Public/MLDeformerModelDetails.h
@@ -7,6 +7,8 @@
 #include "AssetRegistry/AssetData.h"
 #include "Input/Reply.h"
 
+#define UE_API MLDEFORMERFRAMEWORKEDITOR_API
+
 class IDetailLayoutBuilder;
 class USkeleton;
 class IDetailCategoryBuilder;
@@ -25,23 +27,23 @@ namespace UE::MLDeformer
 	 * likely inherit from other classes such as FMLDeformerGeomCacheModelDetails or FMLDeformerMorphModelDetails. 
 	 * Those classes are also inherited from this one though.
 	 */
-	class MLDEFORMERFRAMEWORKEDITOR_API FMLDeformerModelDetails
+	class FMLDeformerModelDetails
 		: public IDetailCustomization
 	{
 	public:
 		// ILayoutDetails overrides.
-		virtual void CustomizeDetails(IDetailLayoutBuilder& DetailBuilder) override;
+		UE_API virtual void CustomizeDetails(IDetailLayoutBuilder& DetailBuilder) override;
 		// ~END ILayoutDetails overrides.
 
 		/**
 		 * Update the class member pointers, which includes the pointer to the model, and its editor model.
 		 * @param Objects The array of objects that the detail customization is showing.
 		 */
-		virtual bool UpdateMemberPointers(const TArray<TWeakObjectPtr<UObject>>& Objects);
+		UE_API virtual bool UpdateMemberPointers(const TArray<TWeakObjectPtr<UObject>>& Objects);
 
 		/**
 		 * Create the categories that we can add properties to. This will update the category member pointers. */
-		virtual void CreateCategories();
+		UE_API virtual void CreateCategories();
 
 		/**
 		 * Add the target mesh property. For example the FMLDeformerGeomCacheModelDetails class will add the geometry cache property here. */
@@ -66,7 +68,7 @@ namespace UE::MLDeformer
 		virtual void AddCurveInputErrors() {}
 
 		/** Add the property that shows the input animations. */
-		virtual void AddTrainingInputAnims();
+		UE_API virtual void AddTrainingInputAnims();
 
 		/** Add training input flags, which basically are things like the check boxes that specify whether bones or curves (or both) should be included. */
 		virtual void AddTrainingInputFlags() {}
@@ -91,7 +93,7 @@ namespace UE::MLDeformer
 
 	protected:
 		/** The filter that only shows anim sequences that are compatible with the given skeleton. */
-		bool FilterAnimSequences(const FAssetData& AssetData, USkeleton* Skeleton);
+		UE_API bool FilterAnimSequences(const FAssetData& AssetData, USkeleton* Skeleton);
 
 		/** Executed when we press the Animated Bones Only button. */
 		UE_DEPRECATED(5.4, "This method will be removed.")
@@ -127,3 +129,5 @@ namespace UE::MLDeformer
 		IDetailCategoryBuilder* LODSettingsCategoryBuilder = nullptr;
 	};
 }	// namespace UE::MLDeformer
+
+#undef UE_API
diff --git a/Engine/Plugins/Animation/MLDeformer/MLDeformerFramework/Source/MLDeformerFrameworkEditor/Public/MLDeformerModelRegistry.h b/Engine/Plugins/Animation/MLDeformer/MLDeformerFramework/Source/MLDeformerFrameworkEditor/Public/MLDeformerModelRegistry.h
--- a/Engine/Plugins/Animation/MLDeformer/MLDeformerFramework/Source/MLDeformerFrameworkEditor/Public/MLDeformerModelRegistry.h
+++ b/Engine/Plugins/Animation/MLDeformer/MLDeformerFramework/Source/MLDeformerFrameworkEditor/Public/MLDeformerModelRegistry.h
@@ -6,6 +6,8 @@
 #include "Containers/Map.h"
 #include "Delegates/Delegate.h"
 
+#define UE_API MLDEFORMERFRAMEWORKEDITOR_API
+
 class UMLDeformerModel;
 
 
@@ -21,13 +23,13 @@ namespace UE::MLDeformer
 	 * You typically do this in your editor module's StartupModule method. And you unregister it again in the editor module's ShutdownModule method.
 	 * Models can have priorities. The one model with the highest priority level will be the model that is created by default when creating a new ML Deformer asset.
 	 */
-	class MLDEFORMERFRAMEWORKEDITOR_API FMLDeformerEditorModelRegistry
+	class FMLDeformerEditorModelRegistry
 	{
 		friend class FMLDeformerEditorToolkit;
 		friend class FMLDeformerEditorModel;
 
 	public:
-		~FMLDeformerEditorModelRegistry();
+		UE_API ~FMLDeformerEditorModelRegistry();
 
 		/**
 		 * Register a new model. This is needed to make the model appear inside the UI of the editor as selectable model.
@@ -38,86 +40,86 @@ namespace UE::MLDeformer
 		 * @param ModelPriority The priority level of this model. The model with the highest priority is created by default when creating a new ML Deformer Asset.
 		 *        If there is a model with the same priority level already, and it happens to be the highest priority model, then this new model will take over as priority model.
 		 */
-		void RegisterEditorModel(UClass* ModelType, FOnGetEditorModelInstance Delegate, int32 ModelPriority = 0);
+		UE_API void RegisterEditorModel(UClass* ModelType, FOnGetEditorModelInstance Delegate, int32 ModelPriority = 0);
 
 		/**
 		 * Unregister the editor model.
 		 * Call this method in your editor module's ShutdownModule call.
 		 * @param ModelType The model type to unregister, for example UNeuralMorphModel::StaticClass().
 		 */
-		void UnregisterEditorModel(const UClass* ModelType);
+		UE_API void UnregisterEditorModel(const UClass* ModelType);
 
 		/**
 		 * Get the editor model related to a given runtime model. 
 		 * @param Model The runtime model to get the editor model object for.
 		 * @return A pointer to the editor model object for this specific runtime model instance, or nullptr when not found.
 		 */
-		FMLDeformerEditorModel* GetEditorModel(UMLDeformerModel* Model);
+		UE_API FMLDeformerEditorModel* GetEditorModel(UMLDeformerModel* Model);
 
 		/**
 		 * Get the number of registered models.
 		 * @return The number of models that have been registered through RegisterEditorModel.
 		 */
-		int32 GetNumRegisteredModels() const;
+		UE_API int32 GetNumRegisteredModels() const;
 
 		/**
 		 * Get the number of instanced models.
 		 * @return The number of instanced models.
 		 */
-		int32 GetNumInstancedModels() const;
+		UE_API int32 GetNumInstancedModels() const;
 
 		/**
 		 * Get the highest priority runtime model class.
 		 * The model with the highest priority is created by default when creating a new ML Deformer Asset.
 		 * @return Returns the model class of the highest priority model, or nullptr when there are none.
 		 */
-		UClass* GetHighestPriorityModel() const;
+		UE_API UClass* GetHighestPriorityModel() const;
 
 		/**
 		 * Get the priority level for a given runtime model type.
 		 * @param ModelType The runtime model type, for example UNeuralMorphModel::StaticClass().
 		 * @return The model priority level, or if it cannot be found it will return MIN_int32.
 		 */
-		int32 GetPriorityForModel(UClass* ModelType) const;
+		UE_API int32 GetPriorityForModel(UClass* ModelType) const;
 
 		/**
 		 * Get the read-only map of registered models.
 		 * @return The map of registered models that have been registered through RegisterEditorModel.
 		 */
-		const TMap<UClass*, FOnGetEditorModelInstance>& GetRegisteredModels() const;
+		UE_API const TMap<UClass*, FOnGetEditorModelInstance>& GetRegisteredModels() const;
 
 		/**
 		 * Get the read-only map that maps instanced models with their editor models.
 		 * @return The map that provides the editor model for a specific runtime model.
 		 */
-		const TMap<UMLDeformerModel*, FMLDeformerEditorModel*>& GetModelInstances() const;
+		UE_API const TMap<UMLDeformerModel*, FMLDeformerEditorModel*>& GetModelInstances() const;
 
 		/**
 		 * Get the highest priority model index into the map.
 		 * It is the 'index'-th item in the array that you can iterate by doing something like: for (auto& Item : Map).
 		 * @return The highest priority model index.
 		 */
-		int32 GetHighestPriorityModelIndex() const;
+		UE_API int32 GetHighestPriorityModelIndex() const;
 
 	private:
 		/**
 		 * Update the highest priority model by looking at the priority map.
 		 * This updates the value returned by GetHighestPriorityModel().
 		 */
-		void UpdateHighestPriorityModel();
+		UE_API void UpdateHighestPriorityModel();
 
 		/**
 		 * Remove a given editor model instance.
 		 * @param EditorModel A pointer to the editor model to remove.
 		 */
-		void RemoveEditorModelInstance(FMLDeformerEditorModel* EditorModel);
+		UE_API void RemoveEditorModelInstance(FMLDeformerEditorModel* EditorModel);
 
 		/**
 		 * Create a new editor model for a given runtime model.
 		 * @param Model The runtime model object to create an editor model for.
 		 * @return A pointer to the newly created editor model.
 		 */
-		FMLDeformerEditorModel* CreateEditorModel(UMLDeformerModel* Model);
+		UE_API FMLDeformerEditorModel* CreateEditorModel(UMLDeformerModel* Model);
 
 	private:
 		/** A map that maps runtime model types with a function that creates its editor model. */
@@ -133,3 +135,5 @@ namespace UE::MLDeformer
 		UClass* HighestPriorityModelType = nullptr;
 	};
 }	// namespace UE::MLDeformer
+
+#undef UE_API
diff --git a/Engine/Plugins/Animation/MLDeformer/MLDeformerFramework/Source/MLDeformerFrameworkEditor/Public/MLDeformerMorphModelDetails.h b/Engine/Plugins/Animation/MLDeformer/MLDeformerFramework/Source/MLDeformerFrameworkEditor/Public/MLDeformerMorphModelDetails.h
--- a/Engine/Plugins/Animation/MLDeformer/MLDeformerFramework/Source/MLDeformerFrameworkEditor/Public/MLDeformerMorphModelDetails.h
+++ b/Engine/Plugins/Animation/MLDeformer/MLDeformerFramework/Source/MLDeformerFrameworkEditor/Public/MLDeformerMorphModelDetails.h
@@ -6,6 +6,8 @@
 #include "MLDeformerGeomCacheModelDetails.h"
 #include "IDetailCustomization.h"
 
+#define UE_API MLDEFORMERFRAMEWORKEDITOR_API
+
 class IDetailLayoutBuilder;
 class UMLDeformerMorphModel;
 
@@ -17,17 +19,17 @@ namespace UE::MLDeformer
 	 * The detail customization for models inherited from the UMLDeformerMorphModel class.
 	 * You can inherit the detail customization for your own model from this class if your model inherited from the UMLDeformerMorphModel class.
 	 */
-	class MLDEFORMERFRAMEWORKEDITOR_API FMLDeformerMorphModelDetails
+	class FMLDeformerMorphModelDetails
 		: public FMLDeformerGeomCacheModelDetails
 	{
 	public:
 		// ILayoutDetails overrides.
-		virtual void CustomizeDetails(IDetailLayoutBuilder& DetailBuilder) override;
+		UE_API virtual void CustomizeDetails(IDetailLayoutBuilder& DetailBuilder) override;
 		// ~END ILayoutDetails overrides.
 
 		// FMLDeformerModelDetails overrides.
-		virtual void CreateCategories() override;
-		virtual bool UpdateMemberPointers(const TArray<TWeakObjectPtr<UObject>>& Objects) override;
+		UE_API virtual void CreateCategories() override;
+		UE_API virtual bool UpdateMemberPointers(const TArray<TWeakObjectPtr<UObject>>& Objects) override;
 		// ~END FMLDeformerModelDetails overrides.
 
 	protected:
@@ -43,7 +45,7 @@ namespace UE::MLDeformer
 		 * 
 		 * @return Returns true if an error should be displayed in the UI about this issue.
 		 */
-		bool ShouldShowShadingError() const;
+		UE_API bool ShouldShowShadingError() const;
 
 	protected:
 		/** A pointer to the morph model. This is updated when UpdateMemberPointers is called. */
@@ -56,3 +58,5 @@ namespace UE::MLDeformer
 		IDetailCategoryBuilder* MorphTargetCategoryBuilder = nullptr;
 	};
 }	// namespace UE::MLDeformer
+
+#undef UE_API
