Updated headers using UnrealCodeFixup to make sure dllstorage is on methods/staticvars instead of types. Part 1/n

diff --git a/Engine/Plugins/Animation/MLDeformer/MLDeformerFramework/Source/MLDeformerFrameworkEditor/Public/MLDeformerMorphModelEditorModel.h b/Engine/Plugins/Animation/MLDeformer/MLDeformerFramework/Source/MLDeformerFrameworkEditor/Public/MLDeformerMorphModelEditorModel.h
--- a/Engine/Plugins/Animation/MLDeformer/MLDeformerFramework/Source/MLDeformerFrameworkEditor/Public/MLDeformerMorphModelEditorModel.h
+++ b/Engine/Plugins/Animation/MLDeformer/MLDeformerFramework/Source/MLDeformerFrameworkEditor/Public/MLDeformerMorphModelEditorModel.h
@@ -9,6 +9,8 @@
 #include "MLDeformerMasking.h"
 #include "Rendering/ColorVertexBuffer.h"
 
+#define UE_API MLDEFORMERFRAMEWORKEDITOR_API
+
 class UMLDeformerMorphModel;
 class UMLDeformerMorphModelVizSettings;
 
@@ -17,27 +19,27 @@ namespace UE::MLDeformer
 	/**
 	 * The editor model related to the runtime UMLDeformerMorphModel class, or models inherited from that.
 	 */
-	class MLDEFORMERFRAMEWORKEDITOR_API FMLDeformerMorphModelEditorModel
+	class FMLDeformerMorphModelEditorModel
 		: public FMLDeformerGeomCacheEditorModel
 	{
 	public:
 		// We need to implement this static MakeInstance method.
-		static FMLDeformerEditorModel* MakeInstance();
+		static UE_API FMLDeformerEditorModel* MakeInstance();
 
 		// FGCObject overrides.
 		virtual FString GetReferencerName() const override { return TEXT("FMLDeformerMorphModelEditorModel"); }
 		// ~END FGCObject overrides.
 
 		// FMLDeformerEditorModel overrides.
-		virtual void CopyBaseSettingsFromModel(const FMLDeformerEditorModel* SourceEditorModel) override;
-		virtual void Render(const FSceneView* View, FViewport* Viewport, FPrimitiveDrawInterface* PDI) override;
-		virtual void OnPropertyChanged(FPropertyChangedEvent& PropertyChangedEvent) override;
-		virtual FString GetHeatMapDeformerGraphPath() const override;
-		virtual FString GetHeatMapDeformerGraphDualQuatPath() const override;
-		virtual void OnPreTraining() override;
-		virtual void OnPostTraining(ETrainingResult TrainingResult, bool bUsePartiallyTrainedWhenAborted) override;
-		virtual void OnMaxNumLODsChanged() override;
-		virtual void OnObjectModified(UObject* Object) override;
+		UE_API virtual void CopyBaseSettingsFromModel(const FMLDeformerEditorModel* SourceEditorModel) override;
+		UE_API virtual void Render(const FSceneView* View, FViewport* Viewport, FPrimitiveDrawInterface* PDI) override;
+		UE_API virtual void OnPropertyChanged(FPropertyChangedEvent& PropertyChangedEvent) override;
+		UE_API virtual FString GetHeatMapDeformerGraphPath() const override;
+		UE_API virtual FString GetHeatMapDeformerGraphDualQuatPath() const override;
+		UE_API virtual void OnPreTraining() override;
+		UE_API virtual void OnPostTraining(ETrainingResult TrainingResult, bool bUsePartiallyTrainedWhenAborted) override;
+		UE_API virtual void OnMaxNumLODsChanged() override;
+		UE_API virtual void OnObjectModified(UObject* Object) override;
 		// ~END FMLDeformerEditorModel overrides.
 
 		/**
@@ -53,7 +55,7 @@ namespace UE::MLDeformer
 		 * Get the mask buffer for a given morph target.
 		 * @param MorphTargetIndex The morph target index, excluding the first 'means' morph target that we always have, so 0 would mean the actual first real morph target.
 		 */
-		virtual const TArrayView<const float> GetMaskForMorphTarget(int32 MorphTargetIndex) const;
+		UE_API virtual const TArrayView<const float> GetMaskForMorphTarget(int32 MorphTargetIndex) const;
 
 		/**
 		 * Check if input masking is supported in this model.
@@ -61,7 +63,7 @@ namespace UE::MLDeformer
 		 * On default this is disabled.
 		 * @return Returns true if supported, or false otherwise.
 		 */
-		virtual bool IsInputMaskingSupported() const;
+		UE_API virtual bool IsInputMaskingSupported() const;
 
 		/**
 		 * Fill the mask buffer for a given item with a specific value.
@@ -70,7 +72,7 @@ namespace UE::MLDeformer
 		 * @param MaskBuffer The mask buffer to write to. The mask buffer view has Model->GetBaseNumVerts() elements.
 		 * @param Value The value to write for all vertices.
 		 */
-		void FillMaskValues(TArrayView<float> ItemMaskBuffer, float Value) const;
+		UE_API void FillMaskValues(TArrayView<float> ItemMaskBuffer, float Value) const;
 
 		/**
 		 * Add bone influences to the mask.
@@ -78,10 +80,10 @@ namespace UE::MLDeformer
 		 * @param SkeletonBoneIndex The bone index inside our reference skeleton.
 		 * @param MaskBuffer The mask buffer to write to.
 		 */
-		void ApplyBoneToMask(int32 SkeletonBoneIndex, TArrayView<float> MaskBuffer);
+		UE_API void ApplyBoneToMask(int32 SkeletonBoneIndex, TArrayView<float> MaskBuffer);
 
 		UE_DEPRECATED(5.4, "Please use the RecursiveAddBoneToMaskUpwards that doesn't take the VirtualParentTable as parameter.")
-		void RecursiveAddBoneToMaskUpwards(const FReferenceSkeleton& RefSkel, int32 SkeletonBoneIndex, int32 MaxHierarchyDepth, const TArray<int32>& VirtualParentTable, TArray<int32>& OutBonesAdded, int32 CurHierarchyDepth = 0);
+		UE_API void RecursiveAddBoneToMaskUpwards(const FReferenceSkeleton& RefSkel, int32 SkeletonBoneIndex, int32 MaxHierarchyDepth, const TArray<int32>& VirtualParentTable, TArray<int32>& OutBonesAdded, int32 CurHierarchyDepth = 0);
 
 		/**
 		 * Recursively add bone masks for this bone and its parents.
@@ -92,10 +94,10 @@ namespace UE::MLDeformer
 		 * @param OutBoneAdded The list of bones that we already added to the mask. Bones that are in this list will not be added to the mask again.
 		 * @param CurHierarchyDepth The current recursive hierarchy depth. This is used to track how deep we are when moving up the hierarchy. A value of 1 means we are currently adding the parent of the initial bone.
 		 */
-		void RecursiveAddBoneToMaskUpwards(const FReferenceSkeleton& RefSkel, int32 SkeletonBoneIndex, int32 MaxHierarchyDepth, TArray<int32>& OutBonesAdded, int32 CurHierarchyDepth = 0);
+		UE_API void RecursiveAddBoneToMaskUpwards(const FReferenceSkeleton& RefSkel, int32 SkeletonBoneIndex, int32 MaxHierarchyDepth, TArray<int32>& OutBonesAdded, int32 CurHierarchyDepth = 0);
 
 		UE_DEPRECATED(5.4, "Please use the RecursiveAddBoneToMaskDownwards that doesn't take the VirtualParentTable as parameter.")
-		void RecursiveAddBoneToMaskDownwards(const FReferenceSkeleton& RefSkel, int32 SkeletonBoneIndex, int32 MaxHierarchyDepth, const TArray<int32>& VirtualParentTable, TArray<int32>& OutBonesAdded, int32 CurHierarchyDepth = 0);
+		UE_API void RecursiveAddBoneToMaskDownwards(const FReferenceSkeleton& RefSkel, int32 SkeletonBoneIndex, int32 MaxHierarchyDepth, const TArray<int32>& VirtualParentTable, TArray<int32>& OutBonesAdded, int32 CurHierarchyDepth = 0);
 
 		/**
 		 * Recursively add bone masks for this bone and its children.
@@ -106,7 +108,7 @@ namespace UE::MLDeformer
 		 * @param OutBoneAdded The list of bones that we already added to the mask. Bones that are in this list will not be added to the mask again.
 		 * @param CurHierarchyDepth The current recursive hierarchy depth. This is used to track how deep we are when moving down the hierarchy. A value of 1 means we are currently adding a child node of the initial bone.
 		 */
-		void RecursiveAddBoneToMaskDownwards(const FReferenceSkeleton& RefSkel, int32 SkeletonBoneIndex, int32 MaxHierarchyDepth, TArray<int32>& OutBonesAdded, int32 CurHierarchyDepth = 0);
+		UE_API void RecursiveAddBoneToMaskDownwards(const FReferenceSkeleton& RefSkel, int32 SkeletonBoneIndex, int32 MaxHierarchyDepth, TArray<int32>& OutBonesAdded, int32 CurHierarchyDepth = 0);
 
 		/**
 		 * Add any required bone to the mask.
@@ -118,7 +120,7 @@ namespace UE::MLDeformer
 		 * @param OutBonesAdded The list of bones we added to the mask. This list can contain items already. Bones that are already in the list won't be added to the mask again.
 		 */
 		UE_DEPRECATED(5.4, "This method will be removed.")
-		void AddRequiredBones(const FReferenceSkeleton& RefSkel, int32 SkeletonBoneIndex, const TArray<int32>& VirtualParentTable, TArray<int32>& OutBonesAdded);
+		UE_API void AddRequiredBones(const FReferenceSkeleton& RefSkel, int32 SkeletonBoneIndex, const TArray<int32>& VirtualParentTable, TArray<int32>& OutBonesAdded);
 
 		/**
 		 * Find the virtual parent for a given bone.
@@ -130,7 +132,7 @@ namespace UE::MLDeformer
 		 * @see BuildVirtualParentTable.
 		 */
 		UE_DEPRECATED(5.4, "This method will be removed.")
-		int32 FindVirtualParentIndex(const FReferenceSkeleton& RefSkel, int32 BoneIndex, const TArray<FName>& IncludedBoneNames) const;
+		UE_API int32 FindVirtualParentIndex(const FReferenceSkeleton& RefSkel, int32 BoneIndex, const TArray<FName>& IncludedBoneNames) const;
 
 		/**
 		 * Build a mapping table that tells us the virtual parent for each bone in the ref skeleton.
@@ -143,7 +145,7 @@ namespace UE::MLDeformer
 		 * @return Returns a mapping array with the length of number of bones in the reference skeleton. Each array element contains a bone index inside the reference skeleton as well.
 		 */
 		UE_DEPRECATED(5.4, "This method will be removed.")
-		TArray<int32> BuildVirtualParentTable(const FReferenceSkeleton& RefSkel, const TArray<FName>& IncludedBoneNames) const;
+		UE_API TArray<int32> BuildVirtualParentTable(const FReferenceSkeleton& RefSkel, const TArray<FName>& IncludedBoneNames) const;
 
 		/**
 		 * Apply the mask info to the mask buffer of floats.
@@ -151,7 +153,7 @@ namespace UE::MLDeformer
 		 * @param MaskInfo The mask info object to apply to the float mask buffer.
 		 * @param ItemMaskBuffer The mask buffer for this item. This should be an array view of size NumBaseMeshVerts.
 		 */
-		void ApplyMaskInfoToBuffer(const USkeletalMesh* SkeletalMesh, const FMLDeformerMaskInfo& MaskInfo, TArrayView<float> ItemMaskBuffer);
+		UE_API void ApplyMaskInfoToBuffer(const USkeletalMesh* SkeletalMesh, const FMLDeformerMaskInfo& MaskInfo, TArrayView<float> ItemMaskBuffer);
 
 		/**
 		 * Apply a generated mask to a specific vertex attribute.
@@ -162,14 +164,14 @@ namespace UE::MLDeformer
 		 * @param MaskInfo The mask info object that we generate the mask from.
 		 * @param AttributeRef The output vertex attributes to write to.
 		 */
-		void ApplyGeneratedMaskToVertexAttributes(USkeletalMesh* SkeletalMesh, FMLDeformerMaskInfo& MaskInfo, TVertexAttributesRef<float> AttributeRef);
+		UE_API void ApplyGeneratedMaskToVertexAttributes(USkeletalMesh* SkeletalMesh, FMLDeformerMaskInfo& MaskInfo, TVertexAttributesRef<float> AttributeRef);
 
 		// Helpers.
-		UMLDeformerMorphModel* GetMorphModel() const;
-		UMLDeformerMorphModelVizSettings* GetMorphModelVizSettings() const;
+		UE_API UMLDeformerMorphModel* GetMorphModel() const;
+		UE_API UMLDeformerMorphModelVizSettings* GetMorphModelVizSettings() const;
 
 	protected:
-		void TransferMorphTargets(TArray<UMorphTarget*> MorphTargetsLODZero);
+		UE_API void TransferMorphTargets(TArray<UMorphTarget*> MorphTargetsLODZero);
 
 		/**
 		 * Initialize a set of engine morph targets and compress them to GPU friendly buffers.
@@ -178,20 +180,20 @@ namespace UE::MLDeformer
 		 * the first morph target, etc. In other words, the layout is: [morph0_deltas][morph1_deltas][morph2_deltas][...].
 		 * @param Deltas The deltas for all morph targets concatenated. So the number of items in this array is a multiple of Model->GetNumBaseVerts().
 		 */
-		virtual void InitEngineMorphTargets(const TArray<FVector3f>& Deltas);
+		UE_API virtual void InitEngineMorphTargets(const TArray<FVector3f>& Deltas);
 
 		/**
 		 * Clamp the current morph target number inside the visualization settings to a valid range.
 		 * This won't let the number of the morph target that we are visualizing go above the total number of morph targets.
 		 */
-		void ClampMorphTargetNumber();
+		UE_API void ClampMorphTargetNumber();
 
 		/**
 		 * Update the morph target error values. This assigns an error value for each morph target.
 		 * Internally this calls SetMorphTargetsErrorOrder to set the actual order and errors.
 		 * @param MorphTargets The morph targets to calculate the errors and order for.
 		 */
-		void UpdateMorphErrorValues(TArrayView<UMorphTarget*> MorphTargets);
+		UE_API void UpdateMorphErrorValues(TArrayView<UMorphTarget*> MorphTargets);
 
 		/**
 		 * Calculate the normals for a given morph target.
@@ -204,10 +206,10 @@ namespace UE::MLDeformer
 		 * @param OutDeltaNormals The array that we will write the generated normals to. This will automatically be resized by this method.
 		 */
 		UE_DEPRECATED(5.3, "Please use the CalcMorphTargetNormals method that takes more parameters.")
-		virtual void CalcMorphTargetNormals(int32 LOD, USkeletalMesh* SkelMesh, int32 MorphTargetIndex, TArrayView<const FVector3f> Deltas, TArrayView<const FVector3f> BaseVertexPositions, TArrayView<FVector3f> BaseNormals, TArray<FVector3f>& OutDeltaNormals);
+		UE_API virtual void CalcMorphTargetNormals(int32 LOD, USkeletalMesh* SkelMesh, int32 MorphTargetIndex, TArrayView<const FVector3f> Deltas, TArrayView<const FVector3f> BaseVertexPositions, TArrayView<FVector3f> BaseNormals, TArray<FVector3f>& OutDeltaNormals);
 
 		UE_DEPRECATED(5.5, "Please use the CalcMorphTaretNormals that takes an array named GlobalMaskWeights.")
-		virtual void CalcMorphTargetNormals(
+		UE_API virtual void CalcMorphTargetNormals(
 			int32 LOD,
 			const USkeletalMesh* SkelMesh,
 			int32 MorphTargetIndex,
@@ -235,7 +237,7 @@ namespace UE::MLDeformer
 		 * @param GlobalMaskWeights The array of weights, one for each render vertex.
 		 * @param OutDeltaNormals The array that we will write the generated normals to. This will automatically be resized by this method.
 		 */
-		virtual void CalcMorphTargetNormals(
+		UE_API virtual void CalcMorphTargetNormals(
 			int32 LOD,
 			const USkeletalMesh* SkelMesh,
 			int32 MorphTargetIndex,
@@ -259,10 +261,10 @@ namespace UE::MLDeformer
 		 * @param MorphTargetIndex The morph target number to visualize.
 		 * @param DrawOffset An offset to perform the debug draw at.
 		 */
-		void DebugDrawMorphTarget(FPrimitiveDrawInterface* PDI, const TArray<FVector3f>& MorphDeltas, float DeltaThreshold, int32 MorphTargetIndex, const FVector& DrawOffset);
+		UE_API void DebugDrawMorphTarget(FPrimitiveDrawInterface* PDI, const TArray<FVector3f>& MorphDeltas, float DeltaThreshold, int32 MorphTargetIndex, const FVector& DrawOffset);
 
 		/** Zero all deltas with a length equal to, or smaller than the threshold value. */
-		void ZeroDeltasByLengthThreshold(TArray<FVector3f>& Deltas, float Threshold);
+		UE_API void ZeroDeltasByLengthThreshold(TArray<FVector3f>& Deltas, float Threshold);
 
 		/**
 		 * Generate engine morph targets from a set of deltas.
@@ -278,7 +280,7 @@ namespace UE::MLDeformer
 		 * @param MaskBuffer An optional mask buffer that contains 'Model->GetNumBaseMeshVerts() * (MorphModel->GetNumMorphTargets() - 1)' number of floats. Deltas will be multiplied by this value. 
 		 *        When the mask buffer is an empty array, it will be ignored.
 		 */
-		void CreateMorphTargets(
+		UE_API void CreateMorphTargets(
 			TArray<UMorphTarget*>& OutMorphTargets,
 			const TArray<FVector3f>& Deltas,
 			const FString& NamePrefix = TEXT("MorphTarget_"),
@@ -295,7 +297,7 @@ namespace UE::MLDeformer
 		 * @param LOD The LOD index to generate the morphs for.
 		 * @param MorphErrorTolerance The error tolerance for the delta compression, in cm. Higher values compress better but can result in artifacts.
 		 */
-		void CompressMorphTargets(FMorphTargetVertexInfoBuffers& OutMorphBuffers, const TArray<UMorphTarget*>& MorphTargets, int32 LOD = 0, float MorphErrorTolerance = 0.01f);
+		UE_API void CompressMorphTargets(FMorphTargetVertexInfoBuffers& OutMorphBuffers, const TArray<UMorphTarget*>& MorphTargets, int32 LOD = 0, float MorphErrorTolerance = 0.01f);
 
 		/**
 		 * Calculate the normals for each vertex, given the triangle data and positions.
@@ -305,7 +307,7 @@ namespace UE::MLDeformer
 		 * @param VertexMap For each render vertex, an imported vertex number. For example, for a cube these indices go from 0..7.
 		 * @param OutNormals The array that will contain the normals. This will automatically be resized internally by this method.
 		 */
-		void CalcVertexNormals(TArrayView<const FVector3f> VertexPositions, TArrayView<const uint32> IndexArray, TArrayView<const int32> VertexMap, TArray<FVector3f>& OutNormals) const;
+		UE_API void CalcVertexNormals(TArrayView<const FVector3f> VertexPositions, TArrayView<const uint32> IndexArray, TArrayView<const int32> VertexMap, TArray<FVector3f>& OutNormals) const;
 
 		/**
 		 * Process and filter a vertex delta position and normal and apply scaling to it.
@@ -319,10 +321,10 @@ namespace UE::MLDeformer
 		 * @param GlobalMaskWeight The global weight value for this vertex. This global mask / weight is used to disable the deformer in specific regions on the mesh, for example where the head or neck connects with the body.
 		 * @return Returns true if this vertex should be included in the morph target, or false if it has been filtered out based on the provided delta threshold.
 		 */
-		bool ProcessVertexDelta(FVector3f& OutScaledDelta, FVector3f& OutScaledDeltaNormal, const FVector3f RawDelta, const FVector3f RawDeltaNormal, float DeltaThreshold, float MorphMaskWeight, float GlobalMaskWeight) const;
+		UE_API bool ProcessVertexDelta(FVector3f& OutScaledDelta, FVector3f& OutScaledDeltaNormal, const FVector3f RawDelta, const FVector3f RawDeltaNormal, float DeltaThreshold, float MorphMaskWeight, float GlobalMaskWeight) const;
 
 		UE_DEPRECATED(5.5, "Please use the CalcGlobalMaskWeights method instead")
-		float CalcGlobalMaskWeight(int32 RenderVertexIndex, const FColorVertexBuffer& ColorBuffer, EMLDeformerMaskChannel MaskChannel, bool bInvertMaskChannel) const;
+		UE_API float CalcGlobalMaskWeight(int32 RenderVertexIndex, const FColorVertexBuffer& ColorBuffer, EMLDeformerMaskChannel MaskChannel, bool bInvertMaskChannel) const;
 
 		/**
 		 * Calculate the global mask weight for a specific render vertex, using a color buffer and specified mask channel settings.
@@ -332,7 +334,7 @@ namespace UE::MLDeformer
 		 * @param bInvertMaskChannel Set to true when we should invert the weight value.
 		 * @return An array with a weight per render vertex. A value of 1.0 is returned when the mask channel is set to disable the mask. The return value is always between 0 and 1.
 		 */
-		TArray<float> CalcGlobalMaskWeights(const TArray<int32>& VertexMap, const FColorVertexBuffer& ColorBuffer, EMLDeformerMaskChannel MaskChannel, bool bInvertMaskChannel) const;
+		UE_API TArray<float> CalcGlobalMaskWeights(const TArray<int32>& VertexMap, const FColorVertexBuffer& ColorBuffer, EMLDeformerMaskChannel MaskChannel, bool bInvertMaskChannel) const;
 
 		/**
 		 * Tell the editor model that we should update the mask buffer on the next Tick call.
@@ -356,3 +358,5 @@ namespace UE::MLDeformer
 		bool bNeedsRebuildMaskInfo = false;
 	};
 }	// namespace UE::MLDeformer
+
+#undef UE_API
diff --git a/Engine/Plugins/Animation/MLDeformer/MLDeformerFramework/Source/MLDeformerFrameworkEditor/Public/MLDeformerMorphModelVizSettingsDetails.h b/Engine/Plugins/Animation/MLDeformer/MLDeformerFramework/Source/MLDeformerFrameworkEditor/Public/MLDeformerMorphModelVizSettingsDetails.h
--- a/Engine/Plugins/Animation/MLDeformer/MLDeformerFramework/Source/MLDeformerFrameworkEditor/Public/MLDeformerMorphModelVizSettingsDetails.h
+++ b/Engine/Plugins/Animation/MLDeformer/MLDeformerFramework/Source/MLDeformerFrameworkEditor/Public/MLDeformerMorphModelVizSettingsDetails.h
@@ -5,6 +5,8 @@
 #include "CoreTypes.h"
 #include "MLDeformerGeomCacheVizSettingsDetails.h"
 
+#define UE_API MLDEFORMERFRAMEWORKEDITOR_API
+
 class UMLDeformerMorphModel;
 class UMLDeformerMorphModelVizSettings;
 
@@ -13,18 +15,18 @@ namespace UE::MLDeformer
 	/**
 	 * The visualization settings for models inherited from the UMLDeformerMorphModel class.
 	 */
-	class MLDEFORMERFRAMEWORKEDITOR_API FMLDeformerMorphModelVizSettingsDetails
+	class FMLDeformerMorphModelVizSettingsDetails
 		: public FMLDeformerGeomCacheVizSettingsDetails
 	{
 	public:
 		// FMLDeformerVizSettingsDetails overrides.
-		virtual bool UpdateMemberPointers(const TArray<TWeakObjectPtr<UObject>>& Objects) override;
-		virtual void AddAdditionalSettings() override;
+		UE_API virtual bool UpdateMemberPointers(const TArray<TWeakObjectPtr<UObject>>& Objects) override;
+		UE_API virtual void AddAdditionalSettings() override;
 		// ~END FMLDeformerVizSettingsDetails overrides.
 
 	protected:
 		/** Is the morph target visualization option/checkbox enabled? */
-		bool IsMorphTargetsEnabled() const;
+		UE_API bool IsMorphTargetsEnabled() const;
 
 		/** A pointer to the runtime morph model. This is updated when UpdateMemberPointers is called. */
 		TObjectPtr<UMLDeformerMorphModel> MorphModel = nullptr;
@@ -33,3 +35,5 @@ namespace UE::MLDeformer
 		TObjectPtr<UMLDeformerMorphModelVizSettings> MorphModelVizSettings = nullptr;
 	};
 }	// namespace UE::MLDeformer
+
+#undef UE_API
diff --git a/Engine/Plugins/Animation/MLDeformer/MLDeformerFramework/Source/MLDeformerFrameworkEditor/Public/MLDeformerPaintMode.h b/Engine/Plugins/Animation/MLDeformer/MLDeformerFramework/Source/MLDeformerFrameworkEditor/Public/MLDeformerPaintMode.h
--- a/Engine/Plugins/Animation/MLDeformer/MLDeformerFramework/Source/MLDeformerFrameworkEditor/Public/MLDeformerPaintMode.h
+++ b/Engine/Plugins/Animation/MLDeformer/MLDeformerFramework/Source/MLDeformerFrameworkEditor/Public/MLDeformerPaintMode.h
@@ -6,6 +6,8 @@
 #include "EdMode.h"
 #include "MLDeformerPaintMode.generated.h"
 
+#define UE_API MLDEFORMERFRAMEWORKEDITOR_API
+
 class USkeletalMeshComponent;
 
 namespace UE::MLDeformer
@@ -13,28 +15,30 @@ namespace UE::MLDeformer
 	class FMLDeformerEditorToolkit;
 }
 
-UCLASS()
-class MLDEFORMERFRAMEWORKEDITOR_API UMLDeformerPaintMode
+UCLASS(MinimalAPI)
+class UMLDeformerPaintMode
 	: public UBaseLegacyWidgetEdMode
 {
 	GENERATED_BODY()
 
 public:
-	const static FEditorModeID Id;	
+	UE_API const static FEditorModeID Id;	
 
-	UMLDeformerPaintMode();
+	UE_API UMLDeformerPaintMode();
 
 	// UEdMode overrides
-	virtual void Enter() override;
-	virtual void CreateToolkit() override;
+	UE_API virtual void Enter() override;
+	UE_API virtual void CreateToolkit() override;
 	virtual bool UsesTransformWidget() const override	{ return false; };
 	virtual bool UsesPropertyWidgets() const override	{ return false; };
 	virtual bool UsesToolkits() const override			{ return true; }
 	// ~END UEdMode overrides
 
-	void SetMLDeformerEditor(UE::MLDeformer::FMLDeformerEditorToolkit* Editor);
-	void UpdatePose(USkeletalMeshComponent* SkeletalMeshComponent, bool bFullUpdate=true);
+	UE_API void SetMLDeformerEditor(UE::MLDeformer::FMLDeformerEditorToolkit* Editor);
+	UE_API void UpdatePose(USkeletalMeshComponent* SkeletalMeshComponent, bool bFullUpdate=true);
 
 private:
 	TArray<FVector3f> BindPosePositions;
 };
+
+#undef UE_API
diff --git a/Engine/Plugins/Animation/MLDeformer/MLDeformerFramework/Source/MLDeformerFrameworkEditor/Public/MLDeformerPaintModeToolkit.h b/Engine/Plugins/Animation/MLDeformer/MLDeformerFramework/Source/MLDeformerFrameworkEditor/Public/MLDeformerPaintModeToolkit.h
--- a/Engine/Plugins/Animation/MLDeformer/MLDeformerFramework/Source/MLDeformerFrameworkEditor/Public/MLDeformerPaintModeToolkit.h
+++ b/Engine/Plugins/Animation/MLDeformer/MLDeformerFramework/Source/MLDeformerFrameworkEditor/Public/MLDeformerPaintModeToolkit.h
@@ -6,6 +6,8 @@
 #include "StatusBarSubsystem.h"
 #include "Toolkits/BaseToolkit.h"
 
+#define UE_API MLDEFORMERFRAMEWORKEDITOR_API
+
 class IToolkitHost;
 class STextBlock;
 class SButton;
@@ -14,40 +16,40 @@ namespace UE::MLDeformer
 {
 	class FMLDeformerEditorToolkit;
 
-	class MLDEFORMERFRAMEWORKEDITOR_API FMLDeformerPaintModeToolkit
+	class FMLDeformerPaintModeToolkit
 		: public FModeToolkit
 	{
 	public:
-		virtual ~FMLDeformerPaintModeToolkit() override;
+		UE_API virtual ~FMLDeformerPaintModeToolkit() override;
 	
 		// IToolkit overrides
-		virtual void Init(const TSharedPtr<IToolkitHost>& InToolkitHost, TWeakObjectPtr<UEdMode> InOwningMode) override;
-		virtual FName GetToolkitFName() const override;
-		virtual FText GetBaseToolkitName() const override;
+		UE_API virtual void Init(const TSharedPtr<IToolkitHost>& InToolkitHost, TWeakObjectPtr<UEdMode> InOwningMode) override;
+		UE_API virtual FName GetToolkitFName() const override;
+		UE_API virtual FText GetBaseToolkitName() const override;
 		virtual TSharedPtr<SWidget> GetInlineContent() const override			{ return ToolkitWidget; }
 
 		// FModeToolkit overrides
-		virtual void OnToolStarted(UInteractiveToolManager* Manager, UInteractiveTool* Tool) override;
-		virtual void OnToolEnded(UInteractiveToolManager* Manager, UInteractiveTool* Tool) override;
-		virtual FText GetActiveToolDisplayName() const override;
-		virtual FText GetActiveToolMessage() const override;
+		UE_API virtual void OnToolStarted(UInteractiveToolManager* Manager, UInteractiveTool* Tool) override;
+		UE_API virtual void OnToolEnded(UInteractiveToolManager* Manager, UInteractiveTool* Tool) override;
+		UE_API virtual FText GetActiveToolDisplayName() const override;
+		UE_API virtual FText GetActiveToolMessage() const override;
 
 		void SetMLDeformerEditor(FMLDeformerEditorToolkit* Editor)				{ MLDeformerEditor = Editor; }
 		UE::MLDeformer::FMLDeformerEditorToolkit* GetMLDeformerEditor() const	{ return MLDeformerEditor; }
 
 	private:
-		void PostNotification(const FText& Message);
-		void ClearNotification();
+		UE_API void PostNotification(const FText& Message);
+		UE_API void ClearNotification();
 
-		void PostWarning(const FText& Message);
-		void ClearWarning();
+		UE_API void PostWarning(const FText& Message);
+		UE_API void ClearWarning();
 
-		void UpdateActiveToolProperties(UInteractiveTool* Tool);
+		UE_API void UpdateActiveToolProperties(UInteractiveTool* Tool);
 
-		void RegisterPalettes();
+		UE_API void RegisterPalettes();
 		FDelegateHandle ActivePaletteChangedHandle;
 
-		void MakeToolAcceptCancelWidget();
+		UE_API void MakeToolAcceptCancelWidget();
 
 		FText ActiveToolName;
 		FText ActiveToolMessage;
@@ -65,3 +67,5 @@ namespace UE::MLDeformer
 		TSharedPtr<STextBlock> ToolWarningArea;
 	};
 } // namespace UE::MLDeformer
+
+#undef UE_API
diff --git a/Engine/Plugins/Animation/MLDeformer/MLDeformerFramework/Source/MLDeformerFrameworkEditor/Public/MLDeformerSampler.h b/Engine/Plugins/Animation/MLDeformer/MLDeformerFramework/Source/MLDeformerFrameworkEditor/Public/MLDeformerSampler.h
--- a/Engine/Plugins/Animation/MLDeformer/MLDeformerFramework/Source/MLDeformerFrameworkEditor/Public/MLDeformerSampler.h
+++ b/Engine/Plugins/Animation/MLDeformer/MLDeformerFramework/Source/MLDeformerFrameworkEditor/Public/MLDeformerSampler.h
@@ -6,6 +6,8 @@
 #include "MLDeformerModel.h"
 #include "UObject/ObjectPtr.h"
 
+#define UE_API MLDEFORMERFRAMEWORKEDITOR_API
+
 class UDebugSkelMeshComponent;
 class UGeometryCacheComponent;
 class UWorld;
@@ -34,19 +36,19 @@ namespace UE::MLDeformer
 	 * This class can sample bone rotations, curve values and vertex deltas.
 	 * It does this by creating two temp actors, one with skeletal mesh component and one with geom cache component.
 	 */
-	class MLDEFORMERFRAMEWORKEDITOR_API FMLDeformerSampler
+	class FMLDeformerSampler
 	{
 	public:
-		virtual ~FMLDeformerSampler();
+		UE_API virtual ~FMLDeformerSampler();
 
 		/** Call Init whenever assets or other relevant settings change. */
 		UE_DEPRECATED(5.4, "Please switch to using Init(Model, AnimIndex) instead.")
-		virtual void Init(FMLDeformerEditorModel* Model);
+		UE_API virtual void Init(FMLDeformerEditorModel* Model);
 
-		virtual void Init(FMLDeformerEditorModel* Model, int32 AnimIndex);
+		UE_API virtual void Init(FMLDeformerEditorModel* Model, int32 AnimIndex);
 
 		/** Call this every time the frame changes. This will update all buffer contents. */
-		virtual void Sample(int32 AnimFrameIndex);
+		UE_API virtual void Sample(int32 AnimFrameIndex);
 
 		/** Get the time in seconds, at a given frame index. */
 		virtual float GetTimeAtFrame(int32 InAnimFrameIndex) const		{ return 0.0f; }
@@ -90,10 +92,10 @@ namespace UE::MLDeformer
 		bool IsInitialized() const										{ return SkelMeshActor.Get() != nullptr; }
 
 		/** Get the number of bones. */
-		int32 GetNumBones() const;
+		UE_API int32 GetNumBones() const;
 
 		/** Calculate the memory usage of the sampler. */
-		SIZE_T CalcMemUsagePerFrameInBytes() const;
+		UE_API SIZE_T CalcMemUsagePerFrameInBytes() const;
 		UDebugSkelMeshComponent* GetSkeletalMeshComponent() { return SkeletalMeshComponent; }
 
 		/**
@@ -106,42 +108,42 @@ namespace UE::MLDeformer
 		int32 GetNumFloatsPerCurve() const								{ return NumFloatsPerCurve; }
 
 		/** Get the skinned positions for a given LOD level. */
-		void ExtractSkinnedPositions(int32 LODIndex, TArray<FVector3f>& OutPositions);
+		UE_API void ExtractSkinnedPositions(int32 LODIndex, TArray<FVector3f>& OutPositions);
 
 	protected:
 		/** Create the actors used for sampling. This creates two actors, one for the base skeletal mesh and one for the target mesh. */
-		virtual void CreateActors();
+		UE_API virtual void CreateActors();
 
 		/** Register the target mesh components. */
 		UE_DEPRECATED(5.4, "Please create your target component inside the overloaded Init method instead.")
 		virtual void RegisterTargetComponents() {}
 
 		/** Extract the skinned vertex positions, of the skeletal mesh. */
-		void ExtractSkinnedPositions(int32 LODIndex, TArray<FMatrix44f>& InBoneMatrices, TArray<FVector3f>& TempPositions, TArray<FVector3f>& OutPositions) const;
+		UE_API void ExtractSkinnedPositions(int32 LODIndex, TArray<FMatrix44f>& InBoneMatrices, TArray<FVector3f>& TempPositions, TArray<FVector3f>& OutPositions) const;
 
 		/** Extract the unskinned vertex positions, of the skeletal mesh. */
-		void ExtractUnskinnedPositions(int32 LODIndex, TArray<FVector3f>& OutPositions) const;
+		UE_API void ExtractUnskinnedPositions(int32 LODIndex, TArray<FVector3f>& OutPositions) const;
 
 		/** Calculate the inverse skinning transform, which is used to convert from post-skinning to pre-skinning space. */
-		FMatrix44f CalcInverseSkinningTransform(int32 VertexIndex, const FSkeletalMeshLODRenderData& SkelMeshLODData, const FSkinWeightVertexBuffer& SkinWeightBuffer) const;
+		UE_API FMatrix44f CalcInverseSkinningTransform(int32 VertexIndex, const FSkeletalMeshLODRenderData& SkelMeshLODData, const FSkinWeightVertexBuffer& SkinWeightBuffer) const;
 
 		/** Calculate the delta using dual quaternions, taking a world space delta vector as input. */
-		FVector3f CalcDualQuaternionDelta(int32 VertexIndex, const FVector3f& WorldDelta, const FSkeletalMeshLODRenderData& SkelMeshLODData, const FSkinWeightVertexBuffer& SkinWeightBuffer) const;
+		UE_API FVector3f CalcDualQuaternionDelta(int32 VertexIndex, const FVector3f& WorldDelta, const FSkeletalMeshLODRenderData& SkelMeshLODData, const FSkinWeightVertexBuffer& SkinWeightBuffer) const;
 
 		/** Create a new actor with a specific name, in a specified world. */
-		AActor* CreateNewActor(UWorld* InWorld, const FName& Name) const;
+		UE_API AActor* CreateNewActor(UWorld* InWorld, const FName& Name) const;
 
 		/** Update the skeletal mesh component's play position, and refresh its bone transforms internally. */
-		void UpdateSkeletalMeshComponent();
+		UE_API void UpdateSkeletalMeshComponent();
 
 		/** Update the skinned vertex positions. */
-		void UpdateSkinnedPositions();
+		UE_API void UpdateSkinnedPositions();
 
 		/** Update the bone rotation buffers. */
-		void UpdateBoneRotations();
+		UE_API void UpdateBoneRotations();
 
 		/** Update the curve value buffers. */
-		void UpdateCurveValues();
+		UE_API void UpdateCurveValues();
 
 	protected:
 		/** The vertex delta editor model used to sample. */
@@ -213,3 +215,5 @@ namespace UE::MLDeformer
 		EMLDeformerSkinningMode SkinningMode = EMLDeformerSkinningMode::Linear;
 	};
 }	// namespace UE::MLDeformer
+
+#undef UE_API
diff --git a/Engine/Plugins/Animation/MLDeformer/MLDeformerFramework/Source/MLDeformerFrameworkEditor/Public/MLDeformerTestHelpers.h b/Engine/Plugins/Animation/MLDeformer/MLDeformerFramework/Source/MLDeformerFrameworkEditor/Public/MLDeformerTestHelpers.h
--- a/Engine/Plugins/Animation/MLDeformer/MLDeformerFramework/Source/MLDeformerFrameworkEditor/Public/MLDeformerTestHelpers.h
+++ b/Engine/Plugins/Animation/MLDeformer/MLDeformerFramework/Source/MLDeformerFrameworkEditor/Public/MLDeformerTestHelpers.h
@@ -5,6 +5,8 @@
 #include "Containers/UnrealString.h"
 #include "Misc/AutomationTest.h"
 
+#define UE_API MLDEFORMERFRAMEWORKEDITOR_API
+
 class UMLDeformerAsset;
 class UMLDeformerModel;
 
@@ -14,22 +16,22 @@ namespace UE::MLDeformer
 	class FMLDeformerEditorModelRegistry;
 
 	/** A scoped ML Deformer editor, which auto closes the window on destruction. */
-	class MLDEFORMERFRAMEWORKEDITOR_API FMLDeformerScopedEditor
+	class FMLDeformerScopedEditor
 	{
 	public:
 		FMLDeformerScopedEditor() = delete;
 		FMLDeformerScopedEditor(const FMLDeformerScopedEditor&) = delete;
 		FMLDeformerScopedEditor(FMLDeformerScopedEditor&&) = delete;
-		FMLDeformerScopedEditor(FMLDeformerEditorToolkit* InEditor);
-		FMLDeformerScopedEditor(UMLDeformerAsset* Asset);
-		FMLDeformerScopedEditor(const FString& AssetName);
-		~FMLDeformerScopedEditor();	// Auto close the editor window.
-		bool IsValid() const;
-		void SetCloseEditor(bool bCloseOnDestroy);
-		FMLDeformerEditorToolkit* Get() const;
+		UE_API FMLDeformerScopedEditor(FMLDeformerEditorToolkit* InEditor);
+		UE_API FMLDeformerScopedEditor(UMLDeformerAsset* Asset);
+		UE_API FMLDeformerScopedEditor(const FString& AssetName);
+		UE_API ~FMLDeformerScopedEditor();	// Auto close the editor window.
+		UE_API bool IsValid() const;
+		UE_API void SetCloseEditor(bool bCloseOnDestroy);
+		UE_API FMLDeformerEditorToolkit* Get() const;
 		FMLDeformerEditorToolkit* operator = (const FMLDeformerScopedEditor&) = delete;
 		FMLDeformerEditorToolkit* operator = (const FMLDeformerScopedEditor&&) = delete;
-		FMLDeformerEditorToolkit* operator ->() const;
+		UE_API FMLDeformerEditorToolkit* operator ->() const;
 
 	private:
 		FMLDeformerEditorToolkit* Editor = nullptr;
@@ -37,7 +39,7 @@ namespace UE::MLDeformer
 	};
 
 	/** Some helper functions for testing ML Deformer related things. */
-	class MLDEFORMERFRAMEWORKEDITOR_API FMLDeformerTestHelpers
+	class FMLDeformerTestHelpers
 	{
 	public:
 		/**
@@ -46,7 +48,7 @@ namespace UE::MLDeformer
 		 * @param Asset A pointer to the asset to open an editor for.
 		 * @return A pointer to the editor toolkit, or nullptr when it didn't succeed.
 		 */
-		static FMLDeformerEditorToolkit* OpenAssetEditor(UMLDeformerAsset* Asset);
+		static UE_API FMLDeformerEditorToolkit* OpenAssetEditor(UMLDeformerAsset* Asset);
 
 		/**
 		 * Open the asset editor for a given ML Deformer asset, by its package name.
@@ -55,7 +57,7 @@ namespace UE::MLDeformer
 		 * @param Asset The package name, for example something like: "/MLDeformerFramework/Tests/MLD_Rampage.MLD_Rampage".
 		 * @return A pointer to the editor toolkit, or nullptr when it didn't succeed.
 		 */
-		static FMLDeformerEditorToolkit* OpenAssetEditor(const FString& AssetName);
+		static UE_API FMLDeformerEditorToolkit* OpenAssetEditor(const FString& AssetName);
 
 		/**
 		 * Post an edit change event for a given property in the Model class.
@@ -63,6 +65,8 @@ namespace UE::MLDeformer
 		 * @param PropertyName The name of the property (the class member name).
 		 * @return Returns true if successfully emitted the PostEditChangeProperty call, or false if the property cannot be found.
 		 */
-		static bool PostModelPropertyChanged(UMLDeformerModel* Model, const FName PropertyName);
+		static UE_API bool PostModelPropertyChanged(UMLDeformerModel* Model, const FName PropertyName);
 	};
 }	// namespace UE::MLDeformer
+
+#undef UE_API
diff --git a/Engine/Plugins/Animation/MLDeformer/MLDeformerFramework/Source/MLDeformerFrameworkEditor/Public/MLDeformerTrainingModel.h b/Engine/Plugins/Animation/MLDeformer/MLDeformerFramework/Source/MLDeformerFrameworkEditor/Public/MLDeformerTrainingModel.h
--- a/Engine/Plugins/Animation/MLDeformer/MLDeformerFramework/Source/MLDeformerFrameworkEditor/Public/MLDeformerTrainingModel.h
+++ b/Engine/Plugins/Animation/MLDeformer/MLDeformerFramework/Source/MLDeformerFrameworkEditor/Public/MLDeformerTrainingModel.h
@@ -7,6 +7,8 @@
 #include "UObject/ObjectPtr.h"
 #include "MLDeformerTrainingModel.generated.h"
 
+#define UE_API MLDEFORMERFRAMEWORKEDITOR_API
+
 class UMLDeformerModel;
 
 namespace UE::MLDeformer
@@ -55,8 +57,8 @@ namespace UE::MLDeformer
  * Keep in mind that in Unreal Engine all python code is lower case. So a "Train" method inside c++ will need to be called "train" inside the python code.
  * Or if you have something called "PerformMyTraining" it will need to be called "perform_my_training" inside Python.
  */
-UCLASS(Blueprintable)
-class MLDEFORMERFRAMEWORKEDITOR_API UMLDeformerTrainingModel
+UCLASS(MinimalAPI, Blueprintable)
+class UMLDeformerTrainingModel
 	: public UObject
 {
 	GENERATED_BODY()
@@ -67,31 +69,31 @@ public:
 	 * This is automatically called by the editor.
 	 * @param InEditorModel The pointer to the editor model that this is a training model for.
 	 */
-	virtual void Init(UE::MLDeformer::FMLDeformerEditorModel* InEditorModel);
+	UE_API virtual void Init(UE::MLDeformer::FMLDeformerEditorModel* InEditorModel);
 
 	/** Get the runtime ML Deformer model object. */
 	UFUNCTION(BlueprintPure, Category = "Training Data")
-	UMLDeformerModel* GetModel() const;
+	UE_API UMLDeformerModel* GetModel() const;
 
 	/** Get the number of input transforms. This is the number of bones. */
 	UFUNCTION(BlueprintPure, Category = "Training Data")
-	int32 GetNumberSampleTransforms() const;
+	UE_API int32 GetNumberSampleTransforms() const;
 
 	/** Get number of input curves. */
 	UFUNCTION(BlueprintPure, Category = "Training Data")
-	int32 GetNumberSampleCurves() const;
+	UE_API int32 GetNumberSampleCurves() const;
 
 	/** Get the number of vertex deltas. */
 	UFUNCTION(BlueprintPure, Category = "Training Data")
-	int32 GetNumberSampleDeltas() const;
+	UE_API int32 GetNumberSampleDeltas() const;
 
 	/** Get the number of samples in this data set. This is the number of sample frames we want to train on. */
 	UFUNCTION(BlueprintPure, Category = "Training Data")
-    int32 NumSamples() const;
+    UE_API int32 NumSamples() const;
 
 	/** This will make the sampling start again from the beginning. This can be used if you have to iterate multiple times over the data set. */
 	UFUNCTION(BlueprintCallable, Category = "Training Data")
-    void ResetSampling();
+    UE_API void ResetSampling();
 
 	/** 
 	 * Set the current sample frame. This will internally call the SampleFrame method, which will update the deltas, curve values and bone rotations. 
@@ -101,14 +103,14 @@ public:
 	 */
 	UE_DEPRECATED(5.4, "Please use NextSample instead.")
 	UFUNCTION(meta = (DeprecatedFunction, DeprecationMessage = "Please use NextSample instead."))
-	bool SetCurrentSampleIndex(int32 Index);
+	UE_API bool SetCurrentSampleIndex(int32 Index);
 
 	/**
 	 * Take the next sample. This will update the deltas, curve values and bone rotation arrays with values sampled at the next frame.
 	 * This will return false when there is something wrong or we sampled more times than NumSamples() returns.
 	 */
 	UFUNCTION(BlueprintCallable, Category = "Training Data")
-	bool NextSample();
+	UE_API bool NextSample();
 
 	/** 
 	 * Check whether we need to resample the inputs and outputs, or if we can use a cached version. 
@@ -116,23 +118,23 @@ public:
 	 * @return Returns true when we need to regenerate any cached data, otherwise false is returned.
 	 */
 	UFUNCTION(BlueprintCallable, Category = "Training Data")
-	bool GetNeedsResampling() const;
+	UE_API bool GetNeedsResampling() const;
 
 	/** Specify whether we need to resample any cached data or not, because our input assets, or any other relevant settings changed that would invalidate the cached data. */
 	UFUNCTION(BlueprintCallable, Category = "Training Data")
-	void SetNeedsResampling(bool bNeedsResampling);
+	UE_API void SetNeedsResampling(bool bNeedsResampling);
 
 	/** Set the number of floats per curve. On default this is one. */
 	UFUNCTION(BlueprintCallable, Category = "Training Data")
-	void SetNumFloatsPerCurve(int32 NumFloatsPerCurve);
+	UE_API void SetNumFloatsPerCurve(int32 NumFloatsPerCurve);
 
 	/** Get the names of all valid training animation masks. They are valid when they are used, the anim is enabled, and the attribute exists on the skeletal mesh. */
 	UFUNCTION(BlueprintPure, Category = "Training Data")
-    TArray<FName> GetTrainingInputAnimMasks() const;
+    UE_API TArray<FName> GetTrainingInputAnimMasks() const;
 
 	/** Get the per vertex data for a given mask name. The mask name must be one that is present in the array returned by GetTrainingInputAnimMasks. */
 	UFUNCTION(BlueprintPure, Category = "Training Data")
-	TArray<float> GetTrainingInputAnimMaskData(FName MaskName) const;
+	UE_API TArray<float> GetTrainingInputAnimMaskData(FName MaskName) const;
 
 	/** 
 	 * For each sample we took, this contains the index inside the mask as returned by the GetTrainingInputAnimMasks() array.
@@ -142,24 +144,24 @@ public:
 
 	/** Set the list of possible training devices. */
 	UFUNCTION(BlueprintCallable, Category = "Training Model")
-	void SetDeviceList(const TArray<FString>& DeviceNames, int32 PreferredDeviceIndex);
+	UE_API void SetDeviceList(const TArray<FString>& DeviceNames, int32 PreferredDeviceIndex);
 
 	UFUNCTION(BlueprintImplementableEvent, Category = "Training Model")
-	void UpdateAvailableDevices() const;
+	UE_API void UpdateAvailableDevices() const;
 
 protected:
 	/** This updates the sample deltas, curves, and bone rotations. */
 	UE_DEPRECATED(5.4, "Please use UMLDeformerTrainingModel::SampleNextFrame() instead.")
-    virtual bool SampleFrame(int32 Index);
+    UE_API virtual bool SampleFrame(int32 Index);
 
 	/** Sample the next frame in the sequence of samples. */
-	virtual bool SampleNextFrame();
+	UE_API virtual bool SampleNextFrame();
 
 	/** Change the pointer to the editor model. */
-	void SetEditorModel(UE::MLDeformer::FMLDeformerEditorModel* InModel);
+	UE_API void SetEditorModel(UE::MLDeformer::FMLDeformerEditorModel* InModel);
 
 	/** Get a pointer to the editor model. */ 
-	UE::MLDeformer::FMLDeformerEditorModel* GetEditorModel() const;
+	UE_API UE::MLDeformer::FMLDeformerEditorModel* GetEditorModel() const;
 
 	/**
 	 * Find the next input animation to sample from.
@@ -174,7 +176,7 @@ protected:
 	 * Get the mask index for a given input training animation.
 	 * @see GetTrainingInputAnimMasks.
 	 */
-	int32 GetMaskIndexForAnimIndex(int32 AnimIndex) const;
+	UE_API int32 GetMaskIndexForAnimIndex(int32 AnimIndex) const;
 
 public:
 	// The delta values per vertex for this sample. This is updated after NextSample is called. Contains an xyz (3 floats) for each vertex.
@@ -214,5 +216,7 @@ protected:
 	TArray<FName> MaskNames;
 
 	/** The default mask name. */
-	static FName DefaultMaskName;
+	static UE_API FName DefaultMaskName;
 };
+
+#undef UE_API
diff --git a/Engine/Plugins/Animation/MLDeformer/MLDeformerFramework/Source/MLDeformerFrameworkEditor/Public/MLDeformerVizSettingsDetails.h b/Engine/Plugins/Animation/MLDeformer/MLDeformerFramework/Source/MLDeformerFrameworkEditor/Public/MLDeformerVizSettingsDetails.h
--- a/Engine/Plugins/Animation/MLDeformer/MLDeformerFramework/Source/MLDeformerFrameworkEditor/Public/MLDeformerVizSettingsDetails.h
+++ b/Engine/Plugins/Animation/MLDeformer/MLDeformerFramework/Source/MLDeformerFrameworkEditor/Public/MLDeformerVizSettingsDetails.h
@@ -7,6 +7,8 @@
 #include "AssetRegistry/AssetData.h"
 #include "PropertyHandle.h"
 
+#define UE_API MLDEFORMERFRAMEWORKEDITOR_API
+
 class IDetailCategoryBuilder;
 class IDetailLayoutBuilder;
 class IDetailGroup;
@@ -22,12 +24,12 @@ namespace UE::MLDeformer
 	 * The detail customization base class for any UMLDeformerVizSettings inherited object.
 	 * This automatically adds the settings of the UMLDeformerVizSettings class to the UI and creates some groups, inserts error messages when needed, etc.
 	 */
-	class MLDEFORMERFRAMEWORKEDITOR_API FMLDeformerVizSettingsDetails
+	class FMLDeformerVizSettingsDetails
 		: public IDetailCustomization
 	{
 	public:
 		// ILayoutDetails overrides.
-		virtual void CustomizeDetails(IDetailLayoutBuilder& DetailBuilder) override;
+		UE_API virtual void CustomizeDetails(IDetailLayoutBuilder& DetailBuilder) override;
 		// ~END ILayoutDetails overrides.
 
 		/** Get a pointer to the detail layout builder. */
@@ -35,30 +37,30 @@ namespace UE::MLDeformer
 
 		// UE_DEPRECATED(5.4, Remove these constructors when StatsMainMemUsageGroup is removed and there is no longer warnings in the implicit functions)
 		FMLDeformerVizSettingsDetails() = default;
-		FMLDeformerVizSettingsDetails(const FMLDeformerVizSettingsDetails&);
-		FMLDeformerVizSettingsDetails(FMLDeformerVizSettingsDetails&&);
-		FMLDeformerVizSettingsDetails& operator=(const FMLDeformerVizSettingsDetails&);
-		FMLDeformerVizSettingsDetails& operator=(FMLDeformerVizSettingsDetails&&);
+		UE_API FMLDeformerVizSettingsDetails(const FMLDeformerVizSettingsDetails&);
+		UE_API FMLDeformerVizSettingsDetails(FMLDeformerVizSettingsDetails&&);
+		UE_API FMLDeformerVizSettingsDetails& operator=(const FMLDeformerVizSettingsDetails&);
+		UE_API FMLDeformerVizSettingsDetails& operator=(FMLDeformerVizSettingsDetails&&);
 		// END UE_DEPRECATED(5.4)
 
 	protected:
 		/** The filter that only shows anim sequences that are compatible with the given skeleton. */
-		bool FilterAnimSequences(const FAssetData& AssetData, USkeleton* Skeleton);
+		UE_API bool FilterAnimSequences(const FAssetData& AssetData, USkeleton* Skeleton);
 
 		/** When we click the reset to default button next to the deformer graph property. */
-		void OnResetToDefaultDeformerGraph(TSharedPtr<IPropertyHandle> PropertyHandle);
+		UE_API void OnResetToDefaultDeformerGraph(TSharedPtr<IPropertyHandle> PropertyHandle);
 
 		/** Check if the reset to default button next to the deformer graph property should be visible or not. */
-		bool IsResetToDefaultDeformerGraphVisible(TSharedPtr<IPropertyHandle> PropertyHandle);
+		UE_API bool IsResetToDefaultDeformerGraphVisible(TSharedPtr<IPropertyHandle> PropertyHandle);
 
 		/**
 		 * Update the class member pointers, which includes the pointer to the model, its editor model, and the viz settings.
 		 * @param Objects The array of objects that the detail customization is showing.
 		 */
-		virtual bool UpdateMemberPointers(const TArray<TWeakObjectPtr<UObject>>& Objects);
+		UE_API virtual bool UpdateMemberPointers(const TArray<TWeakObjectPtr<UObject>>& Objects);
 
 		/** Create the categories to which we can add properties. Update the member category pointers to those. */
-		virtual void CreateCategories();
+		UE_API virtual void CreateCategories();
 
 		/** Add some error message related to the test animation sequence. */
 		virtual void AddTestSequenceErrors() {}
@@ -73,9 +75,9 @@ namespace UE::MLDeformer
 		virtual void AddAdditionalSettings() {}
 
 		/** Add any statistics. */
-		virtual void AddStatistics();
+		UE_API virtual void AddStatistics();
 
-		void AddStatsPerfRow(
+		UE_API void AddStatsPerfRow(
 			IDetailGroup& Group,
 			const FText& Label,
 			const FMLDeformerEditorModel* InEditorModel,
@@ -133,3 +135,5 @@ namespace UE::MLDeformer
 		FNumberFormattingOptions MemUsageMetricFormat;
 	};
 }	// namespace UE::MLDeformer
+
+#undef UE_API
diff --git a/Engine/Plugins/Animation/MLDeformer/MLDeformerFramework/Source/MLDeformerFrameworkEditor/Public/SMLDeformerBonePickerDialog.h b/Engine/Plugins/Animation/MLDeformer/MLDeformerFramework/Source/MLDeformerFrameworkEditor/Public/SMLDeformerBonePickerDialog.h
--- a/Engine/Plugins/Animation/MLDeformer/MLDeformerFramework/Source/MLDeformerFrameworkEditor/Public/SMLDeformerBonePickerDialog.h
+++ b/Engine/Plugins/Animation/MLDeformer/MLDeformerFramework/Source/MLDeformerFrameworkEditor/Public/SMLDeformerBonePickerDialog.h
@@ -10,6 +10,8 @@
 #include "Framework/Commands/Commands.h"
 #include "MLDeformerEditorStyle.h"
 
+#define UE_API MLDEFORMERFRAMEWORKEDITOR_API
+
 struct FReferenceSkeleton;
 
 namespace UE::MLDeformer
@@ -75,7 +77,7 @@ namespace UE::MLDeformer
 	 * A bone picker dialog window.
 	 * This allows users to select one or more bones.
 	 */
-	class MLDEFORMERFRAMEWORKEDITOR_API SMLDeformerBonePickerDialog
+	class SMLDeformerBonePickerDialog
 		: public SCustomDialog
 	{
 	public:
@@ -92,24 +94,24 @@ namespace UE::MLDeformer
 		SLATE_ARGUMENT(TSharedPtr<SWidget>, ExtraWidget)		// An optional extra widget shown on top, next to the search field.
 		SLATE_END_ARGS()
 
-		void Construct(const FArguments& InArgs);
+		UE_API void Construct(const FArguments& InArgs);
 
 		/**
 		 * Get the array of picked bone names.
 		 * This will be an empty list when nothing got selected, which means we basically cancelled.
 		 * @return The array of selected bone names. Empty if cancelled or if no bones were there to begin with.
 		 */
-		const TArray<FName>& GetPickedBoneNames() const;
+		UE_API const TArray<FName>& GetPickedBoneNames() const;
 
 	private:
-		TSharedRef<SMLDeformerBonePickerTreeWidget> CreateBoneTree();
-		FReply OnKeyDown(const FGeometry& InGeometry, const FKeyEvent& InKeyEvent) override;
-		void OnOkClicked();
-		void OnCancelClicked();
-		void OnFilterTextChanged(const FText& InFilterText);
-		void RefreshBoneTreeElements() const;
-		void SelectInitialItems();
-		void SelectInitialItemsRecursive(const TSharedPtr<FMLDeformerBonePickerTreeElement>& Item);
+		UE_API TSharedRef<SMLDeformerBonePickerTreeWidget> CreateBoneTree();
+		UE_API FReply OnKeyDown(const FGeometry& InGeometry, const FKeyEvent& InKeyEvent) override;
+		UE_API void OnOkClicked();
+		UE_API void OnCancelClicked();
+		UE_API void OnFilterTextChanged(const FText& InFilterText);
+		UE_API void RefreshBoneTreeElements() const;
+		UE_API void SelectInitialItems();
+		UE_API void SelectInitialItemsRecursive(const TSharedPtr<FMLDeformerBonePickerTreeElement>& Item);
 
 	private:
 		TSharedPtr<SMLDeformerBonePickerTreeWidget> BoneTreeWidget;
@@ -124,3 +126,5 @@ namespace UE::MLDeformer
 		FString FilterText;	
 	};
 }
+
+#undef UE_API
diff --git a/Engine/Plugins/Animation/MLDeformer/MLDeformerFramework/Source/MLDeformerFrameworkEditor/Public/SMLDeformerCurvePickerDialog.h b/Engine/Plugins/Animation/MLDeformer/MLDeformerFramework/Source/MLDeformerFrameworkEditor/Public/SMLDeformerCurvePickerDialog.h
--- a/Engine/Plugins/Animation/MLDeformer/MLDeformerFramework/Source/MLDeformerFrameworkEditor/Public/SMLDeformerCurvePickerDialog.h
+++ b/Engine/Plugins/Animation/MLDeformer/MLDeformerFramework/Source/MLDeformerFrameworkEditor/Public/SMLDeformerCurvePickerDialog.h
@@ -10,6 +10,8 @@
 #include "Framework/Commands/Commands.h"
 #include "MLDeformerEditorStyle.h"
 
+#define UE_API MLDEFORMERFRAMEWORKEDITOR_API
+
 class USkeleton;
 class USkeletalMesh;
 
@@ -71,7 +73,7 @@ namespace UE::MLDeformer
 	};
 
 
-	class MLDEFORMERFRAMEWORKEDITOR_API SMLDeformerCurvePickerDialog
+	class SMLDeformerCurvePickerDialog
 		: public SCustomDialog
 	{
 	public:
@@ -86,22 +88,22 @@ namespace UE::MLDeformer
 		SLATE_ARGUMENT(bool, AllowMultiSelect)					// Set to false to allow the user to only select one curve, or true to allow for multiple.
 		SLATE_END_ARGS()
 
-		void Construct(const FArguments& InArgs);
+		UE_API void Construct(const FArguments& InArgs);
 
 		/**
 		 * Get the array of picked curves names.
 		 * This will be an empty list when nothing got selected, which means we basically cancelled.
 		 * @return The array of selected curves names. Empty if cancelled or if no curves were there to begin with.
 		 */
-		const TArray<FName>& GetPickedCurveNames() const;
+		UE_API const TArray<FName>& GetPickedCurveNames() const;
 
 	private:
-		TSharedRef<SMLDeformerCurvePickerListView> CreateListWidget();
-		FReply OnKeyDown(const FGeometry& InGeometry, const FKeyEvent& InKeyEvent) override;
-		void OnOkClicked();
-		void OnCancelClicked();
-		void OnFilterTextChanged(const FText& InFilterText);
-		void RefreshListElements() const;
+		UE_API TSharedRef<SMLDeformerCurvePickerListView> CreateListWidget();
+		UE_API FReply OnKeyDown(const FGeometry& InGeometry, const FKeyEvent& InKeyEvent) override;
+		UE_API void OnOkClicked();
+		UE_API void OnCancelClicked();
+		UE_API void OnFilterTextChanged(const FText& InFilterText);
+		UE_API void RefreshListElements() const;
 
 	private:
 		TSharedPtr<SMLDeformerCurvePickerListView> CurveListView;
@@ -114,3 +116,5 @@ namespace UE::MLDeformer
 		bool bAllowMultiSelect = true;
 	};
 }
+
+#undef UE_API
diff --git a/Engine/Plugins/Animation/MLDeformer/MLDeformerFramework/Source/MLDeformerFrameworkEditor/Public/SMLDeformerCurvePickerWidget.h b/Engine/Plugins/Animation/MLDeformer/MLDeformerFramework/Source/MLDeformerFrameworkEditor/Public/SMLDeformerCurvePickerWidget.h
--- a/Engine/Plugins/Animation/MLDeformer/MLDeformerFramework/Source/MLDeformerFrameworkEditor/Public/SMLDeformerCurvePickerWidget.h
+++ b/Engine/Plugins/Animation/MLDeformer/MLDeformerFramework/Source/MLDeformerFrameworkEditor/Public/SMLDeformerCurvePickerWidget.h
@@ -9,6 +9,8 @@
 #include "Widgets/DeclarativeSyntaxSupport.h"
 #include "Widgets/Views/STableRow.h"
 
+#define UE_API MLDEFORMERFRAMEWORKEDITOR_API
+
 class SComboButton;
 class SSearchBox;
 class USkeletalMesh;
@@ -25,7 +27,7 @@ namespace UE::MLDeformer
 	/**
 	 * The curve picker widget for the ML Deformer editor.
 	 */
-	class MLDEFORMERFRAMEWORKEDITOR_API SCurvePickerWidget
+	class SCurvePickerWidget
 		: public SCompoundWidget
 	{
 	public:
@@ -34,16 +36,16 @@ namespace UE::MLDeformer
 		SLATE_EVENT(FOnGetSkeletalMesh, OnGetSkeletalMesh)
 		SLATE_END_ARGS()
 
-		void Construct(const FArguments& InArgs);
+		UE_API void Construct(const FArguments& InArgs);
 
 		TSharedPtr<SSearchBox> GetFilterTextWidget() const { return SearchBox; }
 
 	private:
-		void RefreshListItems();
-		void FilterAvailableCurves();
-		void HandleSelectionChanged(TSharedPtr<FString> InItem, ESelectInfo::Type InSelectionType);
-		void HandleFilterTextChanged(const FText& InFilterText);
-		TSharedRef<ITableRow> HandleGenerateRow(TSharedPtr<FString> InItem, const TSharedRef<STableViewBase>& InOwnerTable);
+		UE_API void RefreshListItems();
+		UE_API void FilterAvailableCurves();
+		UE_API void HandleSelectionChanged(TSharedPtr<FString> InItem, ESelectInfo::Type InSelectionType);
+		UE_API void HandleFilterTextChanged(const FText& InFilterText);
+		UE_API TSharedRef<ITableRow> HandleGenerateRow(TSharedPtr<FString> InItem, const TSharedRef<STableViewBase>& InOwnerTable);
 
 	private:
 		/** Delegate fired when a curve name is picked. */
@@ -102,3 +104,5 @@ namespace UE::MLDeformer
 		FText SelectedCurve;
 	};
 }	// namespace UE::MLDeformer
+
+#undef UE_API
diff --git a/Engine/Plugins/Animation/MLDeformer/MLDeformerFramework/Source/MLDeformerFrameworkEditor/Public/SMLDeformerInputBonesWidget.h b/Engine/Plugins/Animation/MLDeformer/MLDeformerFramework/Source/MLDeformerFrameworkEditor/Public/SMLDeformerInputBonesWidget.h
--- a/Engine/Plugins/Animation/MLDeformer/MLDeformerFramework/Source/MLDeformerFrameworkEditor/Public/SMLDeformerInputBonesWidget.h
+++ b/Engine/Plugins/Animation/MLDeformer/MLDeformerFramework/Source/MLDeformerFrameworkEditor/Public/SMLDeformerInputBonesWidget.h
@@ -11,6 +11,8 @@
 #include "UObject/NameTypes.h"
 #include "Styling/SlateColor.h"
 
+#define UE_API MLDEFORMERFRAMEWORKEDITOR_API
+
 struct FReferenceSkeleton;
 
 namespace UE::MLDeformer
@@ -35,11 +37,11 @@ namespace UE::MLDeformer
 	};
 
 
-	class MLDEFORMERFRAMEWORKEDITOR_API FMLDeformerInputBoneTreeElement
+	class FMLDeformerInputBoneTreeElement
 		: public TSharedFromThis<FMLDeformerInputBoneTreeElement>
 	{
 	public:
-		TSharedRef<ITableRow> MakeTreeRowWidget(const TSharedRef<STableViewBase>& InOwnerTable, TSharedRef<FMLDeformerInputBoneTreeElement> InTreeElement, TSharedPtr<SMLDeformerInputBoneTreeWidget> InTreeWidget);
+		UE_API TSharedRef<ITableRow> MakeTreeRowWidget(const TSharedRef<STableViewBase>& InOwnerTable, TSharedRef<FMLDeformerInputBoneTreeElement> InTreeElement, TSharedPtr<SMLDeformerInputBoneTreeWidget> InTreeWidget);
 
 	public:
 		FName Name;
@@ -48,21 +50,21 @@ namespace UE::MLDeformer
 	};
 
 
-	class MLDEFORMERFRAMEWORKEDITOR_API SMLDeformerInputBoneTreeRowWidget 
+	class SMLDeformerInputBoneTreeRowWidget 
 		: public STableRow<TSharedPtr<FMLDeformerInputBoneTreeElement>>
 	{
 	public:
-		void Construct(const FArguments& InArgs, const TSharedRef<STableViewBase>& OwnerTable, TSharedRef<FMLDeformerInputBoneTreeElement> InTreeElement, TSharedPtr<SMLDeformerInputBoneTreeWidget> InTreeView);
+		UE_API void Construct(const FArguments& InArgs, const TSharedRef<STableViewBase>& OwnerTable, TSharedRef<FMLDeformerInputBoneTreeElement> InTreeElement, TSharedPtr<SMLDeformerInputBoneTreeWidget> InTreeView);
 
 	private:
 		TWeakPtr<FMLDeformerInputBoneTreeElement> WeakTreeElement;
-		FText GetName() const;
+		UE_API FText GetName() const;
 
 		friend class SMLDeformerInputBoneTreeWidget; 
 	};
 
 
-	class MLDEFORMERFRAMEWORKEDITOR_API SMLDeformerInputBoneTreeWidget
+	class SMLDeformerInputBoneTreeWidget
 		: public STreeView<TSharedPtr<FMLDeformerInputBoneTreeElement>>
 	{
 		SLATE_BEGIN_ARGS(SMLDeformerInputBoneTreeWidget) {}
@@ -70,21 +72,21 @@ namespace UE::MLDeformer
 		SLATE_END_ARGS()
 
 	public:
-		void Construct(const FArguments& InArgs);
+		UE_API void Construct(const FArguments& InArgs);
 		const TArray<TSharedPtr<FMLDeformerInputBoneTreeElement>>& GetRootElements() const	{ return RootElements; }
-		void AddElement(TSharedPtr<FMLDeformerInputBoneTreeElement> Element, TSharedPtr<FMLDeformerInputBoneTreeElement> ParentElement);
-		void RefreshElements(const TArray<FName>& BoneNames, const FReferenceSkeleton* RefSkeleton);
-		TArray<FName> ExtractAllElementNames() const;
+		UE_API void AddElement(TSharedPtr<FMLDeformerInputBoneTreeElement> Element, TSharedPtr<FMLDeformerInputBoneTreeElement> ParentElement);
+		UE_API void RefreshElements(const TArray<FName>& BoneNames, const FReferenceSkeleton* RefSkeleton);
+		UE_API TArray<FName> ExtractAllElementNames() const;
 
 	private:
-		TSharedRef<ITableRow> MakeTableRowWidget(TSharedPtr<FMLDeformerInputBoneTreeElement> InItem, const TSharedRef<STableViewBase>& OwnerTable);
-		void HandleGetChildrenForTree(TSharedPtr<FMLDeformerInputBoneTreeElement> InItem, TArray<TSharedPtr<FMLDeformerInputBoneTreeElement>>& OutChildren);
-		FReply OnKeyDown(const FGeometry& InGeometry, const FKeyEvent& InKeyEvent) override;
-		TSharedPtr<SWidget> OnContextMenuOpening() const;
-		void RecursiveSortElements(TSharedPtr<FMLDeformerInputBoneTreeElement> Element);
-		void RecursiveAddNames(const FMLDeformerInputBoneTreeElement& Element, TArray<FName>& OutNames) const;
-		void OnSelectionChanged(TSharedPtr<FMLDeformerInputBoneTreeElement> Selection, ESelectInfo::Type SelectInfo);
-		TSharedPtr<FMLDeformerInputBoneTreeElement> FindParentElementForBone(FName BoneName, const FReferenceSkeleton& RefSkeleton, const TMap<FName, TSharedPtr<FMLDeformerInputBoneTreeElement>>& NameToElementMap) const;
+		UE_API TSharedRef<ITableRow> MakeTableRowWidget(TSharedPtr<FMLDeformerInputBoneTreeElement> InItem, const TSharedRef<STableViewBase>& OwnerTable);
+		UE_API void HandleGetChildrenForTree(TSharedPtr<FMLDeformerInputBoneTreeElement> InItem, TArray<TSharedPtr<FMLDeformerInputBoneTreeElement>>& OutChildren);
+		UE_API FReply OnKeyDown(const FGeometry& InGeometry, const FKeyEvent& InKeyEvent) override;
+		UE_API TSharedPtr<SWidget> OnContextMenuOpening() const;
+		UE_API void RecursiveSortElements(TSharedPtr<FMLDeformerInputBoneTreeElement> Element);
+		UE_API void RecursiveAddNames(const FMLDeformerInputBoneTreeElement& Element, TArray<FName>& OutNames) const;
+		UE_API void OnSelectionChanged(TSharedPtr<FMLDeformerInputBoneTreeElement> Selection, ESelectInfo::Type SelectInfo);
+		UE_API TSharedPtr<FMLDeformerInputBoneTreeElement> FindParentElementForBone(FName BoneName, const FReferenceSkeleton& RefSkeleton, const TMap<FName, TSharedPtr<FMLDeformerInputBoneTreeElement>>& NameToElementMap) const;
 
 	private:
 		TArray<TSharedPtr<FMLDeformerInputBoneTreeElement>> RootElements;
@@ -92,7 +94,7 @@ namespace UE::MLDeformer
 	};
 
 
-	class MLDEFORMERFRAMEWORKEDITOR_API SMLDeformerInputBonesWidget
+	class SMLDeformerInputBonesWidget
 		: public SCompoundWidget
 	{
 		friend class SMLDeformerInputBoneTreeWidget;
@@ -103,22 +105,22 @@ namespace UE::MLDeformer
 		SLATE_END_ARGS()
 
 	public:
-		void Construct(const FArguments& InArgs);
-		void Refresh();
-		void BindCommands(TSharedPtr<FUICommandList> CommandList);
-		TSharedPtr<SMLDeformerInputBoneTreeWidget> GetTreeWidget() const;
-		TSharedPtr<SMLDeformerInputWidget> GetInputWidget() const;
-		FText GetSectionTitle() const;
+		UE_API void Construct(const FArguments& InArgs);
+		UE_API void Refresh();
+		UE_API void BindCommands(TSharedPtr<FUICommandList> CommandList);
+		UE_API TSharedPtr<SMLDeformerInputBoneTreeWidget> GetTreeWidget() const;
+		UE_API TSharedPtr<SMLDeformerInputWidget> GetInputWidget() const;
+		UE_API FText GetSectionTitle() const;
 
 	private:
-		void OnFilterTextChanged(const FText& InFilterText);
-		void RefreshTree(bool bBroadCastPropertyChanged=true);
-		bool BroadcastModelPropertyChanged(const FName PropertyName);
+		UE_API void OnFilterTextChanged(const FText& InFilterText);
+		UE_API void RefreshTree(bool bBroadCastPropertyChanged=true);
+		UE_API bool BroadcastModelPropertyChanged(const FName PropertyName);
 
-		void OnAddInputBones();
-		void OnDeleteInputBones();
-		void OnClearInputBones();
-		void OnAddAnimatedBones();
+		UE_API void OnAddInputBones();
+		UE_API void OnDeleteInputBones();
+		UE_API void OnClearInputBones();
+		UE_API void OnAddAnimatedBones();
 
 	private:
 		TSharedPtr<SMLDeformerInputBoneTreeWidget> TreeWidget;
@@ -130,3 +132,4 @@ namespace UE::MLDeformer
 
 }	// namespace UE::MLDeformer
 
+#undef UE_API
diff --git a/Engine/Plugins/Animation/MLDeformer/MLDeformerFramework/Source/MLDeformerFrameworkEditor/Public/SMLDeformerInputCurvesWidget.h b/Engine/Plugins/Animation/MLDeformer/MLDeformerFramework/Source/MLDeformerFrameworkEditor/Public/SMLDeformerInputCurvesWidget.h
--- a/Engine/Plugins/Animation/MLDeformer/MLDeformerFramework/Source/MLDeformerFrameworkEditor/Public/SMLDeformerInputCurvesWidget.h
+++ b/Engine/Plugins/Animation/MLDeformer/MLDeformerFramework/Source/MLDeformerFrameworkEditor/Public/SMLDeformerInputCurvesWidget.h
@@ -11,6 +11,8 @@
 #include "UObject/NameTypes.h"
 #include "Styling/SlateColor.h"
 
+#define UE_API MLDEFORMERFRAMEWORKEDITOR_API
+
 class USkeletalMesh;
 
 namespace UE::MLDeformer
@@ -35,11 +37,11 @@ namespace UE::MLDeformer
 	};
 
 
-	class MLDEFORMERFRAMEWORKEDITOR_API FMLDeformerInputCurveListElement
+	class FMLDeformerInputCurveListElement
 		: public TSharedFromThis<FMLDeformerInputCurveListElement>
 	{
 	public:
-		TSharedRef<ITableRow> MakeRowWidget(const TSharedRef<STableViewBase>& InOwnerTable, TSharedRef<FMLDeformerInputCurveListElement> InTreeElement, TSharedPtr<SMLDeformerInputCurveListWidget> InTreeWidget);
+		UE_API TSharedRef<ITableRow> MakeRowWidget(const TSharedRef<STableViewBase>& InOwnerTable, TSharedRef<FMLDeformerInputCurveListElement> InTreeElement, TSharedPtr<SMLDeformerInputCurveListWidget> InTreeWidget);
 
 	public:
 		FName Name;
@@ -47,21 +49,21 @@ namespace UE::MLDeformer
 	};
 
 
-	class MLDEFORMERFRAMEWORKEDITOR_API SMLDeformerInputCurveListRowWidget 
+	class SMLDeformerInputCurveListRowWidget 
 		: public STableRow<TSharedPtr<FMLDeformerInputCurveListElement>>
 	{
 	public:
-		void Construct(const FArguments& InArgs, const TSharedRef<STableViewBase>& OwnerTable, TSharedRef<FMLDeformerInputCurveListElement> InElement, TSharedPtr<SMLDeformerInputCurveListWidget> InListView);
+		UE_API void Construct(const FArguments& InArgs, const TSharedRef<STableViewBase>& OwnerTable, TSharedRef<FMLDeformerInputCurveListElement> InElement, TSharedPtr<SMLDeformerInputCurveListWidget> InListView);
 
 	private:
 		TWeakPtr<FMLDeformerInputCurveListElement> WeakElement;
-		FText GetName() const;
+		UE_API FText GetName() const;
 
 		friend class SMLDeformerInputCurveListWidget; 
 	};
 
 
-	class MLDEFORMERFRAMEWORKEDITOR_API SMLDeformerInputCurveListWidget
+	class SMLDeformerInputCurveListWidget
 		: public SListView<TSharedPtr<FMLDeformerInputCurveListElement>>
 	{
 		SLATE_BEGIN_ARGS(SMLDeformerInputCurveListWidget) {}
@@ -69,17 +71,17 @@ namespace UE::MLDeformer
 		SLATE_END_ARGS()
 
 	public:
-		void Construct(const FArguments& InArgs);
+		UE_API void Construct(const FArguments& InArgs);
 		const TArray<TSharedPtr<FMLDeformerInputCurveListElement>>& GetElements() const	{ return Elements; }
-		void AddElement(TSharedPtr<FMLDeformerInputCurveListElement> Element);
-		void RefreshElements(const TArray<FName>& CurveNames, USkeletalMesh* SkeletalMesh);
+		UE_API void AddElement(TSharedPtr<FMLDeformerInputCurveListElement> Element);
+		UE_API void RefreshElements(const TArray<FName>& CurveNames, USkeletalMesh* SkeletalMesh);
 
 	private:
-		TSharedRef<ITableRow> MakeTableRowWidget(TSharedPtr<FMLDeformerInputCurveListElement> InItem, const TSharedRef<STableViewBase>& OwnerTable);
-		FReply OnKeyDown(const FGeometry& InGeometry, const FKeyEvent& InKeyEvent) override;
-		TSharedPtr<SWidget> OnContextMenuOpening() const;
-		void SortElements();
-		void OnSelectionChanged(TSharedPtr<FMLDeformerInputCurveListElement> Selection, ESelectInfo::Type SelectInfo);
+		UE_API TSharedRef<ITableRow> MakeTableRowWidget(TSharedPtr<FMLDeformerInputCurveListElement> InItem, const TSharedRef<STableViewBase>& OwnerTable);
+		UE_API FReply OnKeyDown(const FGeometry& InGeometry, const FKeyEvent& InKeyEvent) override;
+		UE_API TSharedPtr<SWidget> OnContextMenuOpening() const;
+		UE_API void SortElements();
+		UE_API void OnSelectionChanged(TSharedPtr<FMLDeformerInputCurveListElement> Selection, ESelectInfo::Type SelectInfo);
 
 	private:
 		TArray<TSharedPtr<FMLDeformerInputCurveListElement>> Elements;
@@ -87,7 +89,7 @@ namespace UE::MLDeformer
 	};
 
 
-	class MLDEFORMERFRAMEWORKEDITOR_API SMLDeformerInputCurvesWidget
+	class SMLDeformerInputCurvesWidget
 		: public SCompoundWidget
 	{
 		friend class SMLDeformerInputCurveListWidget;
@@ -98,22 +100,22 @@ namespace UE::MLDeformer
 		SLATE_END_ARGS()
 
 	public:
-		void Construct(const FArguments& InArgs);
-		void Refresh();
-		void BindCommands(TSharedPtr<FUICommandList> CommandList);
-		FText GetSectionTitle() const;
-		TSharedPtr<SMLDeformerInputCurveListWidget> GetListWidget() const;
-		TSharedPtr<SMLDeformerInputWidget> GetInputWidget() const;
+		UE_API void Construct(const FArguments& InArgs);
+		UE_API void Refresh();
+		UE_API void BindCommands(TSharedPtr<FUICommandList> CommandList);
+		UE_API FText GetSectionTitle() const;
+		UE_API TSharedPtr<SMLDeformerInputCurveListWidget> GetListWidget() const;
+		UE_API TSharedPtr<SMLDeformerInputWidget> GetInputWidget() const;
 
 	private:
-		void OnFilterTextChanged(const FText& InFilterText);
-		void RefreshList(bool bBroadCastPropertyChanged=true);
-		bool BroadcastModelPropertyChanged(const FName PropertyName);
+		UE_API void OnFilterTextChanged(const FText& InFilterText);
+		UE_API void RefreshList(bool bBroadCastPropertyChanged=true);
+		UE_API bool BroadcastModelPropertyChanged(const FName PropertyName);
 
-		void OnAddInputCurves();
-		void OnDeleteInputCurves();
-		void OnClearInputCurves();
-		void OnAddAnimatedCurves();
+		UE_API void OnAddInputCurves();
+		UE_API void OnDeleteInputCurves();
+		UE_API void OnClearInputCurves();
+		UE_API void OnAddAnimatedCurves();
 
 	private:
 		TSharedPtr<SMLDeformerInputWidget> InputWidget;
@@ -125,3 +127,4 @@ namespace UE::MLDeformer
 
 }	// namespace UE::MLDeformer
 
+#undef UE_API
diff --git a/Engine/Plugins/Animation/MLDeformer/MLDeformerFramework/Source/MLDeformerFrameworkEditor/Public/SMLDeformerInputWidget.h b/Engine/Plugins/Animation/MLDeformer/MLDeformerFramework/Source/MLDeformerFrameworkEditor/Public/SMLDeformerInputWidget.h
--- a/Engine/Plugins/Animation/MLDeformer/MLDeformerFramework/Source/MLDeformerFrameworkEditor/Public/SMLDeformerInputWidget.h
+++ b/Engine/Plugins/Animation/MLDeformer/MLDeformerFramework/Source/MLDeformerFrameworkEditor/Public/SMLDeformerInputWidget.h
@@ -8,6 +8,8 @@
 #include "Framework/SlateDelegates.h"
 #include "Framework/Commands/UICommandList.h"
 
+#define UE_API MLDEFORMERFRAMEWORKEDITOR_API
+
 class SWidget;
 class FMenuBuilder;
 class SVerticalBox;
@@ -18,7 +20,7 @@ namespace UE::MLDeformer
 	class SMLDeformerInputCurvesWidget;
 	class FMLDeformerEditorModel;
 
-	class MLDEFORMERFRAMEWORKEDITOR_API SMLDeformerInputWidget
+	class SMLDeformerInputWidget
 		: public SCompoundWidget
 	{
 	public:
@@ -26,23 +28,23 @@ namespace UE::MLDeformer
 		SLATE_ARGUMENT(FMLDeformerEditorModel*, EditorModel)
 		SLATE_END_ARGS()
 
-		struct MLDEFORMERFRAMEWORKEDITOR_API FSectionInfo
+		struct FSectionInfo
 		{
 			TAttribute<FText> PlusButtonTooltip;
 			TAttribute<FText> SectionTitle;
 			FOnClicked PlusButtonPressed;
 		};
 
-		void Construct(const FArguments& InArgs);
-		void AddSection(TSharedPtr<SWidget> Widget, const FSectionInfo& SectionInfo);
+		UE_API void Construct(const FArguments& InArgs);
+		UE_API void AddSection(TSharedPtr<SWidget> Widget, const FSectionInfo& SectionInfo);
 
-		TSharedPtr<FUICommandList> GetBonesCommandList() const;
-		TSharedPtr<FUICommandList> GetCurvesCommandList() const;
+		UE_API TSharedPtr<FUICommandList> GetBonesCommandList() const;
+		UE_API TSharedPtr<FUICommandList> GetCurvesCommandList() const;
 
 		/**
 		 * Refresh all sub-widget contents.
 		 */
-		virtual void Refresh();
+		UE_API virtual void Refresh();
 
 		virtual void AddInputBonesMenuItems(FMenuBuilder& MenuBuilder) {}
 		virtual void AddInputBonesPlusIconMenuItems(FMenuBuilder& MenuBuilder) {}
@@ -59,23 +61,23 @@ namespace UE::MLDeformer
 		virtual void OnDeleteInputCurves(const TArray<FName>& Names) {}
 		virtual void OnAddAnimatedCurves() {}
 
-		virtual void ClearSelectionForAllWidgetsExceptThis(TSharedPtr<SWidget> ExceptThisWidget);
+		UE_API virtual void ClearSelectionForAllWidgetsExceptThis(TSharedPtr<SWidget> ExceptThisWidget);
 		virtual void OnSelectInputBone(FName BoneName) {}
 		virtual void OnSelectInputCurve(FName BoneName) {}
 
-		virtual TSharedPtr<SWidget> GetExtraBonePickerWidget();
+		UE_API virtual TSharedPtr<SWidget> GetExtraBonePickerWidget();
 
 		// Registers the commands that are being executed by the different widgets.
-		static void RegisterCommands();
-		static void UnregisterCommands();
+		static UE_API void RegisterCommands();
+		static UE_API void UnregisterCommands();
 
 	protected:
-		void AddSectionSeparator();
-		void CreateBonesWidget();
-		void CreateCurvesWidget();
+		UE_API void AddSectionSeparator();
+		UE_API void CreateBonesWidget();
+		UE_API void CreateCurvesWidget();
 
-		FReply ShowCurvesPlusIconContextMenu();
-		FReply ShowBonesPlusIconContextMenu();
+		UE_API FReply ShowCurvesPlusIconContextMenu();
+		UE_API FReply ShowBonesPlusIconContextMenu();
 
 	protected:
 		TSharedPtr<SMLDeformerInputBonesWidget> InputBonesWidget;
@@ -86,3 +88,5 @@ namespace UE::MLDeformer
 		FMLDeformerEditorModel* EditorModel = nullptr;
 	};
 }
+
+#undef UE_API
diff --git a/Engine/Plugins/Animation/MLDeformer/MLDeformerFramework/Source/MLDeformerFrameworkEditor/Public/SMLDeformerMaskConfigDialog.h b/Engine/Plugins/Animation/MLDeformer/MLDeformerFramework/Source/MLDeformerFrameworkEditor/Public/SMLDeformerMaskConfigDialog.h
--- a/Engine/Plugins/Animation/MLDeformer/MLDeformerFramework/Source/MLDeformerFrameworkEditor/Public/SMLDeformerMaskConfigDialog.h
+++ b/Engine/Plugins/Animation/MLDeformer/MLDeformerFramework/Source/MLDeformerFrameworkEditor/Public/SMLDeformerMaskConfigDialog.h
@@ -10,6 +10,8 @@
 #include "Widgets/Input/STextComboBox.h"
 #include "Widgets/DeclarativeSyntaxSupport.h"
 
+#define UE_API MLDEFORMERFRAMEWORKEDITOR_API
+
 namespace UE::MLDeformer
 {
 	class FMLDeformerEditorModel;
@@ -21,7 +23,7 @@ namespace UE::MLDeformer
 	 * When using a vertex attribute, you can pick which attribute to use, or create a new one.
 	 * Creating a new attribute will pop up the SMLDeformerNewVertexAttributeDialog.
 	 */
-	class MLDEFORMERFRAMEWORKEDITOR_API SMLDeformerMaskConfigDialog
+	class SMLDeformerMaskConfigDialog
 		: public SCustomDialog
 	{
 	public:
@@ -30,7 +32,7 @@ namespace UE::MLDeformer
 		SLATE_EVENT(FOnSetNewVertexAttributeValues, OnSetNewVertexAttributeValues)
 		SLATE_END_ARGS()
 
-		void Construct(const FArguments& InArgs, FMLDeformerEditorModel* InEditorModel);
+		UE_API void Construct(const FArguments& InArgs, FMLDeformerEditorModel* InEditorModel);
 
 		/**
 		 * This returns the mask info object how it was setup in the dialog.
@@ -39,10 +41,10 @@ namespace UE::MLDeformer
 		const FMLDeformerMaskInfo& GetMaskInfo() const { return MaskInfo; }
 
 	private:
-		void UpdateAttributeNames();
-		TSharedPtr<SWidget> CreateMeshAttributeModeWidget();
-		TSharedPtr<FString> MaskModeEnumToString(EMLDeformerMaskingMode MaskMode) const;
-		FReply OnCreateVertexAttributeClicked();
+		UE_API void UpdateAttributeNames();
+		UE_API TSharedPtr<SWidget> CreateMeshAttributeModeWidget();
+		UE_API TSharedPtr<FString> MaskModeEnumToString(EMLDeformerMaskingMode MaskMode) const;
+		UE_API FReply OnCreateVertexAttributeClicked();
 
 	private:
 		TArray<FName> AttributeNames;
@@ -54,3 +56,5 @@ namespace UE::MLDeformer
 	};
 
 }	// namespace UE::MLDeformer
+
+#undef UE_API
diff --git a/Engine/Plugins/Animation/MLDeformer/MLDeformerFramework/Source/MLDeformerFrameworkEditor/Public/Tools/TrainingDataProcessor/AnimCustomization.h b/Engine/Plugins/Animation/MLDeformer/MLDeformerFramework/Source/MLDeformerFrameworkEditor/Public/Tools/TrainingDataProcessor/AnimCustomization.h
--- a/Engine/Plugins/Animation/MLDeformer/MLDeformerFramework/Source/MLDeformerFrameworkEditor/Public/Tools/TrainingDataProcessor/AnimCustomization.h
+++ b/Engine/Plugins/Animation/MLDeformer/MLDeformerFramework/Source/MLDeformerFrameworkEditor/Public/Tools/TrainingDataProcessor/AnimCustomization.h
@@ -6,6 +6,8 @@
 #include "PropertyHandle.h"
 #include "DetailWidgetRow.h"
 
+#define UE_API MLDEFORMERFRAMEWORKEDITOR_API
+
 class UMLDeformerTrainingDataProcessorSettings;
 
 namespace UE::MLDeformer::TrainingDataProcessor
@@ -14,22 +16,24 @@ namespace UE::MLDeformer::TrainingDataProcessor
 	 * The property detail customization for a single animation input that can be enabled or disabled.
 	 * This is a customization for the type FMLDeformerTrainingDataProcessorAnim.
 	 */
-	class MLDEFORMERFRAMEWORKEDITOR_API FAnimCustomization final : public IPropertyTypeCustomization
+	class FAnimCustomization final : public IPropertyTypeCustomization
 	{
 	public:
-		static TSharedRef<IPropertyTypeCustomization> MakeInstance();
+		static UE_API TSharedRef<IPropertyTypeCustomization> MakeInstance();
 
 		//~ Begin IPropertyTypeCustomization overrides.
-		virtual void CustomizeHeader(TSharedRef<IPropertyHandle> StructPropertyHandle, class FDetailWidgetRow& HeaderRow, IPropertyTypeCustomizationUtils& StructCustomizationUtils) override;
-		virtual void CustomizeChildren(TSharedRef<IPropertyHandle> StructPropertyHandle, class IDetailChildrenBuilder& ChildBuilder, IPropertyTypeCustomizationUtils& StructCustomizationUtils) override;
-		EVisibility GetAnimErrorVisibility(const TSharedRef<IPropertyHandle> StructPropertyHandle, int32 AnimIndex) const;
+		UE_API virtual void CustomizeHeader(TSharedRef<IPropertyHandle> StructPropertyHandle, class FDetailWidgetRow& HeaderRow, IPropertyTypeCustomizationUtils& StructCustomizationUtils) override;
+		UE_API virtual void CustomizeChildren(TSharedRef<IPropertyHandle> StructPropertyHandle, class IDetailChildrenBuilder& ChildBuilder, IPropertyTypeCustomizationUtils& StructCustomizationUtils) override;
+		UE_API EVisibility GetAnimErrorVisibility(const TSharedRef<IPropertyHandle> StructPropertyHandle, int32 AnimIndex) const;
 		// ~End IPropertyTypeCustomization overrides.
 
 	private:
-		void RefreshProperties() const;
-		static UMLDeformerTrainingDataProcessorSettings* FindSettings(const TSharedRef<IPropertyHandle>& StructPropertyHandle);
+		UE_API void RefreshProperties() const;
+		static UE_API UMLDeformerTrainingDataProcessorSettings* FindSettings(const TSharedRef<IPropertyHandle>& StructPropertyHandle);
 
 	private:
 		TSharedPtr<IPropertyUtilities> PropertyUtilities;
 	};
 }	// namespace UE::MLDeformer::TrainingDataProcessor
+
+#undef UE_API
diff --git a/Engine/Plugins/Animation/MLDeformer/MLDeformerFramework/Source/MLDeformerFrameworkEditor/Public/Tools/TrainingDataProcessor/BoneGroupsListCustomization.h b/Engine/Plugins/Animation/MLDeformer/MLDeformerFramework/Source/MLDeformerFrameworkEditor/Public/Tools/TrainingDataProcessor/BoneGroupsListCustomization.h
--- a/Engine/Plugins/Animation/MLDeformer/MLDeformerFramework/Source/MLDeformerFrameworkEditor/Public/Tools/TrainingDataProcessor/BoneGroupsListCustomization.h
+++ b/Engine/Plugins/Animation/MLDeformer/MLDeformerFramework/Source/MLDeformerFrameworkEditor/Public/Tools/TrainingDataProcessor/BoneGroupsListCustomization.h
@@ -6,6 +6,8 @@
 #include "PropertyHandle.h"
 #include "MLDeformerTrainingDataProcessorSettings.h"
 
+#define UE_API MLDEFORMERFRAMEWORKEDITOR_API
+
 class UMLDeformerTrainingDataProcessorSettings;
 class USkeleton;
 
@@ -15,23 +17,25 @@ namespace UE::MLDeformer::TrainingDataProcessor
 	 * The property detail customization for a list of bone groups.
 	 * This is a customization for the type FMLDeformerTrainingDataProcessorBoneGroupsList.
 	 */
-	class MLDEFORMERFRAMEWORKEDITOR_API FBoneGroupsListCustomization : public IPropertyTypeCustomization
+	class FBoneGroupsListCustomization : public IPropertyTypeCustomization
 	{
 	public:
-		static TSharedRef<IPropertyTypeCustomization> MakeInstance();
+		static UE_API TSharedRef<IPropertyTypeCustomization> MakeInstance();
 
 		//~ Begin IPropertyTypeCustomization overrides.
-		virtual void CustomizeHeader(TSharedRef<IPropertyHandle> StructPropertyHandle, class FDetailWidgetRow& HeaderRow,
+		UE_API virtual void CustomizeHeader(TSharedRef<IPropertyHandle> StructPropertyHandle, class FDetailWidgetRow& HeaderRow,
 		                             IPropertyTypeCustomizationUtils& StructCustomizationUtils) override;
 		virtual void CustomizeChildren(TSharedRef<IPropertyHandle> StructPropertyHandle, class IDetailChildrenBuilder& ChildBuilder,
 		                               IPropertyTypeCustomizationUtils& StructCustomizationUtils) override {}
 		//~ End IPropertyTypeCustomization overrides.
 
 	private:
-		TArray<FMLDeformerTrainingDataProcessorBoneGroup>* GetBoneGroups() const;
+		UE_API TArray<FMLDeformerTrainingDataProcessorBoneGroup>* GetBoneGroups() const;
 
 	private:
 		TSharedPtr<IPropertyUtilities> PropertyUtilities;
 		TSharedPtr<IPropertyHandle> StructProperty;
 	};
 } // namespace UE::MLDeformer::TrainingDataProcessor
+
+#undef UE_API
diff --git a/Engine/Plugins/Animation/MLDeformer/MLDeformerFramework/Source/MLDeformerFrameworkEditor/Public/Tools/TrainingDataProcessor/BoneListCustomization.h b/Engine/Plugins/Animation/MLDeformer/MLDeformerFramework/Source/MLDeformerFrameworkEditor/Public/Tools/TrainingDataProcessor/BoneListCustomization.h
--- a/Engine/Plugins/Animation/MLDeformer/MLDeformerFramework/Source/MLDeformerFrameworkEditor/Public/Tools/TrainingDataProcessor/BoneListCustomization.h
+++ b/Engine/Plugins/Animation/MLDeformer/MLDeformerFramework/Source/MLDeformerFrameworkEditor/Public/Tools/TrainingDataProcessor/BoneListCustomization.h
@@ -5,6 +5,8 @@
 #include "IPropertyTypeCustomization.h"
 #include "PropertyHandle.h"
 
+#define UE_API MLDEFORMERFRAMEWORKEDITOR_API
+
 class UMLDeformerTrainingDataProcessorSettings;
 class USkeleton;
 
@@ -14,23 +16,25 @@ namespace UE::MLDeformer::TrainingDataProcessor
 	 * The property detail customization for a list of bones.
 	 * This is a detail customization for the type FMLDeformerTrainingDataProcessorBoneList.
 	 */
-	class MLDEFORMERFRAMEWORKEDITOR_API FBoneListCustomization : public IPropertyTypeCustomization
+	class FBoneListCustomization : public IPropertyTypeCustomization
 	{
 	public:
-		static TSharedRef<IPropertyTypeCustomization> MakeInstance();
+		static UE_API TSharedRef<IPropertyTypeCustomization> MakeInstance();
 
 		//~ Begin IPropertyTypeCustomization overrides.
-		virtual void CustomizeHeader(TSharedRef<IPropertyHandle> StructPropertyHandle, class FDetailWidgetRow& HeaderRow,
+		UE_API virtual void CustomizeHeader(TSharedRef<IPropertyHandle> StructPropertyHandle, class FDetailWidgetRow& HeaderRow,
 		                             IPropertyTypeCustomizationUtils& StructCustomizationUtils) override;
 		virtual void CustomizeChildren(TSharedRef<IPropertyHandle> StructPropertyHandle, class IDetailChildrenBuilder& ChildBuilder,
 		                               IPropertyTypeCustomizationUtils& StructCustomizationUtils) override { }
 		// ~End IPropertyTypeCustomization overrides.
 
 	private:
-		TArray<FName>* GetBoneNames() const;
+		UE_API TArray<FName>* GetBoneNames() const;
 
 	private:
 		TSharedPtr<IPropertyUtilities> PropertyUtilities;
 		TSharedPtr<IPropertyHandle> StructProperty;
 	};
 } // namespace UE::MLDeformer::TrainingDataProcessor
+
+#undef UE_API
diff --git a/Engine/Plugins/Animation/MLDeformer/MLDeformerFramework/Source/MLDeformerFrameworkEditor/Public/Tools/TrainingDataProcessor/SBoneGroupsListWidget.h b/Engine/Plugins/Animation/MLDeformer/MLDeformerFramework/Source/MLDeformerFrameworkEditor/Public/Tools/TrainingDataProcessor/SBoneGroupsListWidget.h
--- a/Engine/Plugins/Animation/MLDeformer/MLDeformerFramework/Source/MLDeformerFrameworkEditor/Public/Tools/TrainingDataProcessor/SBoneGroupsListWidget.h
+++ b/Engine/Plugins/Animation/MLDeformer/MLDeformerFramework/Source/MLDeformerFrameworkEditor/Public/Tools/TrainingDataProcessor/SBoneGroupsListWidget.h
@@ -11,6 +11,8 @@
 #include "Styling/SlateColor.h"
 #include "EditorUndoClient.h"
 
+#define UE_API MLDEFORMERFRAMEWORKEDITOR_API
+
 class USkeleton;
 class FUICommandList;
 class FNotifyHook;
@@ -35,10 +37,10 @@ namespace UE::MLDeformer::TrainingDataProcessor
 	};
 
 
-	class MLDEFORMERFRAMEWORKEDITOR_API FBoneGroupTreeElement final : public TSharedFromThis<FBoneGroupTreeElement>
+	class FBoneGroupTreeElement final : public TSharedFromThis<FBoneGroupTreeElement>
 	{
 	public:
-		static TSharedRef<ITableRow> MakeTreeRowWidget(const TSharedRef<STableViewBase>& InOwnerTable,
+		static UE_API TSharedRef<ITableRow> MakeTreeRowWidget(const TSharedRef<STableViewBase>& InOwnerTable,
 		                                               const TSharedRef<FBoneGroupTreeElement>& InTreeElement,
 		                                               const TSharedPtr<SBoneGroupsTreeWidget>& InTreeWidget);
 
@@ -54,14 +56,14 @@ namespace UE::MLDeformer::TrainingDataProcessor
 	};
 
 
-	class MLDEFORMERFRAMEWORKEDITOR_API SBoneGroupTreeRowWidget final : public STableRow<TSharedPtr<FBoneGroupTreeElement>>
+	class SBoneGroupTreeRowWidget final : public STableRow<TSharedPtr<FBoneGroupTreeElement>>
 	{
 	public:
-		void Construct(const FArguments& InArgs, const TSharedRef<STableViewBase>& OwnerTable, const TSharedRef<FBoneGroupTreeElement>& InTreeElement,
+		UE_API void Construct(const FArguments& InArgs, const TSharedRef<STableViewBase>& OwnerTable, const TSharedRef<FBoneGroupTreeElement>& InTreeElement,
 		               const TSharedPtr<SBoneGroupsTreeWidget>& InTreeView);
 
 	private:
-		FText GetName() const;
+		UE_API FText GetName() const;
 
 	private:
 		TWeakPtr<FBoneGroupTreeElement> WeakTreeElement;
@@ -69,7 +71,7 @@ namespace UE::MLDeformer::TrainingDataProcessor
 	};
 
 
-	class MLDEFORMERFRAMEWORKEDITOR_API SBoneGroupsTreeWidget final : public STreeView<TSharedPtr<FBoneGroupTreeElement>>
+	class SBoneGroupsTreeWidget final : public STreeView<TSharedPtr<FBoneGroupTreeElement>>
 	{
 	public:
 		friend class SBoneGroupsListWidget;
@@ -81,24 +83,24 @@ namespace UE::MLDeformer::TrainingDataProcessor
 			SLATE_ARGUMENT(TSharedPtr<SBoneGroupsListWidget>, BoneGroupsWidget)
 		SLATE_END_ARGS()
 
-		void Construct(const FArguments& InArgs);
-		void Refresh();
-		int32 GetNumSelectedGroups() const;
+		UE_API void Construct(const FArguments& InArgs);
+		UE_API void Refresh();
+		UE_API int32 GetNumSelectedGroups() const;
 
 	private:
 		//~ Begin STreeView overrides.
-		virtual FReply OnKeyDown(const FGeometry& InGeometry, const FKeyEvent& InKeyEvent) override;
+		UE_API virtual FReply OnKeyDown(const FGeometry& InGeometry, const FKeyEvent& InKeyEvent) override;
 		//~ End STreeView overrides.
 
-		void AddElement(const TSharedPtr<FBoneGroupTreeElement>& Element, const TSharedPtr<FBoneGroupTreeElement>& ParentElement);
-		const TArray<TSharedPtr<FBoneGroupTreeElement>>& GetRootElements() const;
-		TSharedRef<ITableRow> MakeTableRowWidget(TSharedPtr<FBoneGroupTreeElement> InItem, const TSharedRef<STableViewBase>& OwnerTable);
-		TSharedPtr<SWidget> CreateContextMenuWidget() const;
-		void UpdateTreeElements();
-		void RefreshTree();
+		UE_API void AddElement(const TSharedPtr<FBoneGroupTreeElement>& Element, const TSharedPtr<FBoneGroupTreeElement>& ParentElement);
+		UE_API const TArray<TSharedPtr<FBoneGroupTreeElement>>& GetRootElements() const;
+		UE_API TSharedRef<ITableRow> MakeTableRowWidget(TSharedPtr<FBoneGroupTreeElement> InItem, const TSharedRef<STableViewBase>& OwnerTable);
+		UE_API TSharedPtr<SWidget> CreateContextMenuWidget() const;
+		UE_API void UpdateTreeElements();
+		UE_API void RefreshTree();
 
-		static void HandleGetChildrenForTree(TSharedPtr<FBoneGroupTreeElement> InItem, TArray<TSharedPtr<FBoneGroupTreeElement>>& OutChildren);
-		static TSharedPtr<SWidget> CreateContextWidget();
+		static UE_API void HandleGetChildrenForTree(TSharedPtr<FBoneGroupTreeElement> InItem, TArray<TSharedPtr<FBoneGroupTreeElement>>& OutChildren);
+		static UE_API TSharedPtr<SWidget> CreateContextWidget();
 
 	private:
 		TArray<TSharedPtr<FBoneGroupTreeElement>> RootElements;
@@ -113,7 +115,7 @@ namespace UE::MLDeformer::TrainingDataProcessor
 	 * We see a bone group as a list of bone names. Multiple bone groups can exist. If you need only one list of bones
 	 * then you can use the bone SBoneListWidget instead.
 	 */
-	class MLDEFORMERFRAMEWORKEDITOR_API SBoneGroupsListWidget final : public SCompoundWidget, public FEditorUndoClient
+	class SBoneGroupsListWidget final : public SCompoundWidget, public FEditorUndoClient
 	{
 		friend class SBoneGroupsTreeWidget;
 
@@ -127,32 +129,32 @@ namespace UE::MLDeformer::TrainingDataProcessor
 		SLATE_END_ARGS()
 
 	public:
-		virtual ~SBoneGroupsListWidget() override;
+		UE_API virtual ~SBoneGroupsListWidget() override;
 
-		void Construct(const FArguments& InArgs, FNotifyHook* InNotifyHook);
+		UE_API void Construct(const FArguments& InArgs, FNotifyHook* InNotifyHook);
 
 	private:
 		//~ Begin FEditorUndoClient overrides.
-		virtual void PostUndo(bool bSuccess) override;
-		virtual void PostRedo(bool bSuccess) override;
+		UE_API virtual void PostUndo(bool bSuccess) override;
+		UE_API virtual void PostRedo(bool bSuccess) override;
 		//~ End FEditorUndoClient overrides.
 
-		void BindCommands(const TSharedPtr<FUICommandList>& InCommandList);
-		void OnFilterTextChanged(const FText& InFilterText);
-		void RefreshTree() const;
-		FReply OnAddButtonClicked() const;
-		FReply OnClearButtonClicked() const;
+		UE_API void BindCommands(const TSharedPtr<FUICommandList>& InCommandList);
+		UE_API void OnFilterTextChanged(const FText& InFilterText);
+		UE_API void RefreshTree() const;
+		UE_API FReply OnAddButtonClicked() const;
+		UE_API FReply OnClearButtonClicked() const;
 
-		TSharedPtr<SBoneGroupsTreeWidget> GetTreeWidget() const;
-		TSharedPtr<FUICommandList> GetCommandList() const;
-		TArray<FMLDeformerTrainingDataProcessorBoneGroup>* GetBoneGroupsValues() const;
-		TWeakObjectPtr<USkeleton> GetSkeleton() const;
-		const FString& GetFilterText() const;
+		UE_API TSharedPtr<SBoneGroupsTreeWidget> GetTreeWidget() const;
+		UE_API TSharedPtr<FUICommandList> GetCommandList() const;
+		UE_API TArray<FMLDeformerTrainingDataProcessorBoneGroup>* GetBoneGroupsValues() const;
+		UE_API TWeakObjectPtr<USkeleton> GetSkeleton() const;
+		UE_API const FString& GetFilterText() const;
 
-		void OnCreateBoneGroup() const;
-		void OnDeleteSelectedItems() const;
-		void OnClearBoneGroups() const;
-		void OnAddBoneToGroup() const;
+		UE_API void OnCreateBoneGroup() const;
+		UE_API void OnDeleteSelectedItems() const;
+		UE_API void OnClearBoneGroups() const;
+		UE_API void OnAddBoneToGroup() const;
 
 	private:
 		TSharedPtr<SBoneGroupsTreeWidget> TreeWidget;
@@ -164,3 +166,5 @@ namespace UE::MLDeformer::TrainingDataProcessor
 		FNotifyHook* NotifyHook = nullptr;
 	};
 } // namespace UE::MLDeformer::TrainingDataProcessor
+
+#undef UE_API
diff --git a/Engine/Plugins/Animation/MLDeformer/MLDeformerFramework/Source/MLDeformerFrameworkEditor/Public/Tools/TrainingDataProcessor/SBoneListWidget.h b/Engine/Plugins/Animation/MLDeformer/MLDeformerFramework/Source/MLDeformerFrameworkEditor/Public/Tools/TrainingDataProcessor/SBoneListWidget.h
--- a/Engine/Plugins/Animation/MLDeformer/MLDeformerFramework/Source/MLDeformerFrameworkEditor/Public/Tools/TrainingDataProcessor/SBoneListWidget.h
+++ b/Engine/Plugins/Animation/MLDeformer/MLDeformerFramework/Source/MLDeformerFrameworkEditor/Public/Tools/TrainingDataProcessor/SBoneListWidget.h
@@ -12,6 +12,8 @@
 #include "Styling/SlateColor.h"
 #include "EditorUndoClient.h"
 
+#define UE_API MLDEFORMERFRAMEWORKEDITOR_API
+
 class USkeleton;
 class FUICommandList;
 class FNotifyHook;
@@ -35,10 +37,10 @@ namespace UE::MLDeformer::TrainingDataProcessor
 	};
 
 
-	class MLDEFORMERFRAMEWORKEDITOR_API FBoneTreeWidgetElement final : public TSharedFromThis<FBoneTreeWidgetElement>
+	class FBoneTreeWidgetElement final : public TSharedFromThis<FBoneTreeWidgetElement>
 	{
 	public:
-		TSharedRef<ITableRow> MakeTreeRowWidget(const TSharedRef<STableViewBase>& InOwnerTable, TSharedRef<FBoneTreeWidgetElement> InTreeElement,
+		UE_API TSharedRef<ITableRow> MakeTreeRowWidget(const TSharedRef<STableViewBase>& InOwnerTable, TSharedRef<FBoneTreeWidgetElement> InTreeElement,
 		                                        TSharedPtr<SBoneTreeWidget> InTreeWidget);
 
 	public:
@@ -48,39 +50,39 @@ namespace UE::MLDeformer::TrainingDataProcessor
 	};
 
 
-	class MLDEFORMERFRAMEWORKEDITOR_API SBoneTreeRowWidget final : public STableRow<TSharedPtr<FBoneTreeWidgetElement>>
+	class SBoneTreeRowWidget final : public STableRow<TSharedPtr<FBoneTreeWidgetElement>>
 	{
 	public:
-		void Construct(const FArguments& InArgs, const TSharedRef<STableViewBase>& OwnerTable, const TSharedRef<FBoneTreeWidgetElement>& InTreeElement,
+		UE_API void Construct(const FArguments& InArgs, const TSharedRef<STableViewBase>& OwnerTable, const TSharedRef<FBoneTreeWidgetElement>& InTreeElement,
 		               const TSharedPtr<SBoneTreeWidget>& InTreeView);
 
 	private:
 		TWeakPtr<FBoneTreeWidgetElement> TreeElement;
-		FText GetName() const;
+		UE_API FText GetName() const;
 	};
 
 
-	class MLDEFORMERFRAMEWORKEDITOR_API SBoneTreeWidget final : public STreeView<TSharedPtr<FBoneTreeWidgetElement>>
+	class SBoneTreeWidget final : public STreeView<TSharedPtr<FBoneTreeWidgetElement>>
 	{
 		SLATE_BEGIN_ARGS(SBoneTreeWidget) { }
 			SLATE_ARGUMENT(TSharedPtr<SBoneListWidget>, BoneListWidget)
 		SLATE_END_ARGS()
 
 	public:
-		void Construct(const FArguments& InArgs);
-		void RefreshElements(const TArray<FName>& BoneNames, const FReferenceSkeleton* RefSkeleton, const FString& FilterText);
-		TArray<FName> ExtractAllElementNames() const;
+		UE_API void Construct(const FArguments& InArgs);
+		UE_API void RefreshElements(const TArray<FName>& BoneNames, const FReferenceSkeleton* RefSkeleton, const FString& FilterText);
+		UE_API TArray<FName> ExtractAllElementNames() const;
 		const TArray<TSharedPtr<FBoneTreeWidgetElement>>& GetRootElements() const { return RootElements; }
 
 	private:
-		virtual FReply OnKeyDown(const FGeometry& InGeometry, const FKeyEvent& InKeyEvent) override;
+		UE_API virtual FReply OnKeyDown(const FGeometry& InGeometry, const FKeyEvent& InKeyEvent) override;
 		
-		static void RecursiveAddNames(const FBoneTreeWidgetElement& Element, TArray<FName>& OutNames);
-		TSharedRef<ITableRow> MakeTableRowWidget(TSharedPtr<FBoneTreeWidgetElement> InItem, const TSharedRef<STableViewBase>& OwnerTable);
-		TSharedPtr<SWidget> OnContextMenuOpening() const;
-		static TSharedPtr<FBoneTreeWidgetElement> FindParentElementForBone(FName BoneName, const FReferenceSkeleton& RefSkeleton, const TMap<FName, TSharedPtr<FBoneTreeWidgetElement>>& NameToElementMap);
-		void RecursiveSortElements(const TSharedPtr<FBoneTreeWidgetElement>& Element);
-		static void HandleGetChildrenForTree(TSharedPtr<FBoneTreeWidgetElement> InItem, TArray<TSharedPtr<FBoneTreeWidgetElement>>& OutChildren);
+		static UE_API void RecursiveAddNames(const FBoneTreeWidgetElement& Element, TArray<FName>& OutNames);
+		UE_API TSharedRef<ITableRow> MakeTableRowWidget(TSharedPtr<FBoneTreeWidgetElement> InItem, const TSharedRef<STableViewBase>& OwnerTable);
+		UE_API TSharedPtr<SWidget> OnContextMenuOpening() const;
+		static UE_API TSharedPtr<FBoneTreeWidgetElement> FindParentElementForBone(FName BoneName, const FReferenceSkeleton& RefSkeleton, const TMap<FName, TSharedPtr<FBoneTreeWidgetElement>>& NameToElementMap);
+		UE_API void RecursiveSortElements(const TSharedPtr<FBoneTreeWidgetElement>& Element);
+		static UE_API void HandleGetChildrenForTree(TSharedPtr<FBoneTreeWidgetElement> InItem, TArray<TSharedPtr<FBoneTreeWidgetElement>>& OutChildren);
 
 	private:
 		TArray<TSharedPtr<FBoneTreeWidgetElement>> RootElements;
@@ -98,7 +100,7 @@ namespace UE::MLDeformer::TrainingDataProcessor
 	 * It works directly on a TArray<FName> as source. You can use the GetBoneNames event to provide this array.
 	 * You can use some of the delegates to listen for changes to the array.
 	 */
-	class MLDEFORMERFRAMEWORKEDITOR_API SBoneListWidget final : public SCompoundWidget, public FEditorUndoClient
+	class SBoneListWidget final : public SCompoundWidget, public FEditorUndoClient
 	{
 		SLATE_BEGIN_ARGS(SBoneListWidget) {}
 			SLATE_ARGUMENT(TWeakObjectPtr<USkeleton>, Skeleton)
@@ -110,29 +112,29 @@ namespace UE::MLDeformer::TrainingDataProcessor
 		SLATE_END_ARGS()
 
 	public:
-		virtual ~SBoneListWidget() override;
+		UE_API virtual ~SBoneListWidget() override;
 
-		void Construct(const FArguments& InArgs, FNotifyHook* InNotifyHook);
-		void Refresh() const;
-		TSharedPtr<SBoneTreeWidget> GetTreeWidget() const;
-		TSharedPtr<FUICommandList> GetCommandList() const;
+		UE_API void Construct(const FArguments& InArgs, FNotifyHook* InNotifyHook);
+		UE_API void Refresh() const;
+		UE_API TSharedPtr<SBoneTreeWidget> GetTreeWidget() const;
+		UE_API TSharedPtr<FUICommandList> GetCommandList() const;
 
 	private:
 		// FEditorUndoClient overrides.
-		virtual void PostUndo(bool bSuccess) override;
-		virtual void PostRedo(bool bSuccess) override;
+		UE_API virtual void PostUndo(bool bSuccess) override;
+		UE_API virtual void PostRedo(bool bSuccess) override;
 		// ~END FEditorUndoClient overrides.
 
-		void BindCommands(const TSharedPtr<FUICommandList>& InCommandList);
-		void OnFilterTextChanged(const FText& InFilterText);
-		void RefreshTree() const;
-		FReply OnAddBonesButtonClicked() const;
-		FReply OnClearBonesButtonClicked() const;
-		void NotifyPropertyChanged() const;
+		UE_API void BindCommands(const TSharedPtr<FUICommandList>& InCommandList);
+		UE_API void OnFilterTextChanged(const FText& InFilterText);
+		UE_API void RefreshTree() const;
+		UE_API FReply OnAddBonesButtonClicked() const;
+		UE_API FReply OnClearBonesButtonClicked() const;
+		UE_API void NotifyPropertyChanged() const;
 
-		void OnAddBones() const;
-		void OnRemoveBones() const;
-		void OnClearBones() const;
+		UE_API void OnAddBones() const;
+		UE_API void OnRemoveBones() const;
+		UE_API void OnClearBones() const;
 
 	private:
 		TSharedPtr<SBoneTreeWidget> TreeWidget;
@@ -147,3 +149,5 @@ namespace UE::MLDeformer::TrainingDataProcessor
 		FBoneListWidgetGetBoneNames GetBoneNames;
 	};
 } // namespace UE::MLDeformer::TrainingDataProcessor
+
+#undef UE_API
diff --git a/Engine/Plugins/Animation/MLDeformer/MLDeformerFramework/Source/MLDeformerFrameworkEditor/Public/Tools/TrainingDataProcessor/STrainingDataProcessorMainWidget.h b/Engine/Plugins/Animation/MLDeformer/MLDeformerFramework/Source/MLDeformerFrameworkEditor/Public/Tools/TrainingDataProcessor/STrainingDataProcessorMainWidget.h
--- a/Engine/Plugins/Animation/MLDeformer/MLDeformerFramework/Source/MLDeformerFrameworkEditor/Public/Tools/TrainingDataProcessor/STrainingDataProcessorMainWidget.h
+++ b/Engine/Plugins/Animation/MLDeformer/MLDeformerFramework/Source/MLDeformerFrameworkEditor/Public/Tools/TrainingDataProcessor/STrainingDataProcessorMainWidget.h
@@ -9,6 +9,8 @@
 #include "Templates/SharedPointer.h"
 #include "Widgets/SCompoundWidget.h"
 
+#define UE_API MLDEFORMERFRAMEWORKEDITOR_API
+
 class IDetailsView;
 
 namespace UE::MLDeformer::TrainingDataProcessor
@@ -18,31 +20,31 @@ namespace UE::MLDeformer::TrainingDataProcessor
 	 * This widget is basically what's inside the tab when this tool opens.
 	 * It contains a detail view and generate button.
 	 */
-	class MLDEFORMERFRAMEWORKEDITOR_API STrainingDataProcessorMainWidget final : public SCompoundWidget, public FEditorUndoClient, public FNotifyHook
+	class STrainingDataProcessorMainWidget final : public SCompoundWidget, public FEditorUndoClient, public FNotifyHook
 	{
 	public:
 		SLATE_BEGIN_ARGS(STrainingDataProcessorMainWidget) { }
 			SLATE_ARGUMENT(TObjectPtr<UMLDeformerModel>, Model)
 		SLATE_END_ARGS()
 
-		virtual ~STrainingDataProcessorMainWidget() override;
-		void Construct(const FArguments& InArgs);
+		UE_API virtual ~STrainingDataProcessorMainWidget() override;
+		UE_API void Construct(const FArguments& InArgs);
 
 		//~ Begin FEditorUndoClient overrides.
-		virtual void PostUndo(bool bSuccess) override;
-		virtual void PostRedo(bool bSuccess) override;
+		UE_API virtual void PostUndo(bool bSuccess) override;
+		UE_API virtual void PostRedo(bool bSuccess) override;
 		//~ End FEditorUndoClient overrides.
 
 		//~ Begin FNotifyHook overrides.
-		virtual void NotifyPostChange(const FPropertyChangedEvent& PropertyChangedEvent, FProperty* PropertyThatChanged) override;
+		UE_API virtual void NotifyPostChange(const FPropertyChangedEvent& PropertyChangedEvent, FProperty* PropertyThatChanged) override;
 		//~ End FNotifyHook overrides.
 
 	private:
-		FReply OnGenerateButtonClicked() const;
-		bool IsValidConfiguration() const;
-		bool IsObjectOfInterest(UObject* Object) const;
-		void OnAssetModified(UObject* Object) const;
-		void Refresh() const;
+		UE_API FReply OnGenerateButtonClicked() const;
+		UE_API bool IsValidConfiguration() const;
+		UE_API bool IsObjectOfInterest(UObject* Object) const;
+		UE_API void OnAssetModified(UObject* Object) const;
+		UE_API void Refresh() const;
 
 	private:
 		/** The details view that shows the properties of our UMLDeformerTrainingDataProcessorSettings. */
@@ -55,3 +57,5 @@ namespace UE::MLDeformer::TrainingDataProcessor
 		FDelegateHandle ObjectModifiedHandle;
 	};
 } // namespace UE::MLDeformer::TrainingDataProcessor
+
+#undef UE_API
