[MLDeformer] Fix a crash when adding a bone to the inputs of the ML Deformer.

diff --git a/Engine/Plugins/Animation/MLDeformer/MLDeformerFramework/Source/MLDeformerFrameworkEditor/Private/SMLDeformerInputBonesWidget.cpp b/Engine/Plugins/Animation/MLDeformer/MLDeformerFramework/Source/MLDeformerFrameworkEditor/Private/SMLDeformerInputBonesWidget.cpp
--- a/Engine/Plugins/Animation/MLDeformer/MLDeformerFramework/Source/MLDeformerFrameworkEditor/Private/SMLDeformerInputBonesWidget.cpp
+++ b/Engine/Plugins/Animation/MLDeformer/MLDeformerFramework/Source/MLDeformerFrameworkEditor/Private/SMLDeformerInputBonesWidget.cpp
@@ -345,6 +345,7 @@ namespace UE::MLDeformer
 					}
 				}
 
+				EditorModel->UpdateEditorInputInfo();
 				RefreshTree();
 
 				// Trigger the input widgets events.
diff --git a/Engine/Plugins/Animation/MLDeformer/MLDeformerFramework/Source/MLDeformerFrameworkEditor/Private/SMLDeformerInputCurvesWidget.cpp b/Engine/Plugins/Animation/MLDeformer/MLDeformerFramework/Source/MLDeformerFrameworkEditor/Private/SMLDeformerInputCurvesWidget.cpp
--- a/Engine/Plugins/Animation/MLDeformer/MLDeformerFramework/Source/MLDeformerFrameworkEditor/Private/SMLDeformerInputCurvesWidget.cpp
+++ b/Engine/Plugins/Animation/MLDeformer/MLDeformerFramework/Source/MLDeformerFrameworkEditor/Private/SMLDeformerInputCurvesWidget.cpp
@@ -275,6 +275,7 @@ namespace UE::MLDeformer
 					}
 				}
 
+				EditorModel->UpdateEditorInputInfo();
 				RefreshList();
 
 				// Trigger the input widgets events.
