[MLDeformer] Enable input pose rendering inside the Nearest Neighbor Model using the MLDeformer.DebugDrawInputPoses console command.

diff --git a/Engine/Plugins/Animation/MLDeformer/MLDeformerFramework/Source/MLDeformerFramework/Private/MLDeformerModelInstance.cpp b/Engine/Plugins/Animation/MLDeformer/MLDeformerFramework/Source/MLDeformerFramework/Private/MLDeformerModelInstance.cpp
--- a/Engine/Plugins/Animation/MLDeformer/MLDeformerFramework/Source/MLDeformerFramework/Private/MLDeformerModelInstance.cpp
+++ b/Engine/Plugins/Animation/MLDeformer/MLDeformerFramework/Source/MLDeformerFramework/Private/MLDeformerModelInstance.cpp
@@ -22,7 +22,7 @@ DEFINE_STAT(STAT_MLDeformerInference);
 
 namespace UE::MLDeformer
 {
-	static TAutoConsoleVariable<bool> CVarMLDeformerDebugDrawInputPoses(
+	TAutoConsoleVariable<bool> CVarMLDeformerDebugDrawInputPoses(
 		TEXT("MLDeformer.DebugDrawInputPoses"),
 		false,
 		TEXT("Draw the input pose for each active ML Deformer.")
diff --git a/Engine/Plugins/Animation/MLDeformer/MLDeformerFramework/Source/MLDeformerFramework/Public/MLDeformerModelInstance.h b/Engine/Plugins/Animation/MLDeformer/MLDeformerFramework/Source/MLDeformerFramework/Public/MLDeformerModelInstance.h
--- a/Engine/Plugins/Animation/MLDeformer/MLDeformerFramework/Source/MLDeformerFramework/Public/MLDeformerModelInstance.h
+++ b/Engine/Plugins/Animation/MLDeformer/MLDeformerFramework/Source/MLDeformerFramework/Public/MLDeformerModelInstance.h
@@ -17,6 +17,12 @@ class UMLDeformerComponent;
 
 DECLARE_STATS_GROUP(TEXT("MLDeformer"), STATGROUP_MLDeformer, STATCAT_Advanced);
 DECLARE_CYCLE_STAT_EXTERN(TEXT("ML Deformer Inference"), STAT_MLDeformerInference, STATGROUP_MLDeformer, );
+
+namespace UE::MLDeformer
+{
+	extern UE_API TAutoConsoleVariable<bool> CVarMLDeformerDebugDrawInputPoses;
+}
+
 /**
  * An instance of the ML Deformer model.
  * The ML Deformer model contains shared data, while this instance contains data unique to the actor it is being applied to.
diff --git a/Engine/Plugins/Animation/MLDeformer/NearestNeighborModel/Source/NearestNeighborModel/Private/NearestNeighborModelInstance.cpp b/Engine/Plugins/Animation/MLDeformer/NearestNeighborModel/Source/NearestNeighborModel/Private/NearestNeighborModelInstance.cpp
--- a/Engine/Plugins/Animation/MLDeformer/NearestNeighborModel/Source/NearestNeighborModel/Private/NearestNeighborModelInstance.cpp
+++ b/Engine/Plugins/Animation/MLDeformer/NearestNeighborModel/Source/NearestNeighborModel/Private/NearestNeighborModelInstance.cpp
@@ -105,6 +105,13 @@ void UNearestNeighborModelInstance::Tick(float DeltaTime, float ModelWeight)
 		// calculate the network outputs and possibly use them, depending on how the model works.
 		Execute(ModelWeight);
 		RunNearestNeighborModel(DeltaTime, ModelWeight);
+
+		const bool bDrawDebugInputPoses = UE::MLDeformer::CVarMLDeformerDebugDrawInputPoses->GetBool();
+		if (bDrawDebugInputPoses)
+		{
+			DrawDebugInputPoses();
+		}
+		
 		bExecuteCalled = true;
 	}
 	else
