[MLDeformer] Hide certain conflicting properties (categories even) from the advanced preview scene settings tab. This was needed because the MLD asset editor has its own mesh picking and animation playback. Changing those properties inside the preview settings tab would mess things up. So we now hide those specific conflicting settings.

diff --git a/Engine/Plugins/Animation/MLDeformer/MLDeformerFramework/Source/MLDeformerFrameworkEditor/Private/MLDeformerApplicationMode.cpp b/Engine/Plugins/Animation/MLDeformer/MLDeformerFramework/Source/MLDeformerFrameworkEditor/Private/MLDeformerApplicationMode.cpp
--- a/Engine/Plugins/Animation/MLDeformer/MLDeformerFramework/Source/MLDeformerFrameworkEditor/Private/MLDeformerApplicationMode.cpp
+++ b/Engine/Plugins/Animation/MLDeformer/MLDeformerFramework/Source/MLDeformerFrameworkEditor/Private/MLDeformerApplicationMode.cpp
@@ -13,7 +13,7 @@
 namespace UE::MLDeformer
 {
 	FName FMLDeformerApplicationMode::ModeName("MLDeformerAssetEditMode");
-
+	
 	FMLDeformerApplicationMode::FMLDeformerApplicationMode(
 		TSharedRef<FWorkflowCentricApplication> InHostingApp,
 		TSharedRef<IPersonaPreviewScene> InPreviewScene)
@@ -32,10 +32,18 @@ namespace UE::MLDeformer
 		ViewportArgs.ContextName = TEXT("MLDeformerEditor.Viewport");
 		ViewportArgs.OnViewportCreated = FOnViewportCreated::CreateSP(MLDeformerEditor, &FMLDeformerEditorToolkit::HandleViewportCreated);
 
+		// Hide certain property categories inside the advanced preview scene settings tab.
+		FPersonaAdvancedPreviewSceneTabSettings AdvancedPreviewSceneTabSettings;
+		AdvancedPreviewSceneTabSettings.bHideAdditionalMeshesCategory = true;
+		AdvancedPreviewSceneTabSettings.bHideAnimationCategory = true;
+		AdvancedPreviewSceneTabSettings.bHideMeshCategory = true;
+		AdvancedPreviewSceneTabSettings.bHidePhysicsCategory = true;
+		AdvancedPreviewSceneTabSettings.bHideAnimationBlueprintCategory = true;
+
 		// Register Persona tabs.
 		FPersonaModule& PersonaModule = FModuleManager::LoadModuleChecked<FPersonaModule>("Persona");
 		TabFactories.RegisterFactory(PersonaModule.CreatePersonaViewportTabFactory(InHostingApp, ViewportArgs));
-		TabFactories.RegisterFactory(PersonaModule.CreateAdvancedPreviewSceneTabFactory(InHostingApp, InPreviewScene));
+		TabFactories.RegisterFactory(PersonaModule.CreateAdvancedPreviewSceneTabFactory(InHostingApp, InPreviewScene, AdvancedPreviewSceneTabSettings));
 		TabFactories.RegisterFactory(PersonaModule.CreatePersonaToolboxTabFactory(MLDeformerEditor));
 		TabFactories.RegisterFactory(PersonaModule.CreateDetailsTabFactory(InHostingApp, FOnDetailsCreated::CreateSP(&MLDeformerEditor.Get(), &FMLDeformerEditorToolkit::HandleDetailsCreated)));
 
diff --git a/Engine/Source/Editor/Persona/Private/PersonaModule.cpp b/Engine/Source/Editor/Persona/Private/PersonaModule.cpp
--- a/Engine/Source/Editor/Persona/Private/PersonaModule.cpp
+++ b/Engine/Source/Editor/Persona/Private/PersonaModule.cpp
@@ -309,6 +309,11 @@ TSharedRef<class FWorkflowTabFactory> FPersonaModule::CreateAdvancedPreviewScene
 	return MakeShareable(new FAdvancedPreviewSceneTabSummoner(InHostingApp, InPreviewScene));
 }
 
+TSharedRef<class FWorkflowTabFactory> FPersonaModule::CreateAdvancedPreviewSceneTabFactory(const TSharedRef<class FWorkflowCentricApplication>& InHostingApp, const TSharedRef<IPersonaPreviewScene>& InPreviewScene, const FPersonaAdvancedPreviewSceneTabSettings& DetailCustomizationSettings) const
+{
+	return MakeShareable(new FAdvancedPreviewSceneTabSummoner(InHostingApp, InPreviewScene, DetailCustomizationSettings));
+}
+
 TSharedRef<class FWorkflowTabFactory> FPersonaModule::CreateAnimationAssetBrowserTabFactory(const TSharedRef<class FWorkflowCentricApplication>& InHostingApp, const TSharedRef<IPersonaToolkit>& InPersonaToolkit, FOnOpenNewAsset InOnOpenNewAsset, FOnAnimationSequenceBrowserCreated InOnAnimationSequenceBrowserCreated, bool bInShowHistory) const
 {
 	return MakeShareable(new FAnimationAssetBrowserSummoner(InHostingApp, InPersonaToolkit, InOnOpenNewAsset, InOnAnimationSequenceBrowserCreated, bInShowHistory));
diff --git a/Engine/Source/Editor/Persona/Private/PreviewSceneCustomizations.cpp b/Engine/Source/Editor/Persona/Private/PreviewSceneCustomizations.cpp
--- a/Engine/Source/Editor/Persona/Private/PreviewSceneCustomizations.cpp
+++ b/Engine/Source/Editor/Persona/Private/PreviewSceneCustomizations.cpp
@@ -62,6 +62,12 @@ FPreviewSceneDescriptionCustomization::FPreviewSceneDescriptionCustomization(con
 	}
 }
 
+FPreviewSceneDescriptionCustomization::FPreviewSceneDescriptionCustomization(const FString& InSkeletonName, const TSharedRef<class IPersonaToolkit>& InPersonaToolkit, const FPersonaAdvancedPreviewSceneTabSettings& InDetailCustomizationSettings)
+	: FPreviewSceneDescriptionCustomization(InSkeletonName, InPersonaToolkit)
+{
+	DetailCustomizationSettings = InDetailCustomizationSettings;
+}
+
 FPreviewSceneDescriptionCustomization::~FPreviewSceneDescriptionCustomization()
 {
 	if (FactoryToUse)
@@ -499,6 +505,31 @@ void FPreviewSceneDescriptionCustomization::CustomizeDetails(IDetailLayoutBuilde
 			}
 		}
 	}
+
+	if (DetailCustomizationSettings.bHideMeshCategory)
+	{
+		DetailBuilder.HideCategory("Mesh");
+	}
+
+	if (DetailCustomizationSettings.bHideAdditionalMeshesCategory)
+	{
+		DetailBuilder.HideCategory("Additional Meshes");
+	}
+
+	if (DetailCustomizationSettings.bHideAnimationCategory)
+	{
+		DetailBuilder.HideCategory("Animation");
+	}
+
+	if (DetailCustomizationSettings.bHidePhysicsCategory)
+	{
+		DetailBuilder.HideCategory("Physics");
+	}
+
+	if (DetailCustomizationSettings.bHideAnimationBlueprintCategory)
+	{
+		DetailBuilder.HideCategory("Animation Blueprint");
+	}
 }
 
 EVisibility FPreviewSceneDescriptionCustomization::GetSaveButtonVisibility(TSharedRef<IPropertyHandle> InAdditionalMeshesProperty) const
diff --git a/Engine/Source/Editor/Persona/Private/PreviewSceneCustomizations.h b/Engine/Source/Editor/Persona/Private/PreviewSceneCustomizations.h
--- a/Engine/Source/Editor/Persona/Private/PreviewSceneCustomizations.h
+++ b/Engine/Source/Editor/Persona/Private/PreviewSceneCustomizations.h
@@ -10,6 +10,7 @@
 #include "IEditableSkeleton.h"
 #include "IPersonaToolkit.h"
 #include "AnimationEditorPreviewScene.h"
+#include "PersonaModule.h"
 
 #ifndef CHAOS_SIMULATION_DETAIL_VIEW_FACTORY_SELECTOR  // TODO: Decide whether to keep the detail view cloth selector after nvcloth has been removed
 	#define CHAOS_SIMULATION_DETAIL_VIEW_FACTORY_SELECTOR 1
@@ -42,6 +43,7 @@ class FPreviewSceneDescriptionCustomization : public IDetailCustomization
 {
 public:
 	FPreviewSceneDescriptionCustomization(const FString& InSkeletonName, const TSharedRef<class IPersonaToolkit>& InPersonaToolkit);
+	FPreviewSceneDescriptionCustomization(const FString& InSkeletonName, const TSharedRef<class IPersonaToolkit>& InPersonaToolkit, const FPersonaAdvancedPreviewSceneTabSettings& InDetailCustomizationSettings);
 
 	~FPreviewSceneDescriptionCustomization();
 
@@ -131,6 +133,8 @@ private:
 
 	TWeakObjectPtr<UDataAsset>	DataAssetToDisplay;
 
+	FPersonaAdvancedPreviewSceneTabSettings DetailCustomizationSettings;
+
 	/**
 	* Called to get the visibility of the replace button
 	*/
diff --git a/Engine/Source/Editor/Persona/Private/TabSpawners.cpp b/Engine/Source/Editor/Persona/Private/TabSpawners.cpp
--- a/Engine/Source/Editor/Persona/Private/TabSpawners.cpp
+++ b/Engine/Source/Editor/Persona/Private/TabSpawners.cpp
@@ -685,6 +685,12 @@ FAdvancedPreviewSceneTabSummoner::FAdvancedPreviewSceneTabSummoner(TSharedPtr<cl
 	ViewMenuTooltip = LOCTEXT("AdvancedPreviewScene_ToolTip", "Shows the advanced preview scene settings");
 }
 
+FAdvancedPreviewSceneTabSummoner::FAdvancedPreviewSceneTabSummoner(TSharedPtr<class FAssetEditorToolkit> InHostingApp, const TSharedRef<IPersonaPreviewScene>& InPreviewScene, const FPersonaAdvancedPreviewSceneTabSettings& InDetailCustomizationSettings)
+	: FAdvancedPreviewSceneTabSummoner(InHostingApp, InPreviewScene)
+{
+	DetailCustomizationSettings = InDetailCustomizationSettings;
+}
+
 TSharedRef<class IDetailCustomization> FAdvancedPreviewSceneTabSummoner::CustomizePreviewSceneDescription()
 {
 	TSharedRef<IPersonaPreviewScene> PreviewSceneRef = PreviewScene.Pin().ToSharedRef();
@@ -694,7 +700,7 @@ TSharedRef<class IDetailCustomization> FAdvancedPreviewSceneTabSummoner::Customi
 	{
 		SkeletonName = FAssetData(&EditableSkeleton->GetSkeleton()).GetExportTextName();
 	}
-	return MakeShareable(new FPreviewSceneDescriptionCustomization(SkeletonName, PreviewSceneRef->GetPersonaToolkit()));
+	return MakeShareable(new FPreviewSceneDescriptionCustomization(SkeletonName, PreviewSceneRef->GetPersonaToolkit(), DetailCustomizationSettings));
 }
 
 TSharedRef<class IPropertyTypeCustomization> FAdvancedPreviewSceneTabSummoner::CustomizePreviewMeshCollectionEntry()
diff --git a/Engine/Source/Editor/Persona/Private/TabSpawners.h b/Engine/Source/Editor/Persona/Private/TabSpawners.h
--- a/Engine/Source/Editor/Persona/Private/TabSpawners.h
+++ b/Engine/Source/Editor/Persona/Private/TabSpawners.h
@@ -363,6 +363,8 @@ struct FAdvancedPreviewSceneTabSummoner : public FWorkflowTabFactory
 {
 public:
  	FAdvancedPreviewSceneTabSummoner(TSharedPtr<class FAssetEditorToolkit> InHostingApp, const TSharedRef<class IPersonaPreviewScene>& InPreviewScene);
+	FAdvancedPreviewSceneTabSummoner(TSharedPtr<class FAssetEditorToolkit> InHostingApp, const TSharedRef<class IPersonaPreviewScene>& InPreviewScene, const FPersonaAdvancedPreviewSceneTabSettings& InDetailCustomizationSettings);
+	
 	virtual TSharedRef<SWidget> CreateTabBody(const FWorkflowTabSpawnInfo& Info) const override;
 	virtual FText GetTabToolTipText(const FWorkflowTabSpawnInfo& Info) const override;
 
@@ -375,6 +377,7 @@ private:
 
 private:
 	TWeakPtr<class IPersonaPreviewScene> PreviewScene;
+	FPersonaAdvancedPreviewSceneTabSettings DetailCustomizationSettings;
 };
 
 /////////////////////////////////////////////////////
diff --git a/Engine/Source/Editor/Persona/Public/PersonaModule.h b/Engine/Source/Editor/Persona/Public/PersonaModule.h
--- a/Engine/Source/Editor/Persona/Public/PersonaModule.h
+++ b/Engine/Source/Editor/Persona/Public/PersonaModule.h
@@ -314,6 +314,19 @@ struct FPersonaViewportArgs
 	FAnimationScrubPanelDelegates TimelineDelegates;
 };
 
+/**
+ * Some settings for the detail customization inside the advanced preview scene settings tab.
+ * This allows you to hide specific categories.
+ */
+struct FPersonaAdvancedPreviewSceneTabSettings
+{
+	bool bHideAnimationCategory = false;
+	bool bHideAnimationBlueprintCategory = false;
+	bool bHideMeshCategory = false;
+	bool bHidePhysicsCategory = false;
+	bool bHideAdditionalMeshesCategory = false;
+};
+
 /**
  * Persona module manages the lifetime of all instances of Persona editors.
  */
@@ -377,6 +390,7 @@ public:
 
 	/** Create a tab factory used to configure preview scene settings */
 	virtual TSharedRef<FWorkflowTabFactory> CreateAdvancedPreviewSceneTabFactory(const TSharedRef<class FWorkflowCentricApplication>& InHostingApp, const TSharedRef<IPersonaPreviewScene>& InPreviewScene) const;
+	virtual TSharedRef<FWorkflowTabFactory> CreateAdvancedPreviewSceneTabFactory(const TSharedRef<class FWorkflowCentricApplication>& InHostingApp, const TSharedRef<IPersonaPreviewScene>& InPreviewScene, const FPersonaAdvancedPreviewSceneTabSettings& DetailCustomizationSettings) const;
 
 	/** Create a tab factory for the animation asset browser */
 	virtual TSharedRef<FWorkflowTabFactory> CreateAnimationAssetBrowserTabFactory(const TSharedRef<class FWorkflowCentricApplication>& InHostingApp, const TSharedRef<IPersonaToolkit>& InPersonaToolkit, FOnOpenNewAsset InOnOpenNewAsset, FOnAnimationSequenceBrowserCreated InOnAnimationSequenceBrowserCreated, bool bInShowHistory) const;
