[MLDeformer] Fix a problem where training with curves would cause a failure during training. Also fixed a bug where curves on the skeletal mesh instead of on the skeleton would cause some error checking to fail.

diff --git a/Engine/Plugins/Animation/MLDeformer/MLDeformerFramework/Source/MLDeformerFramework/Private/MLDeformerInputInfo.cpp b/Engine/Plugins/Animation/MLDeformer/MLDeformerFramework/Source/MLDeformerFramework/Private/MLDeformerInputInfo.cpp
--- a/Engine/Plugins/Animation/MLDeformer/MLDeformerFramework/Source/MLDeformerFramework/Private/MLDeformerInputInfo.cpp
+++ b/Engine/Plugins/Animation/MLDeformer/MLDeformerFramework/Source/MLDeformerFramework/Private/MLDeformerInputInfo.cpp
@@ -145,25 +145,26 @@ void UMLDeformerInputInfo::SetNumTargetVertices(int32 NumVerts)
 	NumTargetMeshVertices = NumVerts;
 }
 
-bool UMLDeformerInputInfo::IsCompatible(USkeletalMesh* InSkeletalMesh) const
+bool UMLDeformerInputInfo::IsCompatible(const USkeletalMeshComponent* InSkeletalMeshComponent) const
 {
+	const USkeletalMesh* ComponentSkeletalMesh = Cast<USkeletalMesh>(InSkeletalMeshComponent->GetSkinnedAsset());
 	if (SkeletalMesh.IsNull())
 	{
 		return true;
 	}
 
-	if (InSkeletalMesh == nullptr)
+	if (ComponentSkeletalMesh == nullptr)
 	{
 		return false;
 	}
 
-	if ((SkeletalMesh.IsValid() && !SkeletalMesh.IsNull()) && FSoftObjectPath(InSkeletalMesh) != SkeletalMesh)
+	if ((SkeletalMesh.IsValid() && !SkeletalMesh.IsNull()) && FSoftObjectPath(ComponentSkeletalMesh) != SkeletalMesh)
 	{
 		return false;
 	}
 
 	// Verify that all required bones are there.
-	const FReferenceSkeleton& RefSkeleton = InSkeletalMesh->GetRefSkeleton();
+	const FReferenceSkeleton& RefSkeleton = ComponentSkeletalMesh->GetRefSkeleton();
 	for (const FName BoneName : BoneNames)
 	{
 		const int32 BoneIndex = RefSkeleton.FindBoneIndex(BoneName);
@@ -174,10 +175,10 @@ bool UMLDeformerInputInfo::IsCompatible(USkeletalMesh* InSkeletalMesh) const
 	}
 
 	// Verify that all required curves are there.
-	USkeleton* Skeleton = InSkeletalMesh->GetSkeleton();
+	const FBlendedHeapCurve& Curves = InSkeletalMeshComponent->GetAnimCurves();
 	for (const FName CurveName : CurveNames)
 	{
-		if (Skeleton->GetCurveMetaData(CurveName) == nullptr)
+		if (!Curves.HasElement(CurveName))
 		{
 			return false;
 		}
@@ -186,6 +187,78 @@ bool UMLDeformerInputInfo::IsCompatible(USkeletalMesh* InSkeletalMesh) const
 	return true;
 }
 
+bool UMLDeformerInputInfo::IsCompatible(USkeletalMesh* InSkeletalMesh) const
+{
+	if (SkeletalMesh.IsNull())
+	{
+		return true;
+	}
+
+	if (InSkeletalMesh == nullptr)
+	{
+		return false;
+	}
+
+	if ((SkeletalMesh.IsValid() && !SkeletalMesh.IsNull()) && FSoftObjectPath(InSkeletalMesh) != SkeletalMesh)
+	{
+		return false;
+	}
+
+	return true;
+}
+
+FString UMLDeformerInputInfo::GenerateCompatibilityErrorString(const USkeletalMeshComponent* InSkeletalMeshComponent) const
+{
+	const USkeletalMesh* ComponentSkeletalMesh = Cast<USkeletalMesh>(InSkeletalMeshComponent->GetSkinnedAsset());
+	if (ComponentSkeletalMesh == nullptr)
+	{
+		return FString();
+	}
+
+	// Verify that all required bones are there.
+	const FReferenceSkeleton& RefSkeleton = ComponentSkeletalMesh->GetRefSkeleton();
+	FString ErrorString;
+	for (const FName BoneName : BoneNames)
+	{
+		const int32 BoneIndex = RefSkeleton.FindBoneIndex(BoneName);
+		if (BoneIndex == INDEX_NONE) // We're missing the required bone. The network needs to input the transform for this bone.
+		{
+			ErrorString += FString::Format(TEXT("Required bone '{0}' is missing.\n"), {*BoneName.ToString()});
+		}
+	}
+
+	// Verify that all required curves are there.
+	const FBlendedHeapCurve& Curves = InSkeletalMeshComponent->GetAnimCurves();
+	for (const FName CurveName : CurveNames)
+	{
+		if (!Curves.HasElement(CurveName))
+		{
+			ErrorString += FString::Format(TEXT("Required curve '{0}' is missing.\n"), {*CurveName.ToString()});
+		}
+	}
+	
+	// Check vertex count.
+#if WITH_EDITORONLY_DATA
+	if ((NumBaseMeshVertices > 0 && NumTargetMeshVertices > 0) &&
+		NumBaseMeshVertices != ComponentSkeletalMesh->GetNumImportedVertices())
+	{
+		ErrorString += FString::Format(TEXT("The number of vertices that the network was trained on ({0} verts) doesn't match the skeletal mesh '{1}' ({2} verts)..\n"), 
+			{
+				NumBaseMeshVertices, 
+				ComponentSkeletalMesh->GetName(),
+				ComponentSkeletalMesh->GetNumImportedVertices(),
+			} );
+	}
+#endif
+
+	if (SkeletalMesh.IsValid() && FSoftObjectPath(ComponentSkeletalMesh) != SkeletalMesh)
+	{
+		ErrorString += FString::Format(TEXT("Model was trained on Skeletal Mesh '{0}', which is not the same as '{1}' it tries to apply on.\n"), {*SkeletalMesh.ToString(), *FSoftObjectPath(ComponentSkeletalMesh).ToString()});
+	}
+
+	return ErrorString;	
+}
+
 FString UMLDeformerInputInfo::GenerateCompatibilityErrorString(USkeletalMesh* InSkeletalMesh) const
 {
 	if (InSkeletalMesh == nullptr)
@@ -214,7 +287,7 @@ FString UMLDeformerInputInfo::GenerateCompatibilityErrorString(USkeletalMesh* In
 			ErrorString += FString::Format(TEXT("Required curve '{0}' is missing.\n"), {*CurveName.ToString()});
 		}
 	}
-
+	
 	// Check vertex count.
 #if WITH_EDITORONLY_DATA
 	if ((NumBaseMeshVertices > 0 && NumTargetMeshVertices > 0) &&
@@ -234,7 +307,7 @@ FString UMLDeformerInputInfo::GenerateCompatibilityErrorString(USkeletalMesh* In
 		ErrorString += FString::Format(TEXT("Model was trained on Skeletal Mesh '{0}', which is not the same as '{1}' it tries to apply on.\n"), {*SkeletalMesh.ToString(), *FSoftObjectPath(InSkeletalMesh).ToString()});
 	}
 
-	return ErrorString;
+	return ErrorString;	
 }
 
 void UMLDeformerInputInfo::ExtractCurveValues(USkeletalMeshComponent* SkelMeshComponent, TArray<float>& OutValues) const
diff --git a/Engine/Plugins/Animation/MLDeformer/MLDeformerFramework/Source/MLDeformerFramework/Private/MLDeformerModelInstance.cpp b/Engine/Plugins/Animation/MLDeformer/MLDeformerFramework/Source/MLDeformerFramework/Private/MLDeformerModelInstance.cpp
--- a/Engine/Plugins/Animation/MLDeformer/MLDeformerFramework/Source/MLDeformerFramework/Private/MLDeformerModelInstance.cpp
+++ b/Engine/Plugins/Animation/MLDeformer/MLDeformerFramework/Source/MLDeformerFramework/Private/MLDeformerModelInstance.cpp
@@ -281,9 +281,9 @@ FString UMLDeformerModelInstance::CheckCompatibility(USkeletalMeshComponent* InS
 	// If we're not compatible, generate a compatibility string.
 	USkeletalMesh* SkelMesh = InSkelMeshComponent ? InSkelMeshComponent->GetSkeletalMeshAsset() : nullptr;
 	UMLDeformerInputInfo* InputInfo = Model->GetInputInfo();
-	if (SkelMesh && !InputInfo->IsCompatible(SkelMesh) && Model->GetDeformerAsset() && Model->IsTrained())
+	if (SkelMesh && !InputInfo->IsCompatible(InSkelMeshComponent) && Model->GetDeformerAsset() && Model->IsTrained())
 	{
-		ErrorText += InputInfo->GenerateCompatibilityErrorString(SkelMesh);
+		ErrorText += InputInfo->GenerateCompatibilityErrorString(InSkelMeshComponent);
 		if (bLogIssues)
 		{
 			UE_LOG(LogMLDeformer, Error, TEXT("ML Deformer '%s' isn't compatible with Skeletal Mesh '%s'.\nReason(s):\n%s\n"), 
diff --git a/Engine/Plugins/Animation/MLDeformer/MLDeformerFramework/Source/MLDeformerFramework/Public/MLDeformerInputInfo.h b/Engine/Plugins/Animation/MLDeformer/MLDeformerFramework/Source/MLDeformerFramework/Public/MLDeformerInputInfo.h
--- a/Engine/Plugins/Animation/MLDeformer/MLDeformerFramework/Source/MLDeformerFramework/Public/MLDeformerInputInfo.h
+++ b/Engine/Plugins/Animation/MLDeformer/MLDeformerFramework/Source/MLDeformerFramework/Public/MLDeformerInputInfo.h
@@ -39,12 +39,21 @@ public:
 
 	/**
 	 * Check whether the current inputs are compatible with a given skeletal mesh.
+	 * For a more extensive check, please use IsCompatible that takes a USkeletalMeshComponent as parameter.
 	 * @param InSkeletalMesh The skeletal mesh to check compatibility with. This may not be a nullptr.
 	 * @return Returns true when we can safely apply the ML Deformer to a character using this skeletal mesh, otherwise false is returned.
 	 * @note Use GenerateCompatibilityErrorString to get the error report.
 	 */
 	UE_API virtual bool IsCompatible(USkeletalMesh* InSkeletalMesh) const;
 
+	/**
+	 * Check whether the current inputs are compatible with a given skeletal mesh component.
+	 * @param InSkeletalMeshComponent The skeletal mesh to check compatibility with. This may not be a nullptr.
+	 * @return Returns true when we can safely apply the ML Deformer to a character using this skeletal mesh, otherwise false is returned.
+	 * @note Use GenerateCompatibilityErrorString to get the error report.
+	 */
+	UE_API virtual bool IsCompatible(const USkeletalMeshComponent* InSkeletalMeshComponent) const;
+
 	/**
 	 * Clear all contents.
 	 * This can clear the list of all bone and curves that are part of this input info.
@@ -57,7 +66,16 @@ public:
 	 * @return Returns an empty string in case there are no compatibility issues, otherwise it contains a string that describes the issue(s).
 	 *         In case a nullptr is passed as SkeletalMesh parameter, an empty string is returned.
 	 */
+	UE_DEPRECATED(5.7, "Please use the GenerateCompatibilityErrorString that takes a USkeletalMeshComponent as parameter instead.")
 	UE_API virtual FString GenerateCompatibilityErrorString(USkeletalMesh* InSkeletalMesh) const;
+	
+	/**
+	 * Get the compatibility error report.
+	 * @param InSkeletalMeshComponent The skeletal mesh component to check compatibility with.
+	 * @return Returns an empty string in case there are no compatibility issues, otherwise it contains a string that describes the issue(s).
+	 *         In case a nullptr is passed as SkeletalMeshComponent parameter, an empty string is returned.
+	 */
+	UE_API virtual FString GenerateCompatibilityErrorString(const USkeletalMeshComponent* InSkeletalMeshComponent) const;
 
 	/**
 	 * Check whether we have any training inputs or not.
diff --git a/Engine/Plugins/Animation/MLDeformer/MLDeformerFramework/Source/MLDeformerFrameworkEditor/Private/MLDeformerGeomCacheTrainingModel.cpp b/Engine/Plugins/Animation/MLDeformer/MLDeformerFramework/Source/MLDeformerFrameworkEditor/Private/MLDeformerGeomCacheTrainingModel.cpp
--- a/Engine/Plugins/Animation/MLDeformer/MLDeformerFramework/Source/MLDeformerFrameworkEditor/Private/MLDeformerGeomCacheTrainingModel.cpp
+++ b/Engine/Plugins/Animation/MLDeformer/MLDeformerFramework/Source/MLDeformerFrameworkEditor/Private/MLDeformerGeomCacheTrainingModel.cpp
@@ -149,6 +149,9 @@ bool UMLDeformerGeomCacheTrainingModel::GenerateBasicInputsAndOutputBuffers(cons
 		return false;
 	}
 
+	// Update the number of floats per curve.
+	SetNumFloatsPerCurve(GetModel()->GetNumFloatsPerCurve());
+	
 	// Sample all the frames.
 	const double StartTime = FPlatformTime::Seconds();
 	for (int32 SampleIndex = 0; SampleIndex < NumFrames; SampleIndex++)
diff --git a/Engine/Plugins/Animation/MLDeformer/NearestNeighborModel/Source/NearestNeighborModel/Private/NearestNeighborModelInstance.cpp b/Engine/Plugins/Animation/MLDeformer/NearestNeighborModel/Source/NearestNeighborModel/Private/NearestNeighborModelInstance.cpp
--- a/Engine/Plugins/Animation/MLDeformer/NearestNeighborModel/Source/NearestNeighborModel/Private/NearestNeighborModelInstance.cpp
+++ b/Engine/Plugins/Animation/MLDeformer/NearestNeighborModel/Source/NearestNeighborModel/Private/NearestNeighborModelInstance.cpp
@@ -166,9 +166,9 @@ FString UNearestNeighborModelInstance::CheckCompatibility(USkeletalMeshComponent
 	// If we're not compatible, generate a compatibility string.
 	USkeletalMesh* const SkelMesh = InSkelMeshComponent ? InSkelMeshComponent->GetSkeletalMeshAsset() : nullptr;
 	UMLDeformerInputInfo* const InputInfo = Model->GetInputInfo();
-	if (SkelMesh && !InputInfo->IsCompatible(SkelMesh) && Model->GetDeformerAsset())
+	if (SkelMesh && !InputInfo->IsCompatible(InSkelMeshComponent) && Model->GetDeformerAsset())
 	{
-		ErrorText += InputInfo->GenerateCompatibilityErrorString(SkelMesh);
+		ErrorText += InputInfo->GenerateCompatibilityErrorString(InSkelMeshComponent);
 		ErrorText += "\n";
 		if (LogIssues)
 		{
diff --git a/Engine/Plugins/Animation/MLDeformer/NeuralMorphModel/Source/NeuralMorphModel/Private/NeuralMorphModel.cpp b/Engine/Plugins/Animation/MLDeformer/NeuralMorphModel/Source/NeuralMorphModel/Private/NeuralMorphModel.cpp
--- a/Engine/Plugins/Animation/MLDeformer/NeuralMorphModel/Source/NeuralMorphModel/Private/NeuralMorphModel.cpp
+++ b/Engine/Plugins/Animation/MLDeformer/NeuralMorphModel/Source/NeuralMorphModel/Private/NeuralMorphModel.cpp
@@ -280,10 +280,6 @@ void UNeuralMorphModel::SetNeuralMorphNetwork(UNeuralMorphNetwork* Net)
 
 int32 UNeuralMorphModel::GetNumFloatsPerCurve() const
 {
-	if (NeuralMorphNetwork)
-	{
-		return NeuralMorphNetwork->GetNumFloatsPerCurve();
-	}
 	return (Mode == ENeuralMorphMode::Local) ? 6 : 1;
 }
 
diff --git a/Engine/Plugins/Animation/MLDeformer/NeuralMorphModel/Source/NeuralMorphModel/Private/NeuralMorphModelInstance.cpp b/Engine/Plugins/Animation/MLDeformer/NeuralMorphModel/Source/NeuralMorphModel/Private/NeuralMorphModelInstance.cpp
--- a/Engine/Plugins/Animation/MLDeformer/NeuralMorphModel/Source/NeuralMorphModel/Private/NeuralMorphModelInstance.cpp
+++ b/Engine/Plugins/Animation/MLDeformer/NeuralMorphModel/Source/NeuralMorphModel/Private/NeuralMorphModelInstance.cpp
@@ -90,7 +90,7 @@ void UNeuralMorphModelInstance::FillNetworkInputs()
 		// Write the curve values for all curve groups.
 		const int32 NumFloatsPerCurve = MorphNetwork->GetNumFloatsPerCurve();
 		const int32 NumCurveGroups = CurveGroupIndices.Num();
-		for (int32 Index = 0; Index < CurveGroupIndices.Num(); ++Index)
+		for (int32 Index = 0; Index < NumCurveGroups; ++Index)
 		{
 			const int32 CurveIndex = CurveGroupIndices[Index];
 			if (CurveIndex != INDEX_NONE)
