MLDeformer: Integrate new deferred python initialization and deferred pip installer interfaces

diff --git a/Engine/Plugins/Animation/MLDeformer/MLDeformerFramework/Source/MLDeformerFrameworkEditor/MLDeformerFrameworkEditor.Build.cs b/Engine/Plugins/Animation/MLDeformer/MLDeformerFramework/Source/MLDeformerFrameworkEditor/MLDeformerFrameworkEditor.Build.cs
--- a/Engine/Plugins/Animation/MLDeformer/MLDeformerFramework/Source/MLDeformerFrameworkEditor/MLDeformerFrameworkEditor.Build.cs
+++ b/Engine/Plugins/Animation/MLDeformer/MLDeformerFramework/Source/MLDeformerFrameworkEditor/MLDeformerFrameworkEditor.Build.cs
@@ -56,6 +56,7 @@ namespace UnrealBuildTool.Rules
 					"AssetDefinition",
 					"SkeletalMeshDescription",
 					"ApplicationCore",
+					"PythonScriptPlugin",
 					"EditorInteractiveToolsFramework",
 					"InteractiveToolsFramework",
 					"MeshModelingTools",
diff --git a/Engine/Plugins/Animation/MLDeformer/MLDeformerFramework/Source/MLDeformerFrameworkEditor/Private/MLDeformerAssetDefinitions.cpp b/Engine/Plugins/Animation/MLDeformer/MLDeformerFramework/Source/MLDeformerFrameworkEditor/Private/MLDeformerAssetDefinitions.cpp
--- a/Engine/Plugins/Animation/MLDeformer/MLDeformerFramework/Source/MLDeformerFrameworkEditor/Private/MLDeformerAssetDefinitions.cpp
+++ b/Engine/Plugins/Animation/MLDeformer/MLDeformerFramework/Source/MLDeformerFrameworkEditor/Private/MLDeformerAssetDefinitions.cpp
@@ -1,9 +1,12 @@
 // Copyright Epic Games, Inc. All Rights Reserved.
 
 #include "MLDeformerAssetDefinitions.h"
+
 #include "MLDeformerAsset.h"
 #include "MLDeformerEditorToolkit.h"
 #include "MLDeformerModule.h"
+#include "IPythonScriptPlugin.h"
+#include "PipInstallHelpers.h"
 
 #include UE_INLINE_GENERATED_CPP_BY_NAME(MLDeformerAssetDefinitions)
 
@@ -33,14 +36,23 @@ TConstArrayView<FAssetCategoryPath> UAssetDefinition_MLDeformer::GetAssetCategor
 EAssetCommandResult UAssetDefinition_MLDeformer::OpenAssets(const FAssetOpenArgs& OpenArgs) const
 {
 	using namespace UE::MLDeformer;
-	for (UMLDeformerAsset* Asset : OpenArgs.LoadObjects<UMLDeformerAsset>())
+	FPythonScriptInitHelper::InitPythonAndPipInstall(FSimpleDelegate::CreateLambda([
+		OpenArgs = FAssetOpenArgs(OpenArgs),
+		// Taking a copy of OpenArgs.Assets as it's a TArrayView and this callback outlives the view
+		OpenAssetsData = TArray<FAssetData>(OpenArgs.Assets)]()
 	{
-		const double StartTime = FPlatformTime::Seconds();
-		TSharedRef<FMLDeformerEditorToolkit> NewEditor(new FMLDeformerEditorToolkit());
-		NewEditor->InitAssetEditor(OpenArgs.GetToolkitMode(), OpenArgs.ToolkitHost, Asset);
-		UE_LOG(LogMLDeformer, Display, TEXT("ML Deformer asset editor took %.1f seconds to open."), FPlatformTime::Seconds() - StartTime); 
-	}
-
+		FAssetArgs WrappedArgs(OpenAssetsData);
+		for (UMLDeformerAsset* Asset : WrappedArgs.LoadObjects<UMLDeformerAsset>())
+		{
+			TSharedRef<FMLDeformerEditorToolkit> NewEditor(new FMLDeformerEditorToolkit());
+			NewEditor->InitAssetEditor(OpenArgs.GetToolkitMode(), OpenArgs.ToolkitHost, Asset);
+		}
+	}),
+	FSimpleDelegate::CreateLambda([]()
+	{
+		// TODO: Could alternatively open editor here and just log that python isn't available
+		UE_LOG(LogMLDeformer, Warning, TEXT("MLDeformer toolkit may not function properly without Python enabled"));
+	}));
 	return EAssetCommandResult::Handled;
 }
 
