MLDeformer: Some fixes where it would not detect curves.

diff --git a/Engine/Plugins/Animation/MLDeformer/MLDeformerFramework/Source/MLDeformerFramework/Private/MLDeformerInputInfo.cpp b/Engine/Plugins/Animation/MLDeformer/MLDeformerFramework/Source/MLDeformerFramework/Private/MLDeformerInputInfo.cpp
--- a/Engine/Plugins/Animation/MLDeformer/MLDeformerFramework/Source/MLDeformerFramework/Private/MLDeformerInputInfo.cpp
+++ b/Engine/Plugins/Animation/MLDeformer/MLDeformerFramework/Source/MLDeformerFramework/Private/MLDeformerInputInfo.cpp
@@ -147,7 +147,7 @@ void UMLDeformerInputInfo::SetNumTargetVertices(int32 NumVerts)
 
 bool UMLDeformerInputInfo::IsCompatible(const USkeletalMeshComponent* InSkeletalMeshComponent) const
 {
-	const USkeletalMesh* ComponentSkeletalMesh = Cast<USkeletalMesh>(InSkeletalMeshComponent->GetSkinnedAsset());
+	USkeletalMesh* ComponentSkeletalMesh = Cast<USkeletalMesh>(InSkeletalMeshComponent->GetSkinnedAsset());
 	if (SkeletalMesh.IsNull())
 	{
 		return true;
@@ -175,10 +175,11 @@ bool UMLDeformerInputInfo::IsCompatible(const USkeletalMeshComponent* InSkeletal
 	}
 
 	// Verify that all required curves are there.
-	const FBlendedHeapCurve& Curves = InSkeletalMeshComponent->GetAnimCurves();
+	TArray<FName> CurveList;
+	BuildCurveList(ComponentSkeletalMesh, CurveList);
 	for (const FName CurveName : CurveNames)
 	{
-		if (!Curves.HasElement(CurveName))
+		if (!CurveList.Contains(CurveName))
 		{
 			return false;
 		}
@@ -207,9 +208,36 @@ bool UMLDeformerInputInfo::IsCompatible(USkeletalMesh* InSkeletalMesh) const
 	return true;
 }
 
+void UMLDeformerInputInfo::BuildCurveList(USkeletalMesh* InSkeletalMesh, TArray<FName>& OutCurveNames)
+{
+	TSet<FName> UniqueNames;
+	if (InSkeletalMesh)
+	{
+		// Collect all curves on the skeleton.
+		const USkeleton* Skeleton = InSkeletalMesh->GetSkeleton();
+		if (Skeleton)
+		{
+			TArray<FName> SkeletonCurveNames;
+			Skeleton->GetCurveMetaDataNames(SkeletonCurveNames);
+			UniqueNames.Append(SkeletonCurveNames);
+		}
+
+		// Collect all curves on the skeletal mesh.
+		TArray<FName> MeshCurveNames;
+		const UAnimCurveMetaData* SkelMeshCurveData = InSkeletalMesh->GetAssetUserData<UAnimCurveMetaData>();
+		if (SkelMeshCurveData)
+		{
+			TArray<FName> SkelMeshCurveNames;
+			SkelMeshCurveData->GetCurveMetaDataNames(SkelMeshCurveNames);
+			UniqueNames.Append(SkelMeshCurveNames);
+		}
+	}
+	OutCurveNames = UniqueNames.Array();
+}
+
 FString UMLDeformerInputInfo::GenerateCompatibilityErrorString(const USkeletalMeshComponent* InSkeletalMeshComponent) const
 {
-	const USkeletalMesh* ComponentSkeletalMesh = Cast<USkeletalMesh>(InSkeletalMeshComponent->GetSkinnedAsset());
+	USkeletalMesh* ComponentSkeletalMesh = Cast<USkeletalMesh>(InSkeletalMeshComponent->GetSkinnedAsset());
 	if (ComponentSkeletalMesh == nullptr)
 	{
 		return FString();
@@ -228,12 +256,13 @@ FString UMLDeformerInputInfo::GenerateCompatibilityErrorString(const USkeletalMe
 	}
 
 	// Verify that all required curves are there.
-	const FBlendedHeapCurve& Curves = InSkeletalMeshComponent->GetAnimCurves();
+	TArray<FName> CurveList;
+	BuildCurveList(ComponentSkeletalMesh, CurveList);
 	for (const FName CurveName : CurveNames)
 	{
-		if (!Curves.HasElement(CurveName))
+		if (!CurveList.Contains(CurveName))
 		{
-			ErrorString += FString::Format(TEXT("Required curve '{0}' is missing.\n"), {*CurveName.ToString()});
+			ErrorString += FString::Format(TEXT("Required curve '{0}' is missing.\n"), { *CurveName.ToString() });
 		}
 	}
 	
@@ -279,12 +308,13 @@ FString UMLDeformerInputInfo::GenerateCompatibilityErrorString(USkeletalMesh* In
 	}
 
 	// Verify that all required curves are there.
-	USkeleton* Skeleton = InSkeletalMesh->GetSkeleton();
+	TArray<FName> CurveList;
+	BuildCurveList(InSkeletalMesh, CurveList);
 	for (const FName CurveName : CurveNames)
 	{
-		if (Skeleton->GetCurveMetaData(CurveName) == nullptr)
+		if (!CurveList.Contains(CurveName))
 		{
-			ErrorString += FString::Format(TEXT("Required curve '{0}' is missing.\n"), {*CurveName.ToString()});
+			ErrorString += FString::Format(TEXT("Required curve '{0}' is missing.\n"), { *CurveName.ToString() });
 		}
 	}
 	
diff --git a/Engine/Plugins/Animation/MLDeformer/MLDeformerFramework/Source/MLDeformerFramework/Public/MLDeformerInputInfo.h b/Engine/Plugins/Animation/MLDeformer/MLDeformerFramework/Source/MLDeformerFramework/Public/MLDeformerInputInfo.h
--- a/Engine/Plugins/Animation/MLDeformer/MLDeformerFramework/Source/MLDeformerFramework/Public/MLDeformerInputInfo.h
+++ b/Engine/Plugins/Animation/MLDeformer/MLDeformerFramework/Source/MLDeformerFramework/Public/MLDeformerInputInfo.h
@@ -244,6 +244,15 @@ public:
 	UE_DEPRECATED(5.2, "This method will be removed soon. It shouldn't be needed anymore.")
 	UE_API void UpdateFNames();
 
+	/**
+	 * Builds a list of available curves on the Skeletal Mesh provided, including its skeleton.
+	 * This internally adds a merged set of curves that are found on the skeleton and the skeletal mesh itself.
+	 * @param InSkeletalMesh The Skeletal Mesh to get the list of curves for.
+	 * @param OutCurveNames The list of curve names that have been found.
+	 */
+	UE_API static void BuildCurveList(USkeletalMesh* InSkeletalMesh, TArray<FName>& OutCurveNames);
+
+
 protected:
 	/**
 	 * The path to the skeletal mesh that this model was trained on.
diff --git a/Engine/Plugins/Animation/MLDeformer/MLDeformerFramework/Source/MLDeformerFrameworkEditor/Private/MLDeformerEditorModel.cpp b/Engine/Plugins/Animation/MLDeformer/MLDeformerFramework/Source/MLDeformerFrameworkEditor/Private/MLDeformerEditorModel.cpp
--- a/Engine/Plugins/Animation/MLDeformer/MLDeformerFramework/Source/MLDeformerFrameworkEditor/Private/MLDeformerEditorModel.cpp
+++ b/Engine/Plugins/Animation/MLDeformer/MLDeformerFramework/Source/MLDeformerFrameworkEditor/Private/MLDeformerEditorModel.cpp
@@ -1771,21 +1771,7 @@ namespace UE::MLDeformer
 		check(Model->GetSkeletalMesh());
 
 		TArray<FName> CurveList;
-		if (Model->GetSkeletalMesh()->GetSkeleton())
-		{
-			Model->GetSkeletalMesh()->GetSkeleton()->GetCurveMetaDataNames(CurveList);
-		}
-
-		const UAnimCurveMetaData* SkelMeshCurveData = Model->GetSkeletalMesh()->GetAssetUserData<UAnimCurveMetaData>();
-		if (SkelMeshCurveData)
-		{
-			TArray<FName> SkelMeshCurveNames;
-			SkelMeshCurveData->GetCurveMetaDataNames(SkelMeshCurveNames);
-			for (const FName CurveName : SkelMeshCurveNames)
-			{
-				CurveList.AddUnique(CurveName);
-			}
-		}
+		UMLDeformerInputInfo::BuildCurveList(Model->GetSkeletalMesh(), CurveList);
 		
 		const USkeleton* Skeleton = Model->GetSkeletalMesh()->GetSkeleton();
 		check(Skeleton);
diff --git a/Engine/Plugins/Animation/MLDeformer/MLDeformerFramework/Source/MLDeformerFrameworkEditor/Private/SMLDeformerCurvePickerDialog.cpp b/Engine/Plugins/Animation/MLDeformer/MLDeformerFramework/Source/MLDeformerFrameworkEditor/Private/SMLDeformerCurvePickerDialog.cpp
--- a/Engine/Plugins/Animation/MLDeformer/MLDeformerFramework/Source/MLDeformerFrameworkEditor/Private/SMLDeformerCurvePickerDialog.cpp
+++ b/Engine/Plugins/Animation/MLDeformer/MLDeformerFramework/Source/MLDeformerFrameworkEditor/Private/SMLDeformerCurvePickerDialog.cpp
@@ -1,6 +1,7 @@
 // Copyright Epic Games, Inc. All Rights Reserved.
 
 #include "SMLDeformerCurvePickerDialog.h"
+#include "MLDeformerInputInfo.h"
 #include "Widgets/Input/SSearchBox.h"
 #include "Widgets/Views/SListView.h"
 #include "Dialog/SCustomDialog.h"
@@ -148,22 +149,9 @@ namespace UE::MLDeformer
 		CurveListView->Clear();
 
 		TArray<FName> CurveNames;
-		const UAnimCurveMetaData* SkelMeshCurveData = SkeletalMesh->GetAssetUserData<UAnimCurveMetaData>();
-		if (SkelMeshCurveData)
+		if (SkeletalMesh.IsValid())
 		{
-			TArray<FName> SkelMeshCurveNames;
-			SkelMeshCurveData->GetCurveMetaDataNames(SkelMeshCurveNames);
-			CurveNames.Append(SkelMeshCurveNames);
-		}
-		
-		if (SkeletalMesh->GetSkeleton())
-		{
-			TArray<FName> SkeletonCurveNames;
-			SkeletalMesh->GetSkeleton()->GetCurveMetaDataNames(SkeletonCurveNames);
-			for (const FName& CurveName : SkeletonCurveNames)
-			{
-				CurveNames.AddUnique(CurveName);
-			}
+			UMLDeformerInputInfo::BuildCurveList(SkeletalMesh.Get(), CurveNames);
 		}
 
 		for (const FName CurveName : CurveNames)
diff --git a/Engine/Plugins/Animation/MLDeformer/MLDeformerFramework/Source/MLDeformerFrameworkEditor/Private/SMLDeformerInputCurvesWidget.cpp b/Engine/Plugins/Animation/MLDeformer/MLDeformerFramework/Source/MLDeformerFrameworkEditor/Private/SMLDeformerInputCurvesWidget.cpp
--- a/Engine/Plugins/Animation/MLDeformer/MLDeformerFramework/Source/MLDeformerFrameworkEditor/Private/SMLDeformerInputCurvesWidget.cpp
+++ b/Engine/Plugins/Animation/MLDeformer/MLDeformerFramework/Source/MLDeformerFrameworkEditor/Private/SMLDeformerInputCurvesWidget.cpp
@@ -122,23 +122,7 @@ namespace UE::MLDeformer
 		else
 		{
 			TArray<FName> AssetCurveNames;
-			if (SkeletalMesh)
-			{
-				const UAnimCurveMetaData* SkelMeshCurveData = SkeletalMesh->GetAssetUserData<UAnimCurveMetaData>();
-				if (SkelMeshCurveData)
-				{
-					TArray<FName> SkelMeshCurveNames;
-					SkelMeshCurveData->GetCurveMetaDataNames(SkelMeshCurveNames);
-					AssetCurveNames.Append(SkelMeshCurveNames);
-				}
-			}
-			
-			if (SkeletalMesh->GetSkeleton())
-			{
-				TArray<FName> SkeletonCurveNames;
-				SkeletalMesh->GetSkeleton()->GetCurveMetaDataNames(SkeletonCurveNames);
-				AssetCurveNames.Append(SkeletonCurveNames);
-			}
+			UMLDeformerInputInfo::BuildCurveList(SkeletalMesh, AssetCurveNames);
 
 			for (int32 CurveIndex = 0; CurveIndex < CurveNames.Num(); ++CurveIndex)
 			{
