Fixed missing 'return' statement in GetBoneMatrix.

diff --git a/Engine/Plugins/Animation/MLDeformer/MLDeformerFramework/Shaders/Private/SkeletonWithWeightedQuats.ush b/Engine/Plugins/Animation/MLDeformer/MLDeformerFramework/Shaders/Private/SkeletonWithWeightedQuats.ush
--- a/Engine/Plugins/Animation/MLDeformer/MLDeformerFramework/Shaders/Private/SkeletonWithWeightedQuats.ush
+++ b/Engine/Plugins/Animation/MLDeformer/MLDeformerFramework/Shaders/Private/SkeletonWithWeightedQuats.ush
@@ -30,7 +30,7 @@ float4 GetBlendWeights_{DataInterfaceName}(uint StreamOffset, uint ExtraInfluenc
 
 float3x4 GetBoneMatrix_{DataInterfaceName}(uint BoneIndex)
 {
-	LoadCompressedBoneTransformTransposed({DataInterfaceName}_BoneMatrices, 0, BoneIndex);
+	return LoadCompressedBoneTransformTransposed({DataInterfaceName}_BoneMatrices, 0, BoneIndex);
 }
 
 uint ReadNumBones_{DataInterfaceName}(uint VertexIndex)
