[ContentBrowser] New Add Menu Animation Menu

#Synced-CL 52417896

diff --git a/Engine/Plugins/Animation/MLDeformer/MLDeformerFramework/Source/MLDeformerFrameworkEditor/Private/MLDeformerAssetDefinitions.cpp b/Engine/Plugins/Animation/MLDeformer/MLDeformerFramework/Source/MLDeformerFrameworkEditor/Private/MLDeformerAssetDefinitions.cpp
--- a/Engine/Plugins/Animation/MLDeformer/MLDeformerFramework/Source/MLDeformerFrameworkEditor/Private/MLDeformerAssetDefinitions.cpp
+++ b/Engine/Plugins/Animation/MLDeformer/MLDeformerFramework/Source/MLDeformerFrameworkEditor/Private/MLDeformerAssetDefinitions.cpp
@@ -29,7 +29,10 @@ TSoftClassPtr<UObject> UAssetDefinition_MLDeformer::GetAssetClass() const
 
 TConstArrayView<FAssetCategoryPath> UAssetDefinition_MLDeformer::GetAssetCategories() const
 { 
-	static const auto Categories = { EAssetCategoryPaths::Animation }; 
+	static const auto Categories = 
+		{
+			FAssetCategoryPath(EAssetCategoryPaths::Animation, LOCTEXT("AssetDefinition_MLDeformerSubMenu", "Deformers"))
+		};
 	return Categories;
 }
 
